Expose some landscape functionallity outside of the editor.

diff --git a/Engine/Source/Runtime/Landscape/Classes/LandscapeComponent.h b/Engine/Source/Runtime/Landscape/Classes/LandscapeComponent.h
--- a/Engine/Source/Runtime/Landscape/Classes/LandscapeComponent.h
+++ b/Engine/Source/Runtime/Landscape/Classes/LandscapeComponent.h
@@ -700,14 +700,13 @@ public:
 	const TArray<FLandscapePerLODMaterialOverride>& GetPerLODOverrideMaterials() const { return PerLODOverrideMaterials; }
 	void SetPerLODOverrideMaterials(const TArray<FLandscapePerLODMaterialOverride>& InValue) { PerLODOverrideMaterials = InValue; }
 
-#if WITH_EDITOR
-	LANDSCAPE_API uint32 ComputeLayerHash(bool InReturnEditingHash = true) const;
-
 	LANDSCAPE_API void SetHeightmap(UTexture2D* NewHeightmap);
-
 	LANDSCAPE_API void SetWeightmapTextures(const TArray<UTexture2D*>& InNewWeightmapTextures, bool InApplyToEditingWeightmap = false);
-
 	LANDSCAPE_API void SetWeightmapLayerAllocations(const TArray<FWeightmapLayerAllocationInfo>& InNewWeightmapLayerAllocations);
+
+#if WITH_EDITOR
+	LANDSCAPE_API uint32 ComputeLayerHash(bool InReturnEditingHash = true) const;
+
 	LANDSCAPE_API void SetWeightmapTexturesUsage(const TArray<ULandscapeWeightmapUsage*>& InNewWeightmapTexturesUsage, bool InApplyToEditingWeightmap = false);
 	LANDSCAPE_API TArray<ULandscapeWeightmapUsage*>& GetWeightmapTexturesUsage(bool InReturnEditingWeightmap = false);
 	LANDSCAPE_API const TArray<ULandscapeWeightmapUsage*>& GetWeightmapTexturesUsage(bool InReturnEditingWeightmap = false) const;
@@ -883,10 +882,10 @@ public:
 
 	LANDSCAPE_API const FMeshMapBuildData* GetMeshMapBuildData() const;
 
-#if WITH_EDITOR
 	/** Initialize the landscape component */
-	LANDSCAPE_API void Init(int32 InBaseX,int32 InBaseY,int32 InComponentSizeQuads, int32 InNumSubsections,int32 InSubsectionSizeQuads);
+	LANDSCAPE_API void Init(int32 InBaseX, int32 InBaseY, int32 InComponentSizeQuads, int32 InNumSubsections, int32 InSubsectionSizeQuads);
 
+#if WITH_EDITOR
 	/**
 	 * Recalculate cached bounds using height values.
 	 */
@@ -1040,15 +1039,6 @@ public:
 	/** returns the 2D bounds of this component (in landscape quads) */
 	LANDSCAPE_API FIntRect GetComponentExtent() const;
 
-	/** Updates navigation properties to match landscape actor's */
-	void UpdateNavigationRelevance();
-
-	/** Updates the reject navmesh underneath flag in the collision component */
-	void UpdateRejectNavmeshUnderneath();
-	
-	/** Updates the values of component-level properties exposed by the Landscape Actor */
-	LANDSCAPE_API void UpdatedSharedPropertiesFromActor();
-
 	LANDSCAPE_API void ClearUpdateFlagsForModes(uint32 InModeMask);
 	LANDSCAPE_API void RequestWeightmapUpdate(bool bUpdateAll = false, bool bUpdateCollision = true);
 	LANDSCAPE_API void RequestHeightmapUpdate(bool bUpdateAll = false, bool bUpdateCollision = true);
@@ -1064,6 +1054,15 @@ public:
 	void GetLandscapeComponentNeighbors3x3(TStaticArray<ULandscapeComponent*, 9>& OutNeighborComponents) const;
 #endif
 
+	/** Updates navigation properties to match landscape actor's */
+	void UpdateNavigationRelevance();
+
+	/** Updates the reject navmesh underneath flag in the collision component */
+	void UpdateRejectNavmeshUnderneath();
+
+	/** Updates the values of component-level properties exposed by the Landscape Actor */
+	LANDSCAPE_API void UpdatedSharedPropertiesFromActor();
+
 	friend class FLandscapeComponentSceneProxy;
 	friend struct FLandscapeComponentDataInterface;
 
diff --git a/Engine/Source/Runtime/Landscape/Private/Landscape.cpp b/Engine/Source/Runtime/Landscape/Private/Landscape.cpp
--- a/Engine/Source/Runtime/Landscape/Private/Landscape.cpp
+++ b/Engine/Source/Runtime/Landscape/Private/Landscape.cpp
@@ -888,6 +888,7 @@ void ULandscapeInfo::UpdateDebugColorMaterial()
 	FlushRenderingCommands();
 	//GWarn->EndSlowTask();
 }
+#endif // WITH_EDITOR
 
 void ULandscapeComponent::UpdatedSharedPropertiesFromActor()
 {
@@ -909,7 +910,6 @@ void ULandscapeComponent::UpdatedSharedPropertiesFromActor()
 	UpdateNavigationRelevance();
 	UpdateRejectNavmeshUnderneath();
 }
-#endif // WITH_EDITOR
 
 void ULandscapeComponent::PostLoad()
 {
@@ -2095,6 +2095,8 @@ void ULandscapeComponent::RemoveLayerData(const FGuid& InLayerGuid)
 	CachedEditingLayerData = nullptr;
 }
 
+#endif // WITH_EDITOR
+
 void ULandscapeComponent::SetHeightmap(UTexture2D* NewHeightmap)
 {
 	check(NewHeightmap != nullptr);
@@ -2103,6 +2105,7 @@ void ULandscapeComponent::SetHeightmap(UTexture2D* NewHeightmap)
 
 void ULandscapeComponent::SetWeightmapTextures(const TArray<UTexture2D*>& InNewWeightmapTextures, bool InApplyToEditingWeightmap)
 {
+#if WITH_EDITOR
 	FLandscapeLayerComponentData* EditingLayer = GetEditingLayer();
 
 	if (InApplyToEditingWeightmap && EditingLayer != nullptr)
@@ -2111,6 +2114,7 @@ void ULandscapeComponent::SetWeightmapTextures(const TArray<UTexture2D*>& InNewW
 		EditingLayer->WeightmapData.Textures.Append(InNewWeightmapTextures);
 	}
 	else
+#endif // WITH_EDITOR
 	{
 		WeightmapTextures = InNewWeightmapTextures;
 	}
@@ -2121,6 +2125,7 @@ void ULandscapeComponent::SetWeightmapLayerAllocations(const TArray<FWeightmapLa
 	WeightmapLayerAllocations = InNewWeightmapLayerAllocations;
 }
 
+#if WITH_EDITOR
 TArray<ULandscapeWeightmapUsage*>& ULandscapeComponent::GetWeightmapTexturesUsage(bool InReturnEditingWeightmap)
 {
 	if (InReturnEditingWeightmap)
@@ -3431,10 +3436,11 @@ void ALandscapeProxy::Destroyed()
 		FeatureLevelChangedDelegateHandle.Reset();
 	}
 }
+#endif // WITH_EDITOR
 
 void ALandscapeProxy::GetSharedProperties(ALandscapeProxy* Landscape)
 {
-	if (GIsEditor && Landscape)
+	if (Landscape)
 	{
 		Modify();
 
@@ -3462,8 +3468,6 @@ void ALandscapeProxy::GetSharedProperties(ALandscapeProxy* Landscape)
 		NumSubsections = Landscape->NumSubsections;
 		SubsectionSizeQuads = Landscape->SubsectionSizeQuads;
 		MaxLODLevel = Landscape->MaxLODLevel;
-		LODDistanceFactor_DEPRECATED = Landscape->LODDistanceFactor_DEPRECATED;
-		LODFalloff_DEPRECATED = Landscape->LODFalloff_DEPRECATED;
 		ComponentScreenSizeToUseSubSections = Landscape->ComponentScreenSizeToUseSubSections;
 		LODDistributionSetting = Landscape->LODDistributionSetting;
 		LOD0DistributionSetting = Landscape->LOD0DistributionSetting;
@@ -3486,18 +3490,28 @@ void ALandscapeProxy::GetSharedProperties(ALandscapeProxy* Landscape)
 		{
 			LandscapeHoleMaterial = Landscape->LandscapeHoleMaterial;
 		}
-		if (LandscapeMaterial == Landscape->LandscapeMaterial)
-		{
-			EditorLayerSettings = Landscape->EditorLayerSettings;
-		}
 		if (!DefaultPhysMaterial)
 		{
 			DefaultPhysMaterial = Landscape->DefaultPhysMaterial;
 		}
 		LightmassSettings = Landscape->LightmassSettings;
 	}
+
+#if WITH_EDITOR
+	if (GIsEditor && Landscape)
+	{
+		LODDistanceFactor_DEPRECATED = Landscape->LODDistanceFactor_DEPRECATED;
+		LODFalloff_DEPRECATED = Landscape->LODFalloff_DEPRECATED;
+		if (LandscapeMaterial == Landscape->LandscapeMaterial)
+		{
+			EditorLayerSettings = Landscape->EditorLayerSettings;
+		}
+	}
+#endif // WITH_EDITOR
 }
 
+#if WITH_EDITOR
+
 void ALandscapeProxy::FixupSharedData(ALandscape* Landscape)
 {
 	if (Landscape == nullptr)
diff --git a/Engine/Source/Runtime/Landscape/Private/LandscapeEdit.cpp b/Engine/Source/Runtime/Landscape/Private/LandscapeEdit.cpp
--- a/Engine/Source/Runtime/Landscape/Private/LandscapeEdit.cpp
+++ b/Engine/Source/Runtime/Landscape/Private/LandscapeEdit.cpp
@@ -115,6 +115,7 @@ extern const size_t ChannelOffsets[4];
 extern int32 LiveRebuildNaniteOnModification;
 
 ULandscapeLayerInfoObject* ALandscapeProxy::VisibilityLayer = nullptr;
+#endif //WITH_EDITOR
 
 void ULandscapeComponent::Init(int32 InBaseX, int32 InBaseY, int32 InComponentSizeQuads, int32 InNumSubsections, int32 InSubsectionSizeQuads)
 {
@@ -138,6 +139,7 @@ void ULandscapeComponent::Init(int32 InBaseX, int32 InBaseY, int32 InComponentSi
 	UpdatedSharedPropertiesFromActor();
 }
 
+#if WITH_EDITOR
 void ULandscapeComponent::UpdateCachedBounds(bool bInApproximateBounds)
 {
 	// Update local-space bounding box
@@ -181,6 +183,7 @@ void ULandscapeComponent::UpdateCachedBounds(bool bInApproximateBounds)
 		HFCollisionComponent->UpdateComponentToWorld();
 	}
 }
+#endif //WITH_EDITOR
 
 void ULandscapeComponent::UpdateNavigationRelevance()
 {
@@ -200,6 +203,7 @@ void ULandscapeComponent::UpdateRejectNavmeshUnderneath()
 	}
 }
 
+#if WITH_EDITOR
 ULandscapeMaterialInstanceConstant* ALandscapeProxy::GetLayerThumbnailMIC(UMaterialInterface* LandscapeMaterial, FName LayerName, UTexture2D* ThumbnailWeightmap, UTexture2D* ThumbnailHeightmap, ALandscapeProxy* Proxy)
 {
 	if (!LandscapeMaterial)
