Fix several crashes in the editor by removing FLandscapeComponentSceneProxy's direct usages of a UMaterialInterface on the render thread.

diff --git a/Engine/Source/Runtime/Landscape/Private/LandscapeRender.cpp b/Engine/Source/Runtime/Landscape/Private/LandscapeRender.cpp
--- a/Engine/Source/Runtime/Landscape/Private/LandscapeRender.cpp
+++ b/Engine/Source/Runtime/Landscape/Private/LandscapeRender.cpp
@@ -801,10 +801,17 @@ FLandscapeComponentSceneProxy::FLandscapeComponentSceneProxy(ULandscapeComponent
 
 	const ERHIFeatureLevel::Type FeatureLevel = GetScene().GetFeatureLevel();
 
+	auto GetRenderProxy = 
+		[](const TObjectPtr<UMaterialInterface>& Material)
+		{
+			return Material ? Material->GetRenderProxy() : nullptr;
+		};
+	TArray<UMaterialInterface*> AvailableMaterialInterfaces;
 	if (FeatureLevel == ERHIFeatureLevel::ES3_1)
 	{
 		WeightmapTextures = InComponent->MobileWeightmapTextures;
-		AvailableMaterials.Append(InComponent->MobileMaterialInterfaces);
+		Algo::Transform(InComponent->MobileMaterialInterfaces, AvailableMaterials, GetRenderProxy);
+		AvailableMaterialInterfaces.Append(InComponent->MobileMaterialInterfaces);
 		//TODO: Add support for bUseDynamicMaterialInstance ?
 	}
 	else
@@ -812,11 +819,13 @@ FLandscapeComponentSceneProxy::FLandscapeComponentSceneProxy(ULandscapeComponent
 		WeightmapTextures = InComponent->GetWeightmapTextures();
 		if (InComponent->GetLandscapeProxy()->bUseDynamicMaterialInstance)
 		{
-			AvailableMaterials.Append(InComponent->MaterialInstancesDynamic);
+			Algo::Transform(InComponent->MaterialInstancesDynamic, AvailableMaterials, GetRenderProxy);
+			AvailableMaterialInterfaces.Append(InComponent->MaterialInstancesDynamic);
 		}
 		else
 		{
-			AvailableMaterials.Append(InComponent->MaterialInstances);
+			Algo::Transform(InComponent->MaterialInstances, AvailableMaterials, GetRenderProxy);
+			AvailableMaterialInterfaces.Append(InComponent->MaterialInstances);
 		}
 	}
 
@@ -924,26 +933,31 @@ FLandscapeComponentSceneProxy::FLandscapeComponentSceneProxy(ULandscapeComponent
 
 	const bool bHasStaticLighting = ComponentLightInfo->GetLightMap() || ComponentLightInfo->GetShadowMap();
 
+	check(AvailableMaterialInterfaces.Num() == AvailableMaterials.Num());
 	// Check material usage
-	if (ensure(AvailableMaterials.Num() > 0))
+	if (ensure(AvailableMaterialInterfaces.Num() > 0))
 	{
-		for (UMaterialInterface*& MaterialInterface : AvailableMaterials)
+		for(int Index = 0; Index < AvailableMaterialInterfaces.Num(); ++Index)
 		{
+			UMaterialInterface*& MaterialInterface = AvailableMaterialInterfaces[Index];
+			FMaterialRenderProxy*& MaterialProxy = AvailableMaterials[Index];
 			if (MaterialInterface == nullptr ||
 				(bHasStaticLighting && !MaterialInterface->CheckMaterialUsage_Concurrent(MATUSAGE_StaticLighting)))
 			{
 				MaterialInterface = UMaterial::GetDefaultMaterial(MD_Surface);
+				MaterialProxy = MaterialInterface->GetRenderProxy();
 			}
 		}
 	}
 	else
 	{
-		AvailableMaterials.Add(UMaterial::GetDefaultMaterial(MD_Surface));
+		AvailableMaterialInterfaces.Add(UMaterial::GetDefaultMaterial(MD_Surface));
+		AvailableMaterials.Add(AvailableMaterialInterfaces.Last()->GetRenderProxy());
 	}
 
 	MaterialRelevances.Reserve(AvailableMaterials.Num());
 
-	for (UMaterialInterface*& MaterialInterface : AvailableMaterials)
+	for (UMaterialInterface*& MaterialInterface : AvailableMaterialInterfaces)
 	{
 		const UMaterial* LandscapeMaterial = MaterialInterface != nullptr ? MaterialInterface->GetMaterial_Concurrent() : nullptr;
 
@@ -1107,7 +1121,7 @@ void FLandscapeComponentSceneProxy::CreateRenderThreadResources()
 
 		// Grass is being generated using LOD0 material only
 		// It uses the fixed grid vertex factory so it doesn't support XY offsets
-		FMaterialRenderProxy* RenderProxy = AvailableMaterials[LODIndexToMaterialIndex[0]]->GetRenderProxy();
+		FMaterialRenderProxy* RenderProxy = AvailableMaterials[LODIndexToMaterialIndex[0]];
 		GrassMeshBatch.VertexFactory = FixedGridVertexFactory;
 		GrassMeshBatch.MaterialRenderProxy = RenderProxy;
 		GrassMeshBatch.LCI = nullptr;
@@ -1592,7 +1606,7 @@ void FLandscapeComponentSceneProxy::OnTransformChanged()
 }
 
 /** Creates a mesh batch for virtual texture rendering. Will render a simple fixed grid with combined subsections. */
-bool FLandscapeComponentSceneProxy::GetMeshElementForVirtualTexture(int32 InLodIndex, ERuntimeVirtualTextureMaterialType MaterialType, UMaterialInterface* InMaterialInterface, FMeshBatch& OutMeshBatch, TArray<FLandscapeBatchElementParams>& OutStaticBatchParamArray) const
+bool FLandscapeComponentSceneProxy::GetMeshElementForVirtualTexture(int32 InLodIndex, ERuntimeVirtualTextureMaterialType MaterialType, FMaterialRenderProxy* InMaterialInterface, FMeshBatch& OutMeshBatch, TArray<FLandscapeBatchElementParams>& OutStaticBatchParamArray) const
 {
 	if (InMaterialInterface == nullptr)
 	{
@@ -1600,7 +1614,7 @@ bool FLandscapeComponentSceneProxy::GetMeshElementForVirtualTexture(int32 InLodI
 	}
 
 	OutMeshBatch.VertexFactory = FixedGridVertexFactory;
-	OutMeshBatch.MaterialRenderProxy = InMaterialInterface->GetRenderProxy();
+	OutMeshBatch.MaterialRenderProxy = InMaterialInterface;
 	OutMeshBatch.ReverseCulling = IsLocalToWorldDeterminantNegative();
 	OutMeshBatch.CastShadow = false;
 	OutMeshBatch.bUseForDepthPass = false;
@@ -1660,15 +1674,15 @@ void FLandscapeComponentSceneProxy::ApplyWorldOffset(FVector InOffset)
 template<class ArrayType>
 bool FLandscapeComponentSceneProxy::GetStaticMeshElement(int32 LODIndex, bool bForToolMesh, FMeshBatch& MeshBatch, ArrayType& OutStaticBatchParamArray) const
 {
-	UMaterialInterface* MaterialInterface = nullptr;
+	FMaterialRenderProxy* Material = nullptr;
 
 	{
 		int32 MaterialIndex = LODIndexToMaterialIndex[LODIndex];
 
 		// Defaults to the material interface w/ potential tessellation
-		MaterialInterface = AvailableMaterials[MaterialIndex];
+		Material = AvailableMaterials[MaterialIndex];
 
-		if (!MaterialInterface)
+		if (!Material)
 		{
 			return false;
 		}
@@ -1676,7 +1690,7 @@ bool FLandscapeComponentSceneProxy::GetStaticMeshElement(int32 LODIndex, bool bF
 
 	{
 		MeshBatch.VertexFactory = VertexFactory;
-		MeshBatch.MaterialRenderProxy = MaterialInterface->GetRenderProxy();
+		MeshBatch.MaterialRenderProxy = Material;
 
 		MeshBatch.LCI = ComponentLightInfo.Get();
 		MeshBatch.ReverseCulling = IsLocalToWorldDeterminantNegative();
@@ -2221,7 +2235,7 @@ void FLandscapeComponentSceneProxy::GetDynamicRayTracingInstances(FRayTracingMat
 
 	const int8 CurrentLODIndex = LODToRender;
 	int8 MaterialIndex = LODIndexToMaterialIndex.IsValidIndex(CurrentLODIndex) ? LODIndexToMaterialIndex[CurrentLODIndex] : INDEX_NONE;
-	UMaterialInterface* SelectedMaterial = MaterialIndex != INDEX_NONE ? AvailableMaterials[MaterialIndex] : nullptr;
+	FMaterialRenderProxy* SelectedMaterial = MaterialIndex != INDEX_NONE ? AvailableMaterials[MaterialIndex] : nullptr;
 
 	// this is really not normal that we have no material at this point, so do not continue
 	if (SelectedMaterial == nullptr)
@@ -2231,7 +2245,7 @@ void FLandscapeComponentSceneProxy::GetDynamicRayTracingInstances(FRayTracingMat
 
 	FMeshBatch BaseMeshBatch;
 	BaseMeshBatch.VertexFactory = VertexFactory;
-	BaseMeshBatch.MaterialRenderProxy = SelectedMaterial->GetRenderProxy();
+	BaseMeshBatch.MaterialRenderProxy = SelectedMaterial;
 	BaseMeshBatch.LCI = ComponentLightInfo.Get();
 	BaseMeshBatch.CastShadow = true;
 	BaseMeshBatch.CastRayTracedShadow = true;
diff --git a/Engine/Source/Runtime/Landscape/Public/LandscapeRender.h b/Engine/Source/Runtime/Landscape/Public/LandscapeRender.h
--- a/Engine/Source/Runtime/Landscape/Public/LandscapeRender.h
+++ b/Engine/Source/Runtime/Landscape/Public/LandscapeRender.h
@@ -737,7 +737,7 @@ protected:
 	FLandscapeVertexFactory*	FixedGridVertexFactory;
 
 	/** All available materials, including LOD Material, Tessellation generated materials*/
-	TArray<UMaterialInterface*> AvailableMaterials;
+	TArray<FMaterialRenderProxy*> AvailableMaterials;
 
 	// FLightCacheInterface
 	TUniquePtr<FLandscapeLCI> ComponentLightInfo;
@@ -786,7 +786,7 @@ protected:
 	
 	int8 GetLODFromScreenSize(float InScreenSizeSquared, float InViewLODScale) const;
 
-	bool GetMeshElementForVirtualTexture(int32 InLodIndex, ERuntimeVirtualTextureMaterialType MaterialType, UMaterialInterface* InMaterialInterface, FMeshBatch& OutMeshBatch, TArray<FLandscapeBatchElementParams>& OutStaticBatchParamArray) const;
+	bool GetMeshElementForVirtualTexture(int32 InLodIndex, ERuntimeVirtualTextureMaterialType MaterialType, FMaterialRenderProxy* InMaterialInterface, FMeshBatch& OutMeshBatch, TArray<FLandscapeBatchElementParams>& OutStaticBatchParamArray) const;
 	template<class ArrayType> bool GetStaticMeshElement(int32 LODIndex, bool bForToolMesh, FMeshBatch& MeshBatch, ArrayType& OutStaticBatchParamArray) const;
 	
 	virtual void ApplyMeshElementModifier(FMeshBatchElement& InOutMeshElement, int32 InLodIndex) const {}
