Add option to sample landscape height per pixel when rendering to RVT.
This is only the landscape setting. The code that applies the change will be submitted in a second pass to allow for shader recompilation warmup.

diff --git a/Engine/Source/Runtime/Landscape/Classes/LandscapeProxy.h b/Engine/Source/Runtime/Landscape/Classes/LandscapeProxy.h
--- a/Engine/Source/Runtime/Landscape/Classes/LandscapeProxy.h
+++ b/Engine/Source/Runtime/Landscape/Classes/LandscapeProxy.h
@@ -498,18 +498,26 @@ public:
 	UPROPERTY(VisibleAnywhere, Transient, Category = VirtualTexture)
 	bool bSetCreateRuntimeVirtualTextureVolumes;
 
+	/** 
+	 * Use a single quad to render this landscape to runtime virtual texture pages. 
+	 * This is the fastest path but it only gives correct results if the runtime virtual texture orientation is aligned with the landscape.
+	 * If the two are unaligned we need to render to the virtual texture using LODs with sufficient density.
+	 */
+	UPROPERTY(EditAnywhere, AdvancedDisplay, BlueprintReadOnly, Category = VirtualTexture)
+	bool bVirtualTextureRenderWithQuad = false;
+
 	/** 
 	 * Number of mesh levels to use when rendering landscape into runtime virtual texture.
 	 * Lower values reduce vertex count when rendering to the runtime virtual texture but decrease accuracy when using values that require vertex interpolation.
 	 */
-	UPROPERTY(EditAnywhere, AdvancedDisplay, BlueprintReadOnly, Category = VirtualTexture, meta = (DisplayName = "Virtual Texture Num LODs", UIMin = "0", UIMax = "7"))
+	UPROPERTY(EditAnywhere, AdvancedDisplay, BlueprintReadOnly, Category = VirtualTexture, meta = (DisplayName = "Virtual Texture Num LODs", EditCondition = "!bVirtualTextureRenderWithQuad", UIMin = "0", UIMax = "7"))
 	int32 VirtualTextureNumLods = 6;
 
 	/** 
 	 * Bias to the LOD selected for rendering to runtime virtual textures.
 	 * Higher values reduce vertex count when rendering to the runtime virtual texture.
 	 */
-	UPROPERTY(EditAnywhere, AdvancedDisplay, BlueprintReadOnly, Category = VirtualTexture, meta = (DisplayName = "Virtual Texture LOD Bias", UIMin = "0", UIMax = "7"))
+	UPROPERTY(EditAnywhere, AdvancedDisplay, BlueprintReadOnly, Category = VirtualTexture, meta = (DisplayName = "Virtual Texture LOD Bias", EditCondition = "!bVirtualTextureRenderWithQuad", UIMin = "0", UIMax = "7"))
 	int32 VirtualTextureLodBias = 0;
 
 	/** Controls if this component draws in the main pass as well as in the virtual texture. */
diff --git a/Engine/Source/Runtime/Landscape/Private/Landscape.cpp b/Engine/Source/Runtime/Landscape/Private/Landscape.cpp
--- a/Engine/Source/Runtime/Landscape/Private/Landscape.cpp
+++ b/Engine/Source/Runtime/Landscape/Private/Landscape.cpp
@@ -3468,6 +3468,7 @@ void ALandscapeProxy::GetSharedProperties(ALandscapeProxy* Landscape)
 		bBakeMaterialPositionOffsetIntoCollision = Landscape->bBakeMaterialPositionOffsetIntoCollision;
 		RuntimeVirtualTextures = Landscape->RuntimeVirtualTextures;
 		VirtualTextureLodBias = Landscape->VirtualTextureLodBias;
+		bVirtualTextureRenderWithQuad = Landscape->bVirtualTextureRenderWithQuad;
 		VirtualTextureNumLods = Landscape->VirtualTextureNumLods;
 		VirtualTextureRenderPassType = Landscape->VirtualTextureRenderPassType;
 
diff --git a/Engine/Source/Runtime/Landscape/Private/LandscapeEdit.cpp b/Engine/Source/Runtime/Landscape/Private/LandscapeEdit.cpp
--- a/Engine/Source/Runtime/Landscape/Private/LandscapeEdit.cpp
+++ b/Engine/Source/Runtime/Landscape/Private/LandscapeEdit.cpp
@@ -5296,6 +5296,7 @@ void ALandscapeProxy::PostEditChangeProperty(FPropertyChangedEvent& PropertyChan
 	}
 	else if (PropertyName == GET_MEMBER_NAME_CHECKED(ALandscapeProxy, RuntimeVirtualTextures)
 		|| PropertyName == GET_MEMBER_NAME_CHECKED(ALandscapeProxy, VirtualTextureRenderPassType)
+		|| PropertyName == GET_MEMBER_NAME_CHECKED(ALandscapeProxy, bVirtualTextureRenderWithQuad)
 		|| PropertyName == GET_MEMBER_NAME_CHECKED(ALandscapeProxy, VirtualTextureNumLods)
 		|| PropertyName == GET_MEMBER_NAME_CHECKED(ALandscapeProxy, VirtualTextureLodBias))
 	{
@@ -5763,6 +5764,7 @@ void ALandscape::PostEditChangeProperty(FPropertyChangedEvent& PropertyChangedEv
 	else if (
 		PropertyName == GET_MEMBER_NAME_CHECKED(ALandscapeProxy, RuntimeVirtualTextures) ||
 		PropertyName == GET_MEMBER_NAME_CHECKED(ALandscapeProxy, VirtualTextureRenderPassType) ||
+		PropertyName == GET_MEMBER_NAME_CHECKED(ALandscapeProxy, bVirtualTextureRenderWithQuad) ||
 		PropertyName == GET_MEMBER_NAME_CHECKED(ALandscapeProxy, VirtualTextureNumLods) ||
 		PropertyName == GET_MEMBER_NAME_CHECKED(ALandscapeProxy, VirtualTextureLodBias))
 	{
