Fix for rare crash caused by calling UTexture2D::IsDefaultTexture() from a render thread
CI Test: https://horde.devtools.epicgames.com/job/635f5f25a27fb703befc351c

diff --git a/Engine/Source/Runtime/Landscape/Private/LandscapeRender.cpp b/Engine/Source/Runtime/Landscape/Private/LandscapeRender.cpp
--- a/Engine/Source/Runtime/Landscape/Private/LandscapeRender.cpp
+++ b/Engine/Source/Runtime/Landscape/Private/LandscapeRender.cpp
@@ -47,6 +47,8 @@ LandscapeRender.cpp: New terrain rendering
 #include "StaticMeshResources.h"
 #include "TextureResource.h"
 #include "NaniteSceneProxy.h"
+#include "Rendering/Texture2DResource.h"
+
 
 IMPLEMENT_GLOBAL_SHADER_PARAMETER_STRUCT(FLandscapeUniformShaderParameters, "LandscapeParameters");
 IMPLEMENT_GLOBAL_SHADER_PARAMETER_STRUCT(FLandscapeFixedGridUniformShaderParameters, "LandscapeFixedGrid");
@@ -3648,7 +3650,7 @@ float FLandscapeComponentSceneProxy::ComputeLODBias() const
 	{
 		if (const FTexture2DResource* TextureResource = (const FTexture2DResource*)HeightmapTexture->GetResource())
 		{
-			ComputedLODBias = HeightmapTexture->GetNumMips() - HeightmapTexture->GetNumResidentMips();
+			ComputedLODBias = TextureResource->GetCurrentFirstMip();
 		}
 	}
 
