Landscape : Dirty Landscape Actor packages on Water & Patch updates.

diff --git a/Engine/Plugins/Experimental/LandscapePatch/Source/LandscapePatch/Private/LandscapePatchComponent.cpp b/Engine/Plugins/Experimental/LandscapePatch/Source/LandscapePatch/Private/LandscapePatchComponent.cpp
--- a/Engine/Plugins/Experimental/LandscapePatch/Source/LandscapePatch/Private/LandscapePatchComponent.cpp
+++ b/Engine/Plugins/Experimental/LandscapePatch/Source/LandscapePatch/Private/LandscapePatchComponent.cpp
@@ -218,7 +218,7 @@ void ULandscapePatchComponent::OnUpdateTransform(EUpdateTransformFlags UpdateTra
 
 	if (IsEnabled())
 	{
-		RequestLandscapeUpdate();
+		RequestLandscapeUpdate(/* bInUserTriggeredUpdate = */ true);
 	}
 }
 
@@ -249,7 +249,7 @@ void ULandscapePatchComponent::PostEditChangeProperty(FPropertyChangedEvent& Pro
 	if (IsEnabled() || (PropertyChangedEvent.Property
 		&& (PropertyChangedEvent.Property->GetFName() == GET_MEMBER_NAME_CHECKED(ULandscapePatchComponent, bIsEnabled))))
 	{
-		RequestLandscapeUpdate();
+		RequestLandscapeUpdate(/* bInUserTriggeredUpdate = */ true);
 	}
 
 	// It is important that this super call happen after the above, because inside a blueprint actor, the call triggers a
@@ -384,7 +384,7 @@ void ULandscapePatchComponent::MoveToTop()
 	}
 }
 
-void ULandscapePatchComponent::RequestLandscapeUpdate()
+void ULandscapePatchComponent::RequestLandscapeUpdate(bool bInUserTriggeredUpdate)
 {
 #if WITH_EDITOR
 	UWorld* World = GetWorld();
@@ -446,7 +446,7 @@ void ULandscapePatchComponent::RequestLandscapeUpdate()
 	if(bRequestUpdate)
 	{
 		ResetWarnings();
-		PatchManager->RequestLandscapeUpdate();
+		PatchManager->RequestLandscapeUpdate(bInUserTriggeredUpdate);
 	}
 #endif
 }
diff --git a/Engine/Plugins/Experimental/LandscapePatch/Source/LandscapePatch/Public/LandscapePatchComponent.h b/Engine/Plugins/Experimental/LandscapePatch/Source/LandscapePatch/Public/LandscapePatchComponent.h
--- a/Engine/Plugins/Experimental/LandscapePatch/Source/LandscapePatch/Public/LandscapePatchComponent.h
+++ b/Engine/Plugins/Experimental/LandscapePatch/Source/LandscapePatch/Public/LandscapePatchComponent.h
@@ -36,7 +36,7 @@ public:
 	virtual bool IsAffectingWeightmapLayer(const FName& InLayerName) const { return false; }
 
 	UFUNCTION(BlueprintCallable, Category = "LandscapePatch")
-	void RequestLandscapeUpdate();
+	void RequestLandscapeUpdate(bool bInUserTriggeredUpdate = false);
 
 	/**
 	 * Allows the patch to be disabled, so that it no longer affects the landscape. This can be useful
diff --git a/Engine/Plugins/Experimental/Water/Source/Editor/Private/WaterLandscapeBrush.cpp b/Engine/Plugins/Experimental/Water/Source/Editor/Private/WaterLandscapeBrush.cpp
--- a/Engine/Plugins/Experimental/Water/Source/Editor/Private/WaterLandscapeBrush.cpp
+++ b/Engine/Plugins/Experimental/Water/Source/Editor/Private/WaterLandscapeBrush.cpp
@@ -26,6 +26,7 @@
 #include "Components/BillboardComponent.h"
 #include "Modules/ModuleManager.h"
 #include "WaterModule.h"
+#include "WaterEditorModule.h"
 #include "WaterEditorSettings.h"
 #include "LandscapeModule.h"
 #include "LandscapeEditorServices.h"
@@ -205,10 +206,10 @@ void AWaterLandscapeBrush::OnWaterBrushActorChanged(const IWaterBrushActorInterf
 	const UWaterEditorSettings* WaterEditorSettings = GetDefault<UWaterEditorSettings>();
 	check(WaterEditorSettings != nullptr);
 
-	bool bAllowLandscapeUpdate = (InParams.PropertyChangedEvent.ChangeType != EPropertyChangeType::Interactive) || WaterEditorSettings->GetUpdateLandscapeDuringInteractiveChanges();
+	const bool bAllowLandscapeUpdate = (InParams.PropertyChangedEvent.ChangeType != EPropertyChangeType::Interactive) || WaterEditorSettings->GetUpdateLandscapeDuringInteractiveChanges();
 	if (bForceUpdateBrush || (bAffectsLandscape && bAllowLandscapeUpdate))
 	{
-		RequestLandscapeUpdate();
+		RequestLandscapeUpdate(/* bInUserTriggered = */ InParams.bUserTriggered);
 	}
 
 	bool bAllowWaterMeshUpdate = (InParams.PropertyChangedEvent.ChangeType != EPropertyChangeType::Interactive) || WaterEditorSettings->GetUpdateWaterMeshDuringInteractiveChanges();
diff --git a/Engine/Plugins/Experimental/Water/Source/Runtime/Private/WaterBodyActor.cpp b/Engine/Plugins/Experimental/Water/Source/Runtime/Private/WaterBodyActor.cpp
--- a/Engine/Plugins/Experimental/Water/Source/Runtime/Private/WaterBodyActor.cpp
+++ b/Engine/Plugins/Experimental/Water/Source/Runtime/Private/WaterBodyActor.cpp
@@ -132,6 +132,7 @@ void AWaterBody::PostEditMove(bool bFinished)
 	FOnWaterBodyChangedParams Params;
 	Params.PropertyChangedEvent.ChangeType = bFinished ? EPropertyChangeType::ValueSet : EPropertyChangeType::Interactive;
 	Params.bShapeOrPositionChanged = true;
+	Params.bUserTriggered = true;
 	WaterBodyComponent->OnWaterBodyChanged(Params);
 }
 
diff --git a/Engine/Plugins/Experimental/Water/Source/Runtime/Private/WaterBodyComponent.cpp b/Engine/Plugins/Experimental/Water/Source/Runtime/Private/WaterBodyComponent.cpp
--- a/Engine/Plugins/Experimental/Water/Source/Runtime/Private/WaterBodyComponent.cpp
+++ b/Engine/Plugins/Experimental/Water/Source/Runtime/Private/WaterBodyComponent.cpp
@@ -1030,6 +1030,7 @@ void UWaterBodyComponent::PostEditImport()
 	FOnWaterBodyChangedParams Params;
 	Params.bShapeOrPositionChanged = true;
 	Params.bWeightmapSettingsChanged = true;
+	Params.bUserTriggered = true;
 	OnWaterBodyChanged(Params);
 
 	RequestGPUWaveDataUpdate();
@@ -1113,6 +1114,7 @@ void UWaterBodyComponent::CheckForErrors()
 void UWaterBodyComponent::PostEditChangeProperty(FPropertyChangedEvent& PropertyChangedEvent)
 {
 	FOnWaterBodyChangedParams Params(PropertyChangedEvent);
+	Params.bUserTriggered = true;
 	OnPostEditChangeProperty(Params);
 	
 	Super::PostEditChangeProperty(PropertyChangedEvent);
@@ -1128,6 +1130,7 @@ void UWaterBodyComponent::OnWaterSplineDataChanged(const FOnWaterSplineDataChang
 	// Transfer the FOnWaterSplineDataChangedParams parameters to FOnWaterBodyChangedParams :
 	FOnWaterBodyChangedParams Params(InParams.PropertyChangedEvent);
 	Params.bShapeOrPositionChanged = true;
+	Params.bUserTriggered = InParams.bUserTriggered;
 	OnWaterBodyChanged(Params);
 }
 
@@ -1154,14 +1157,15 @@ void UWaterBodyComponent::OnWavesDataUpdated(UWaterWavesBase* InWaterWaves, EPro
 	// Waves data affect the navigation :
 	Params.PropertyChangedEvent.ChangeType = InChangeType;
 	Params.bShapeOrPositionChanged = true;
+	Params.bUserTriggered = true;
 	OnWaterBodyChanged(Params);
 }
 
-void UWaterBodyComponent::OnWaterSplineMetadataChanged(UWaterSplineMetadata* InWaterSplineMetadata, FPropertyChangedEvent& PropertyChangedEvent)
+void UWaterBodyComponent::OnWaterSplineMetadataChanged(const FOnWaterSplineMetadataChangedParams& InParams)
 {
 	bool bShapeOrPositionChanged = false;
 
-	FName ChangedProperty = PropertyChangedEvent.GetPropertyName();
+	FName ChangedProperty = InParams.PropertyChangedEvent.GetPropertyName();
 	if ((ChangedProperty == NAME_None)
 		|| (ChangedProperty == GET_MEMBER_NAME_CHECKED(UWaterSplineMetadata, Depth))
 		|| (ChangedProperty == GET_MEMBER_NAME_CHECKED(UWaterSplineMetadata, RiverWidth))
@@ -1178,8 +1182,12 @@ void UWaterBodyComponent::OnWaterSplineMetadataChanged(UWaterSplineMetadata* InW
 		GetWaterSpline()->SynchronizeWaterProperties();
 	}
 
-	// Waves data affect the navigation : 
-	OnWaterBodyChanged(bShapeOrPositionChanged);
+	// Waves data affect the navigation :
+	FOnWaterBodyChangedParams Params;
+	Params.bShapeOrPositionChanged = bShapeOrPositionChanged;
+	Params.bWeightmapSettingsChanged = false;
+	Params.bUserTriggered = InParams.bUserTriggered;
+	OnWaterBodyChanged(Params); 
 }
 
 void UWaterBodyComponent::RegisterOnChangeWaterSplineData(bool bRegister)
@@ -1200,11 +1208,11 @@ void UWaterBodyComponent::RegisterOnChangeWaterSplineData(bool bRegister)
 	{
 		if (bRegister)
 		{
-			WaterSplineMetadata->OnChangeData.AddUObject(this, &UWaterBodyComponent::OnWaterSplineMetadataChanged);
+			WaterSplineMetadata->OnChangeMetadata.AddUObject(this, &UWaterBodyComponent::OnWaterSplineMetadataChanged);
 		}
 		else
 		{
-			WaterSplineMetadata->OnChangeData.RemoveAll(this);
+			WaterSplineMetadata->OnChangeMetadata.RemoveAll(this);
 		}
 	}
 }
@@ -1395,6 +1403,7 @@ void UWaterBodyComponent::OnWaterBodyChanged(const FOnWaterBodyChangedParams& In
 	AWaterBody* const WaterBodyActor = GetWaterBodyActor();
 	// Transfer the FOnWaterBodyChangedParams parameters to FWaterBrushActorChangedEventParams :
 	IWaterBrushActorInterface::FWaterBrushActorChangedEventParams Params(WaterBodyActor, InParams.PropertyChangedEvent);
+	Params.bUserTriggered = InParams.bUserTriggered;
 	Params.bShapeOrPositionChanged = InParams.bShapeOrPositionChanged;
 	Params.bWeightmapSettingsChanged = InParams.bWeightmapSettingsChanged;
 	WaterBodyActor->BroadcastWaterBrushActorChangedEvent(Params);
diff --git a/Engine/Plugins/Experimental/Water/Source/Runtime/Private/WaterBodyIslandActor.cpp b/Engine/Plugins/Experimental/Water/Source/Runtime/Private/WaterBodyIslandActor.cpp
--- a/Engine/Plugins/Experimental/Water/Source/Runtime/Private/WaterBodyIslandActor.cpp
+++ b/Engine/Plugins/Experimental/Water/Source/Runtime/Private/WaterBodyIslandActor.cpp
@@ -238,6 +238,7 @@ void AWaterBodyIsland::PostEditMove(bool bFinished)
 	FOnWaterBodyIslandChangedParams Params;
 	Params.PropertyChangedEvent.ChangeType = bFinished ? EPropertyChangeType::ValueSet : EPropertyChangeType::Interactive;
 	Params.bShapeOrPositionChanged = true;
+	Params.bUserTriggered = true;
 	UpdateAll(Params);
 }
 
@@ -258,6 +259,7 @@ void AWaterBodyIsland::PostEditImport()
 	FOnWaterBodyIslandChangedParams Params;
 	Params.bShapeOrPositionChanged = true;
 	Params.bWeightmapSettingsChanged = true;
+	Params.bUserTriggered = true;
 	UpdateAll(Params);
 }
 
@@ -298,6 +300,7 @@ void AWaterBodyIsland::PostEditChangeProperty(FPropertyChangedEvent& PropertyCha
 	Super::PostEditChangeProperty(PropertyChangedEvent);
 
 	FOnWaterBodyIslandChangedParams Params(PropertyChangedEvent);
+	Params.bUserTriggered = true;
 	if (PropertyChangedEvent.GetPropertyName() == GET_MEMBER_NAME_CHECKED(UWaterBodyComponent, LayerWeightmapSettings))
 	{
 		Params.bWeightmapSettingsChanged = true;
@@ -315,6 +318,7 @@ void AWaterBodyIsland::OnWaterSplineDataChanged(const FOnWaterSplineDataChangedP
 	// Transfer the FOnWaterSplineDataChangedParams parameters to FOnWaterBodyIslandChangedParams :
 	FOnWaterBodyIslandChangedParams Params(InParams.PropertyChangedEvent);
 	Params.bShapeOrPositionChanged = true;
+	Params.bUserTriggered = true;
 	OnWaterBodyIslandChanged(Params);
 }
 
@@ -325,6 +329,7 @@ void AWaterBodyIsland::OnWaterBodyIslandChanged(const FOnWaterBodyIslandChangedP
 	FWaterBrushActorChangedEventParams Params(this, InParams.PropertyChangedEvent);
 	Params.bShapeOrPositionChanged = InParams.bShapeOrPositionChanged;
 	Params.bWeightmapSettingsChanged = InParams.bWeightmapSettingsChanged;
+	Params.bUserTriggered = InParams.bUserTriggered;
 	BroadcastWaterBrushActorChangedEvent(Params);
 #endif
 }
diff --git a/Engine/Plugins/Experimental/Water/Source/Runtime/Private/WaterSplineComponent.cpp b/Engine/Plugins/Experimental/Water/Source/Runtime/Private/WaterSplineComponent.cpp
--- a/Engine/Plugins/Experimental/Water/Source/Runtime/Private/WaterSplineComponent.cpp
+++ b/Engine/Plugins/Experimental/Water/Source/Runtime/Private/WaterSplineComponent.cpp
@@ -154,7 +154,10 @@ void UWaterSplineComponent::PostEditChangeProperty(FPropertyChangedEvent& Proper
 
 	SynchronizeWaterProperties();
 
-	WaterSplineDataChangedEvent.Broadcast(FOnWaterSplineDataChangedParams(PropertyChangedEvent));
+	FOnWaterSplineDataChangedParams OnWaterSplineDataChangedParams(PropertyChangedEvent);
+	OnWaterSplineDataChangedParams.bUserTriggered = true;
+	
+	WaterSplineDataChangedEvent.Broadcast(OnWaterSplineDataChangedParams);
 }
 
 void UWaterSplineComponent::PostEditImport()
diff --git a/Engine/Plugins/Experimental/Water/Source/Runtime/Private/WaterSplineMetadata.cpp b/Engine/Plugins/Experimental/Water/Source/Runtime/Private/WaterSplineMetadata.cpp
--- a/Engine/Plugins/Experimental/Water/Source/Runtime/Private/WaterSplineMetadata.cpp
+++ b/Engine/Plugins/Experimental/Water/Source/Runtime/Private/WaterSplineMetadata.cpp
@@ -38,8 +38,12 @@ bool UWaterSplineMetadata::CanEditVelocity() const
 void UWaterSplineMetadata::PostEditChangeProperty(struct FPropertyChangedEvent& PropertyChangedEvent)
 {
 	Super::PostEditChangeProperty(PropertyChangedEvent);
-
-	OnChangeData.Broadcast(this, PropertyChangedEvent);
+	
+	FOnWaterSplineMetadataChangedParams WaterSplineMetadataChangedParams(PropertyChangedEvent);
+	WaterSplineMetadataChangedParams.WaterSplineMetadata = this;
+	WaterSplineMetadataChangedParams.bUserTriggered = true;
+	
+	OnChangeMetadata.Broadcast(WaterSplineMetadataChangedParams);
 }
 #endif // WITH_EDITOR
 
diff --git a/Engine/Plugins/Experimental/Water/Source/Runtime/Public/WaterBodyActor.h b/Engine/Plugins/Experimental/Water/Source/Runtime/Public/WaterBodyActor.h
--- a/Engine/Plugins/Experimental/Water/Source/Runtime/Public/WaterBodyActor.h
+++ b/Engine/Plugins/Experimental/Water/Source/Runtime/Public/WaterBodyActor.h
@@ -175,7 +175,13 @@ public:
 
 	UE_DEPRECATED(4.27, "Moved to WaterBodyComponent")
 	UFUNCTION(BlueprintCallable, Category=Water, meta = (DeprecatedFunction))
-	void OnWaterBodyChanged(bool bShapeOrPositionChanged, bool bWeightmapSettingsChanged = false) { return WaterBodyComponent->OnWaterBodyChanged(bShapeOrPositionChanged, bWeightmapSettingsChanged); }
+	void OnWaterBodyChanged(bool bShapeOrPositionChanged, bool bWeightmapSettingsChanged = false)
+	{ 
+		FOnWaterBodyChangedParams Params;
+		Params.bShapeOrPositionChanged = bShapeOrPositionChanged;
+		Params.bWeightmapSettingsChanged = bWeightmapSettingsChanged;
+		return WaterBodyComponent->OnWaterBodyChanged(Params);
+	}
 
 #if WITH_EDITOR
 	UE_DEPRECATED(5.1, "Moved to WaterBodyComponent")
diff --git a/Engine/Plugins/Experimental/Water/Source/Runtime/Public/WaterBodyComponent.h b/Engine/Plugins/Experimental/Water/Source/Runtime/Public/WaterBodyComponent.h
--- a/Engine/Plugins/Experimental/Water/Source/Runtime/Public/WaterBodyComponent.h
+++ b/Engine/Plugins/Experimental/Water/Source/Runtime/Public/WaterBodyComponent.h
@@ -122,6 +122,9 @@ struct FOnWaterBodyChangedParams
 
 	/** Indicates that a property affecting the terrain weightmaps has changed */
 	bool bWeightmapSettingsChanged = false;
+
+	/** Indicates user initiated Parameter change */
+	bool bUserTriggered = false;
 };
 
 
@@ -326,7 +329,8 @@ public:
 	virtual void OnUnregister() override;
 	virtual void PostDuplicate(bool bDuplicateForPie) override;
 
-	UFUNCTION(BlueprintCallable, Category=Water)
+	UE_DEPRECATED(5.2, "Use version which takes FOnWaterBodyChangedParams")
+	UFUNCTION(BlueprintCallable, Category=Water, meta=(Deprecated = "5.2"))
 	void OnWaterBodyChanged(bool bShapeOrPositionChanged, bool bWeightmapSettingsChanged = false);
 
 	/** Fills wave-related information at the given world position and for this water depth.
@@ -511,7 +515,7 @@ protected:
 	void OnWaterSplineDataChanged(const FOnWaterSplineDataChangedParams& InParams);
 	void RegisterOnUpdateWavesData(UWaterWavesBase* InWaterWaves, bool bRegister);
 	void OnWavesDataUpdated(UWaterWavesBase* InWaterWaves, EPropertyChangeType::Type InChangeType);
-	void OnWaterSplineMetadataChanged(UWaterSplineMetadata* InWaterSplineMetadata, FPropertyChangedEvent& PropertyChangedEvent);
+	void OnWaterSplineMetadataChanged(const FOnWaterSplineMetadataChangedParams& InParams);
 	void RegisterOnChangeWaterSplineData(bool bRegister);
 
 	void CreateWaterSpriteComponent();
diff --git a/Engine/Plugins/Experimental/Water/Source/Runtime/Public/WaterBodyIslandActor.h b/Engine/Plugins/Experimental/Water/Source/Runtime/Public/WaterBodyIslandActor.h
--- a/Engine/Plugins/Experimental/Water/Source/Runtime/Public/WaterBodyIslandActor.h
+++ b/Engine/Plugins/Experimental/Water/Source/Runtime/Public/WaterBodyIslandActor.h
@@ -33,6 +33,9 @@ struct FOnWaterBodyIslandChangedParams
 
 	/** Indicates that a property affecting the terrain weightmaps has changed */
 	bool bWeightmapSettingsChanged = false;
+
+	/** Indicates user initiated change*/
+	bool bUserTriggered = false;
 };
 
 
diff --git a/Engine/Plugins/Experimental/Water/Source/Runtime/Public/WaterBrushActorInterface.h b/Engine/Plugins/Experimental/Water/Source/Runtime/Public/WaterBrushActorInterface.h
--- a/Engine/Plugins/Experimental/Water/Source/Runtime/Public/WaterBrushActorInterface.h
+++ b/Engine/Plugins/Experimental/Water/Source/Runtime/Public/WaterBrushActorInterface.h
@@ -100,6 +100,9 @@ class WATER_API IWaterBrushActorInterface
 
 		/** Indicates that a property affecting the terrain weightmaps has changed */
 		bool bWeightmapSettingsChanged = false;
+
+		/** Indicates user initiated Parameter change */
+		bool bUserTriggered = false;
 	};
 
 	/** 
diff --git a/Engine/Plugins/Experimental/Water/Source/Runtime/Public/WaterSplineComponent.h b/Engine/Plugins/Experimental/Water/Source/Runtime/Public/WaterSplineComponent.h
--- a/Engine/Plugins/Experimental/Water/Source/Runtime/Public/WaterSplineComponent.h
+++ b/Engine/Plugins/Experimental/Water/Source/Runtime/Public/WaterSplineComponent.h
@@ -17,6 +17,9 @@ struct FOnWaterSplineDataChangedParams
 
 	/** Provides some additional context about how the water brush actor data has changed (property, type of change...) */
 	FPropertyChangedEvent PropertyChangedEvent;
+
+	/** Indicates user initiated change*/
+	bool bUserTriggered = false;
 };
 DECLARE_MULTICAST_DELEGATE_OneParam(FOnWaterSplineDataChanged, const FOnWaterSplineDataChangedParams&);
 #endif // WITH_EDITOR
diff --git a/Engine/Plugins/Experimental/Water/Source/Runtime/Public/WaterSplineMetadata.h b/Engine/Plugins/Experimental/Water/Source/Runtime/Public/WaterSplineMetadata.h
--- a/Engine/Plugins/Experimental/Water/Source/Runtime/Public/WaterSplineMetadata.h
+++ b/Engine/Plugins/Experimental/Water/Source/Runtime/Public/WaterSplineMetadata.h
@@ -8,6 +8,27 @@
 
 #include "WaterSplineMetadata.generated.h"
 
+
+#if WITH_EDITOR
+
+class UWaterSplineMetadata;
+
+struct FOnWaterSplineMetadataChangedParams
+{
+	FOnWaterSplineMetadataChangedParams(const FPropertyChangedEvent& InPropertyChangedEvent = FPropertyChangedEvent(/*InProperty = */nullptr)) 
+		: PropertyChangedEvent(InPropertyChangedEvent)
+	{}
+
+	/** Provides some additional context about how the water brush actor data has changed (property, type of change...) */
+	FPropertyChangedEvent PropertyChangedEvent;
+
+	UWaterSplineMetadata* WaterSplineMetadata = nullptr;
+		
+	/** Indicates user initiated change*/
+	bool bUserTriggered = false;
+};
+#endif // WITH_EDITOR
+
 USTRUCT()
 struct FWaterSplineCurveDefaults
 {
@@ -25,7 +46,7 @@ struct FWaterSplineCurveDefaults
 
 	UPROPERTY(EditDefaultsOnly, Category = WaterCurveDefaults)
 	float DefaultWidth;
-
+	
 	UPROPERTY(EditDefaultsOnly, Category = WaterCurveDefaults)
 	float DefaultVelocity;
 
@@ -87,8 +108,13 @@ public:
 	bool bShouldVisualizeDepth;
 
 	// Delegate called whenever the metadata is updated
+	DECLARE_MULTICAST_DELEGATE_OneParam(FOnWaterSplineMetadataChanged, const FOnWaterSplineMetadataChangedParams&);
+	FOnWaterSplineMetadataChanged OnChangeMetadata;
+	
 	DECLARE_MULTICAST_DELEGATE_TwoParams(FOnChangeData, UWaterSplineMetadata* /*WaterSplineMetadata*/, FPropertyChangedEvent& /*PropertyChangedEvent*/);
-	FOnChangeData OnChangeData;
+	UE_DEPRECATED(5.2, "Use OnChangeMetadata")
+    FOnChangeData OnChangeData;
+    	
 #endif // WITH_EDITORONLY_DATA
 
 private:
diff --git a/Engine/Source/Editor/LandscapeEditor/Private/LandscapeEditorModule.cpp b/Engine/Source/Editor/LandscapeEditor/Private/LandscapeEditorModule.cpp
--- a/Engine/Source/Editor/LandscapeEditor/Private/LandscapeEditorModule.cpp
+++ b/Engine/Source/Editor/LandscapeEditor/Private/LandscapeEditorModule.cpp
@@ -144,13 +144,21 @@ public:
 			FUIAction ActionBuildNanite(FExecuteAction::CreateStatic(&BuildNanite), FCanExecuteAction());
 			Section.AddMenuEntry(NAME_None, LOCTEXT("BuildNaniteOnly", "Build Nanite Only"), LOCTEXT("BuildLandscapeNanite", "Build Nanite representation"), TAttribute<FSlateIcon>(), ActionBuildNanite, EUserInterfaceActionType::Button);
 
-			if (ULandscapeSubsystem::IsDirtyOnlyInModeEnabled())
-			{
-				FUIAction ActionSaveModifiedLandscapes(FExecuteAction::CreateStatic(&SaveModifiedLandscapes), FCanExecuteAction::CreateStatic(&HasModifiedLandscapes));
-				Section.AddMenuEntry(NAME_None,
-					LOCTEXT("SaveModifiedLandscapes", "Save Modified Landscapes"), LOCTEXT("SaveModifiedLandscapesToolTip", "Save landscapes that were modified outside of the editor mode"),
-					TAttribute<FSlateIcon>(), ActionSaveModifiedLandscapes, EUserInterfaceActionType::Button);
-			}
+			FUIAction ActionSaveModifiedLandscapes(FExecuteAction::CreateStatic(&SaveModifiedLandscapes),
+				FCanExecuteAction::CreateLambda([]()
+				{
+					if (!ULandscapeSubsystem::IsDirtyOnlyInModeEnabled())
+					{
+						return false;
+					}
+
+					return HasModifiedLandscapes();
+				}
+			));
+			
+			Section.AddMenuEntry(NAME_None,
+				LOCTEXT("SaveModifiedLandscapes", "Save Modified Landscapes"), LOCTEXT("SaveModifiedLandscapesToolTip", "Save landscapes that were modified outside of the editor mode"),
+				TAttribute<FSlateIcon>(), ActionSaveModifiedLandscapes, EUserInterfaceActionType::Button);
 		}
 		
 		ILandscapeModule& LandscapeModule = FModuleManager::GetModuleChecked<ILandscapeModule>("Landscape");
@@ -280,11 +288,13 @@ public:
 
 	static bool HasModifiedLandscapes()
 	{
+		check(ULandscapeSubsystem::IsDirtyOnlyInModeEnabled());
+		
 		if (UWorld* World = GEditor->GetEditorWorldContext().World())
 		{
 			if (ULandscapeSubsystem* LandscapeSubsystem = World->GetSubsystem<ULandscapeSubsystem>())
 			{
-				return LandscapeSubsystem->HasModifiedLandscapes();
+				return  LandscapeSubsystem->HasModifiedLandscapes();
 			}
 		}
 
diff --git a/Engine/Source/Runtime/Landscape/Classes/Landscape.h b/Engine/Source/Runtime/Landscape/Classes/Landscape.h
--- a/Engine/Source/Runtime/Landscape/Classes/Landscape.h
+++ b/Engine/Source/Runtime/Landscape/Classes/Landscape.h
@@ -282,7 +282,7 @@ public:
 	LANDSCAPE_API virtual void UpdateCachedHasLayersContent(bool bInCheckComponentDataIntegrity) override;
 	LANDSCAPE_API void RequestSplineLayerUpdate();
 	LANDSCAPE_API void RequestLayersInitialization(bool bInRequestContentUpdate = true);
-	LANDSCAPE_API void RequestLayersContentUpdateForceAll(ELandscapeLayerUpdateMode InModeMask = ELandscapeLayerUpdateMode::Update_All);
+	LANDSCAPE_API void RequestLayersContentUpdateForceAll(ELandscapeLayerUpdateMode InModeMask = ELandscapeLayerUpdateMode::Update_All, bool bInUserTriggered = false);
 	LANDSCAPE_API void RequestLayersContentUpdate(ELandscapeLayerUpdateMode InModeMask);
 	LANDSCAPE_API bool ReorderLayer(int32 InStartingLayerIndex, int32 InDestinationLayerIndex);
 	LANDSCAPE_API FLandscapeLayer* DuplicateLayerAndMoveBrushes(const FLandscapeLayer& InOtherLayer);
diff --git a/Engine/Source/Runtime/Landscape/Classes/LandscapeComponent.h b/Engine/Source/Runtime/Landscape/Classes/LandscapeComponent.h
--- a/Engine/Source/Runtime/Landscape/Classes/LandscapeComponent.h
+++ b/Engine/Source/Runtime/Landscape/Classes/LandscapeComponent.h
@@ -10,6 +10,7 @@
 #include "PerPlatformProperties.h"
 #include "Serialization/BulkData.h"
 #include "LandscapePhysicalMaterial.h"
+#include "LandscapeInfo.h"
 #include "LandscapeWeightmapUsage.h"
 #include "Containers/ArrayView.h"
 #include "Engine/StreamableRenderAsset.h"
@@ -455,6 +456,11 @@ class ULandscapeComponent : public UPrimitiveComponent
 	TLazyObjectPtr<ULandscapeHeightfieldCollisionComponent> CollisionComponent;
 
 private:
+
+	/** Store  */ 
+	UPROPERTY(Transient)
+	bool bUserTriggeredChangeRequested = false;
+	
 	UPROPERTY(Transient)
 	bool bNaniteActive;
 
@@ -1045,9 +1051,9 @@ public:
 	LANDSCAPE_API FIntRect GetComponentExtent() const;
 
 	LANDSCAPE_API void ClearUpdateFlagsForModes(uint32 InModeMask);
-	LANDSCAPE_API void RequestWeightmapUpdate(bool bUpdateAll = false, bool bUpdateCollision = true);
-	LANDSCAPE_API void RequestHeightmapUpdate(bool bUpdateAll = false, bool bUpdateCollision = true);
-	LANDSCAPE_API void RequestEditingClientUpdate();
+	LANDSCAPE_API void RequestWeightmapUpdate(bool bUpdateAll = false, bool bUpdateCollision = true, bool bInUserTriggered = false);
+	LANDSCAPE_API void RequestHeightmapUpdate(bool bUpdateAll = false, bool bUpdateCollision = true, bool bInUserTriggered = false);
+	LANDSCAPE_API void RequestEditingClientUpdate(bool bInUserTriggered = false);
 	LANDSCAPE_API void RequestDeferredClientUpdate();
 	LANDSCAPE_API uint32 GetLayerUpdateFlagPerMode() const { return LayerUpdateFlagPerMode; }
 	LANDSCAPE_API uint32 ComputeWeightmapsHash();
@@ -1086,6 +1092,16 @@ public:
 		return bNaniteActive;
 	}
 
+	void SetUserTriggeredChangeRequested(bool bInUserTriggeredChangeRequested)
+	{
+		bUserTriggeredChangeRequested = bInUserTriggeredChangeRequested;
+	}
+
+	bool GetUserTriggeredChangeRequested() const
+	{
+		return bUserTriggeredChangeRequested;
+	}
+
 protected:
 
 #if WITH_EDITOR
diff --git a/Engine/Source/Runtime/Landscape/Classes/LandscapeInfo.h b/Engine/Source/Runtime/Landscape/Classes/LandscapeInfo.h
--- a/Engine/Source/Runtime/Landscape/Classes/LandscapeInfo.h
+++ b/Engine/Source/Runtime/Landscape/Classes/LandscapeInfo.h
@@ -51,6 +51,7 @@ struct FLandscapeDirtyOnlyInModeScope
 {
 	FLandscapeDirtyOnlyInModeScope() = delete;
 	FLandscapeDirtyOnlyInModeScope(ULandscapeInfo* InLandscapeInfo);
+	FLandscapeDirtyOnlyInModeScope(ULandscapeInfo* InLandscapeInfo, bool bInOverrideDirtyMode); /* Override the dirtying in Landscape mode behaviour*/
 	~FLandscapeDirtyOnlyInModeScope();
 
 private:
diff --git a/Engine/Source/Runtime/Landscape/Private/Landscape.cpp b/Engine/Source/Runtime/Landscape/Private/Landscape.cpp
--- a/Engine/Source/Runtime/Landscape/Private/Landscape.cpp
+++ b/Engine/Source/Runtime/Landscape/Private/Landscape.cpp
@@ -69,6 +69,7 @@ Landscape.cpp: Terrain rendering
 #include "Rendering/Texture2DResource.h"
 #include "RenderCaptureInterface.h"
 #include "VisualLogger/VisualLogger.h"
+#include "LandscapeSettings.h"
 
 #if WITH_EDITOR
 #include "LandscapeEdit.h"
@@ -3732,7 +3733,8 @@ void ALandscape::PreSave(FObjectPreSaveContext ObjectSaveContext)
 
 bool ULandscapeInfo::IsDirtyOnlyInModeEnabled()
 {
-	return GLandscapeDirtyOnlyInModeEnabled > 0;
+	const ULandscapeSettings* Settings = GetDefault<ULandscapeSettings>();
+	return Settings->LandscapeDirtyingMode == ELandscapeDirtyingMode::InLandscapeModeOnly || Settings->LandscapeDirtyingMode == ELandscapeDirtyingMode::InLandscapeModeAndUserTriggeredChanges;
 }
 
 FLandscapeDirtyOnlyInModeScope::FLandscapeDirtyOnlyInModeScope(ULandscapeInfo* InLandscapeInfo)
@@ -3742,6 +3744,14 @@ FLandscapeDirtyOnlyInModeScope::FLandscapeDirtyOnlyInModeScope(ULandscapeInfo* I
 	LandscapeInfo->bDirtyOnlyInMode = ULandscapeInfo::IsDirtyOnlyInModeEnabled();
 }
 
+
+FLandscapeDirtyOnlyInModeScope::FLandscapeDirtyOnlyInModeScope(ULandscapeInfo* InLandscapeInfo, bool bInOverrideDirtyMode)
+	: LandscapeInfo(InLandscapeInfo)
+	, bDirtyOnlyInModePrevious(InLandscapeInfo->bDirtyOnlyInMode)
+{
+	LandscapeInfo->bDirtyOnlyInMode = bInOverrideDirtyMode;
+}
+
 FLandscapeDirtyOnlyInModeScope::~FLandscapeDirtyOnlyInModeScope()
 {
 	LandscapeInfo->bDirtyOnlyInMode = bDirtyOnlyInModePrevious;
@@ -3837,7 +3847,7 @@ void ULandscapeInfo::ModifyObject(UObject* InObject, bool bAlwaysMarkDirty)
 		{
 			InObject->Modify(true);
 		}
-		else
+		else 
 		{
 			InObject->Modify(false);
 			ModifiedPackages.Add(InObject->GetPackage());
@@ -3857,7 +3867,7 @@ ALandscapeProxy* ULandscapeInfo::GetLandscapeProxyForLevel(ULevel* Level) const
 	});
 	return LandscapeProxy;
 }
-
+      
 #endif // WITH_EDITOR
 
 ALandscapeProxy* ULandscapeInfo::GetCurrentLevelLandscapeProxy(bool bRegistered) const
diff --git a/Engine/Source/Runtime/Landscape/Private/LandscapeBlueprintBrushBase.cpp b/Engine/Source/Runtime/Landscape/Private/LandscapeBlueprintBrushBase.cpp
--- a/Engine/Source/Runtime/Landscape/Private/LandscapeBlueprintBrushBase.cpp
+++ b/Engine/Source/Runtime/Landscape/Private/LandscapeBlueprintBrushBase.cpp
@@ -4,6 +4,7 @@
 #include "CoreMinimal.h"
 #include "LandscapeProxy.h"
 #include "Landscape.h"
+#include "LandscapeInfo.h"
 #include "LandscapePrivate.h"
 #include "Misc/MapErrors.h"
 #include "Misc/UObjectToken.h"
@@ -59,7 +60,7 @@ void ALandscapeBlueprintBrushBase::Initialize_Implementation(const FTransform& I
 	Initialize_Native(InLandscapeTransform, InLandscapeSize, InLandscapeRenderTargetSize);
 }
 
-void ALandscapeBlueprintBrushBase::RequestLandscapeUpdate()
+void ALandscapeBlueprintBrushBase::RequestLandscapeUpdate(bool bInUserTriggered)
 {
 #if WITH_EDITORONLY_DATA
 	UE_LOG(LogLandscape, Verbose, TEXT("ALandscapeBlueprintBrushBase::RequestLandscapeUpdate"));
@@ -76,7 +77,7 @@ void ALandscapeBlueprintBrushBase::RequestLandscapeUpdate()
 		}
 		if (ModeMask)
 		{
-			OwningLandscape->RequestLayersContentUpdateForceAll((ELandscapeLayerUpdateMode)ModeMask);
+			OwningLandscape->RequestLayersContentUpdateForceAll((ELandscapeLayerUpdateMode)ModeMask, bInUserTriggered);
 			// Just in case differentiate between 0 (default value and frame number)
 			LastRequestLayersContentUpdateFrameNumber = GFrameNumber == InvalidLastRequestLayersContentUpdateFrameNumber ? GFrameNumber + 1 : GFrameNumber;
 		}
diff --git a/Engine/Source/Runtime/Landscape/Private/LandscapeEditLayers.cpp b/Engine/Source/Runtime/Landscape/Private/LandscapeEditLayers.cpp
--- a/Engine/Source/Runtime/Landscape/Private/LandscapeEditLayers.cpp
+++ b/Engine/Source/Runtime/Landscape/Private/LandscapeEditLayers.cpp
@@ -5219,9 +5219,24 @@ bool ALandscape::ResolveLayersTexture(
 		InCPUReadback->Tick();
 	}
 
+	const int32 CompletedReadbackNum = InCPUReadback->GetCompletedResultNum();
+
+	bool bUserTriggered = false;
+
+	if (TArray<ULandscapeComponent*> const * Components = bIsWeightmap ? MapHelper.WeightmapToComponents.Find(InOutputTexture) : MapHelper.HeightmapToComponents.Find(InOutputTexture))
+	{
+		for (ULandscapeComponent* Component : *Components)
+		{
+			if (Component->GetUserTriggeredChangeRequested())
+			{
+				bUserTriggered = true;
+				break;
+			}
+		}
+	}
+	
 	bool bChanged = false;
 
-	const int32 CompletedReadbackNum = InCPUReadback->GetCompletedResultNum();
 	if (CompletedReadbackNum > 0)
 	{
 		TRACE_CPUPROFILER_EVENT_SCOPE(LandscapeLayers_PerformReadbacks);
@@ -5257,7 +5272,16 @@ bool ALandscape::ResolveLayersTexture(
 						bChanged = true;
 
 						// We're about to modify the texture's source data, the texture needs to know so that it can handle properly update cached platform data (additionally, the package needs to be dirtied) :
-						GetLandscapeInfo()->ModifyObject(InOutputTexture);
+						if (GetDefault<ULandscapeSettings>()->LandscapeDirtyingMode == ELandscapeDirtyingMode::InLandscapeModeAndUserTriggeredChanges)
+						{
+							FLandscapeDirtyOnlyInModeScope Scope(GetLandscapeInfo(), !bUserTriggered);
+							GetLandscapeInfo()->ModifyObject(InOutputTexture);	
+						}
+						else
+						{
+							GetLandscapeInfo()->ModifyObject(InOutputTexture);
+						}
+							
 					}
 				}
 
@@ -7487,7 +7511,7 @@ void ALandscape::RequestLayersContentUpdate(ELandscapeLayerUpdateMode InUpdateMo
 	LayerContentUpdateModes |= InUpdateMode;
 }
 
-void ALandscape::RequestLayersContentUpdateForceAll(ELandscapeLayerUpdateMode InModeMask)
+void ALandscape::RequestLayersContentUpdateForceAll(ELandscapeLayerUpdateMode InModeMask, bool bInUserTriggered)
 {
 	// Ignore Update requests while in PostLoad (to avoid dirtying package on load)
 	if (FUObjectThreadContext::Get().IsRoutingPostLoad)
@@ -7509,7 +7533,7 @@ void ALandscape::RequestLayersContentUpdateForceAll(ELandscapeLayerUpdateMode In
 	const bool bUpdateClientUdpateEditing = (InModeMask & ELandscapeLayerUpdateMode::Update_Client_Editing) != 0;
 	if (ULandscapeInfo* LandscapeInfo = GetLandscapeInfo())
 	{
-		LandscapeInfo->ForAllLandscapeProxies([bUpdateHeightmap, bUpdateWeightmap, bUpdateAllHeightmap, bUpdateAllWeightmap, bUpdateHeightCollision, bUpdateWeightCollision, bUpdateClientUdpateEditing](ALandscapeProxy* Proxy)
+		LandscapeInfo->ForAllLandscapeProxies([bUpdateHeightmap, bUpdateWeightmap, bUpdateAllHeightmap, bUpdateAllWeightmap, bUpdateHeightCollision, bUpdateWeightCollision, bUpdateClientUdpateEditing, bInUserTriggered](ALandscapeProxy* Proxy)
 		{
 			if (Proxy)
 			{
@@ -7517,17 +7541,17 @@ void ALandscape::RequestLayersContentUpdateForceAll(ELandscapeLayerUpdateMode In
 				{
 					if (bUpdateHeightmap)
 					{
-						Component->RequestHeightmapUpdate(bUpdateAllHeightmap, bUpdateHeightCollision);
+						Component->RequestHeightmapUpdate(bUpdateAllHeightmap, bUpdateHeightCollision, bInUserTriggered);
 					}
 
 					if (bUpdateWeightmap)
 					{
-						Component->RequestWeightmapUpdate(bUpdateAllWeightmap, bUpdateWeightCollision);
+						Component->RequestWeightmapUpdate(bUpdateAllWeightmap, bUpdateWeightCollision, bInUserTriggered);
 					}
 
 					if (bUpdateClientUdpateEditing)
 					{
-						Component->RequestEditingClientUpdate();
+						Component->RequestEditingClientUpdate(bInUserTriggered);
 					}
 				}
 			}
@@ -7615,8 +7639,10 @@ void ULandscapeComponent::RequestDeferredClientUpdate()
 	LayerUpdateFlagPerMode |= ELandscapeLayerUpdateMode::Update_Client_Deferred;
 }
 
-void ULandscapeComponent::RequestEditingClientUpdate()
+void ULandscapeComponent::RequestEditingClientUpdate(bool bInUserTriggered)
 {
+	bUserTriggeredChangeRequested = bInUserTriggered; 
+	
 	LayerUpdateFlagPerMode |= ELandscapeLayerUpdateMode::Update_Client_Editing;
 	if (ALandscape* LandscapeActor = GetLandscapeActor())
 	{
@@ -7624,8 +7650,9 @@ void ULandscapeComponent::RequestEditingClientUpdate()
 	}
 }
 
-void ULandscapeComponent::RequestHeightmapUpdate(bool bUpdateAll, bool bUpdateCollision)
+void ULandscapeComponent::RequestHeightmapUpdate(bool bUpdateAll, bool bUpdateCollision, bool bInUserTriggered)
 {
+	bUserTriggeredChangeRequested = bInUserTriggered;
 	if (bUpdateAll || bUpdateCollision)
 	{
 		LayerUpdateFlagPerMode |= ELandscapeLayerUpdateMode::Update_Heightmap_Editing;
@@ -7645,8 +7672,10 @@ void ULandscapeComponent::RequestHeightmapUpdate(bool bUpdateAll, bool bUpdateCo
 	}
 }
 
-void ULandscapeComponent::RequestWeightmapUpdate(bool bUpdateAll, bool bUpdateCollision)
+void ULandscapeComponent::RequestWeightmapUpdate(bool bUpdateAll, bool bUpdateCollision, bool bInUserTriggered)
 {
+	bUserTriggeredChangeRequested = bInUserTriggered;
+	
 	if (bUpdateAll || bUpdateCollision)
 	{
 		LayerUpdateFlagPerMode |= ELandscapeLayerUpdateMode::Update_Weightmap_Editing;
@@ -7879,6 +7908,21 @@ void ALandscape::UpdateLayersContent(bool bInWaitForStreaming, bool bInSkipMonit
 	const bool bProcessReadbacks = FLandscapeEditLayerReadback::HasWork();
 	const bool bForceRender = CVarForceLayersUpdate.GetValueOnAnyThread() != 0;
 
+	// User triggered change has been completely processed, resetting user triggered flag on all components.
+	if (!bProcessReadbacks && (LayerContentUpdateModes == 0))
+	{
+		GetLandscapeInfo()->ForAllLandscapeComponents(
+			[this](ULandscapeComponent* Component) -> void
+			{
+				if (Component->GetUserTriggeredChangeRequested())
+				{
+					check(Component->GetLayerUpdateFlagPerMode() == 0);
+					Component->SetUserTriggeredChangeRequested(/* bInUserTriggered = */false);	
+				}
+			}
+		);	
+	}
+	
 	if (LayerContentUpdateModes == 0 && !bForceRender && !bProcessReadbacks)
 	{
 		return;
@@ -8008,7 +8052,6 @@ void ALandscape::UpdateLayersContent(bool bInWaitForStreaming, bool bInSkipMonit
 	// Apply post resolve updates
 	const uint32 ToProcessModes = UpdateAfterReadbackResolves(UpdateLayersContentContext.AllLandscapeComponentReadbackResults);
 	LayerContentUpdateModes |= ToProcessModes;
-
 	if (LandscapeEdMode)
 	{
 		LandscapeEdMode->PostUpdateLayerContent();
@@ -8060,11 +8103,13 @@ uint32 ALandscape::UpdateCollisionAndClients(const TArrayView<FLandscapeEditLaye
 
 	for (const FLandscapeEditLayerComponentReadbackResult& ComponentReadbackResult : InComponentReadbackResults)
 	{
+		ULandscapeComponent* LandscapeComponent = ComponentReadbackResult.LandscapeComponent;
+		
 		bool bDeferClientUpdateForComponent = false;
 		bool bDoUpdateClient = true;
 		if (IsUpdateFlagEnabledForModes(ELandscapeComponentUpdateFlag::Component_Update_Recreate_Collision, ComponentReadbackResult.UpdateModes))
 		{
-			if (ULandscapeHeightfieldCollisionComponent* CollisionComp = ComponentReadbackResult.LandscapeComponent->CollisionComponent.Get())
+			if (ULandscapeHeightfieldCollisionComponent* CollisionComp = LandscapeComponent->CollisionComponent.Get())
 			{
 				FEnableCollisionHashOptimScope Scope(CollisionComp);
 				bDoUpdateClient = CollisionComp->RecreateCollision();
@@ -8075,7 +8120,7 @@ uint32 ALandscape::UpdateCollisionAndClients(const TArrayView<FLandscapeEditLaye
 		{
 			if (!GUndo)
 			{
-				if (ULandscapeHeightfieldCollisionComponent* CollisionComp = ComponentReadbackResult.LandscapeComponent->CollisionComponent.Get())
+				if (ULandscapeHeightfieldCollisionComponent* CollisionComp = LandscapeComponent->CollisionComponent.Get())
 				{
 					FNavigationSystem::UpdateComponentData(*CollisionComp);
 					CollisionComp->SnapFoliageInstances();
@@ -8092,7 +8137,7 @@ uint32 ALandscape::UpdateCollisionAndClients(const TArrayView<FLandscapeEditLaye
 		{
 			if (LandscapeEdModeInfo.ViewMode == ELandscapeViewMode::LayerContribution)
 			{
-				check(ComponentSizeQuads == ComponentReadbackResult.LandscapeComponent->ComponentSizeQuads);
+				check(ComponentSizeQuads == LandscapeComponent->ComponentSizeQuads);
 				const int32 Stride = (1 + ComponentSizeQuads);
 				const int32 ArraySize = Stride * Stride;
 				if (LayerContributionMaskData.Num() != ArraySize)
@@ -8100,13 +8145,13 @@ uint32 ALandscape::UpdateCollisionAndClients(const TArrayView<FLandscapeEditLaye
 					LayerContributionMaskData.AddZeroed(ArraySize);
 				}
 				uint8* LayerContributionMaskDataPtr = LayerContributionMaskData.GetData();
-				const int32 X1 = ComponentReadbackResult.LandscapeComponent->GetSectionBase().X;
+				const int32 X1 = LandscapeComponent->GetSectionBase().X;
 				const int32 X2 = X1 + ComponentSizeQuads;
-				const int32 Y1 = ComponentReadbackResult.LandscapeComponent->GetSectionBase().Y;
+				const int32 Y1 = LandscapeComponent->GetSectionBase().Y;
 				const int32 Y2 = Y1 + ComponentSizeQuads;
 				bool bLayerContributionWrittenData = false;
 
-				ULandscapeInfo* Info = ComponentReadbackResult.LandscapeComponent->GetLandscapeInfo();
+				ULandscapeInfo* Info = LandscapeComponent->GetLandscapeInfo();
 				check(Info);
 				FLandscapeEditDataInterface LandscapeEdit(Info);
 
@@ -8146,7 +8191,7 @@ uint32 ALandscape::UpdateCollisionAndClients(const TArrayView<FLandscapeEditLaye
 
 		if (bDeferClientUpdateForComponent)
 		{
-			ComponentReadbackResult.LandscapeComponent->RequestDeferredClientUpdate();
+			LandscapeComponent->RequestDeferredClientUpdate();
 		}
 	}
 
diff --git a/Engine/Source/Runtime/Landscape/Public/LandscapeBlueprintBrushBase.h b/Engine/Source/Runtime/Landscape/Public/LandscapeBlueprintBrushBase.h
--- a/Engine/Source/Runtime/Landscape/Public/LandscapeBlueprintBrushBase.h
+++ b/Engine/Source/Runtime/Landscape/Public/LandscapeBlueprintBrushBase.h
@@ -49,7 +49,7 @@ public:
 	void Initialize(const FTransform& InLandscapeTransform, const FIntPoint& InLandscapeSize, const FIntPoint& InLandscapeRenderTargetSize);
 
 	UFUNCTION(BlueprintCallable, Category = "Landscape")
-	void RequestLandscapeUpdate();
+	void RequestLandscapeUpdate(bool bInUserTriggered = false);
 
 	UFUNCTION(BlueprintImplementableEvent, CallInEditor)
 	void GetBlueprintRenderDependencies(TArray<UObject*>& OutStreamableAssets);
diff --git a/Engine/Source/Runtime/Landscape/Public/LandscapeSettings.h b/Engine/Source/Runtime/Landscape/Public/LandscapeSettings.h
--- a/Engine/Source/Runtime/Landscape/Public/LandscapeSettings.h
+++ b/Engine/Source/Runtime/Landscape/Public/LandscapeSettings.h
@@ -9,6 +9,14 @@
 class UMaterialInterface;
 class ULandscapeLayerInfoObject;
 
+UENUM()
+enum class ELandscapeDirtyingMode : uint8
+{
+	Auto,
+	InLandscapeModeOnly,
+	InLandscapeModeAndUserTriggeredChanges
+};
+
 UCLASS(config = Engine, defaultconfig, meta = (DisplayName = "Landscape"))
 class LANDSCAPE_API ULandscapeSettings : public UDeveloperSettings
 {
@@ -39,6 +47,9 @@ public:
 	UPROPERTY(config, EditAnywhere, Category = "Layers", meta=(UIMin = "1", UIMax = "32", ClampMin = "1", ClampMax = "32", ToolTip = "This option controls the maximum editing layers that can be added to a Landscape"))
 	int32 MaxNumberOfLayers = 8;
 
+	UPROPERTY(Config, Category = "Configuration", EditAnywhere)
+	ELandscapeDirtyingMode LandscapeDirtyingMode;
+
 protected:
 	UPROPERTY(config)
 	int32 SideResolutionLimit = 2048;
