Add option to sample landscape height per pixel when rendering to RVT.
When we do this we only need a single quad to render each landscape component subsection.
This reduces vertex shader cost which can sometimes dominate RVT rendering.
And it can remove the need to set difficult to understand landscape RVT LOD settings.
One downside is that bilinear texture map interpolation is different to triangle barycentric interpolation, so there is a subtle difference in extreme cases, but nothing visually bad. The difference will matter for VHM rendering though.

diff --git a/Engine/Shaders/Private/LandscapeVertexFactory.ush b/Engine/Shaders/Private/LandscapeVertexFactory.ush
--- a/Engine/Shaders/Private/LandscapeVertexFactory.ush
+++ b/Engine/Shaders/Private/LandscapeVertexFactory.ush
@@ -370,6 +370,21 @@ float3x3 VertexFactoryGetPerPixelTangentBasis(FVertexFactoryInterpolantsVSToPS I
 	return Result;
 }
 
+/** Lookup per-pixel height from heightmap texture */
+#define VF_PER_PIXEL_HEIGHTMAP 1
+
+bool HasVertexFactoryPerPixelHeight()
+{
+	return LandscapeParameters.VirtualTexturePerPixelHeight != 0;
+}
+
+float GetVertexFactoryPerPixelHeight(FVertexFactoryInterpolantsVSToPS Interpolants)
+{
+	// Use NormalmapTextureSampler for bilinear (HeightmapTextureSampler is usually point).
+	float4 SampleValue = LandscapeParameters.HeightmapTexture.Sample(LandscapeParameters.NormalmapTextureSampler, Interpolants.WeightHeightMapTexCoord.zw);
+	return DecodePackedHeight(SampleValue.xy);
+}
+
 /** Converts from vertex factory specific interpolants to a FMaterialPixelParameters, which is used by material inputs. */
 FMaterialPixelParameters GetMaterialPixelParameters(FVertexFactoryInterpolantsVSToPS Interpolants, float4 SvPosition)
 {
diff --git a/Engine/Shaders/Private/VirtualTextureMaterial.usf b/Engine/Shaders/Private/VirtualTextureMaterial.usf
--- a/Engine/Shaders/Private/VirtualTextureMaterial.usf
+++ b/Engine/Shaders/Private/VirtualTextureMaterial.usf
@@ -24,6 +24,15 @@ void MainVS(
 
 	FVertexFactoryIntermediates VFIntermediates = GetVertexFactoryIntermediates(Input);
 
+#if VF_PER_PIXEL_HEIGHTMAP
+	if (HasVertexFactoryPerPixelHeight())
+	{
+		// Special case PerPixel Height for Landscape.
+		// Remove PerVertex Height and apply in Pixel Shader.
+		VFIntermediates.LocalPosition.z = 0;
+	}
+#endif
+
 	float4 WorldPosition = VertexFactoryGetWorldPosition(Input, VFIntermediates);
 	float3x3 TangentToLocal = VertexFactoryGetTangentToLocal(Input, VFIntermediates);	
 	FMaterialVertexParameters VertexParameters = GetMaterialVertexParameters(Input, VFIntermediates, WorldPosition.xyz, TangentToLocal);
@@ -66,6 +75,17 @@ void FPixelShaderInOut_MainPS(
 
 	float4 ScreenPosition = SvPositionToResolvedScreenPosition(In.SvPosition);
 	float3 TranslatedWorldPosition = SvPositionToResolvedTranslatedWorld(In.SvPosition);
+
+#if VF_PER_PIXEL_HEIGHTMAP
+	if (HasVertexFactoryPerPixelHeight())
+	{
+		// Special case PerPixel Height for Landscape.
+		// Assumes Heightmap is camera aligned.
+		float3 LocalPositionDelta = float3(0, 0, GetVertexFactoryPerPixelHeight(Interpolants));
+		TranslatedWorldPosition += TransformLocalVectorToWorld(MaterialParameters, LocalPositionDelta);
+	}
+#endif
+
 	CalcMaterialParametersEx(MaterialParameters, PixelMaterialInputs, In.SvPosition, ScreenPosition, In.bIsFrontFace, TranslatedWorldPosition, TranslatedWorldPosition);
 
 #if VIRTUAL_TEXTURE_OUTPUT
diff --git a/Engine/Source/Runtime/Landscape/Private/LandscapeRender.cpp b/Engine/Source/Runtime/Landscape/Private/LandscapeRender.cpp
--- a/Engine/Source/Runtime/Landscape/Private/LandscapeRender.cpp
+++ b/Engine/Source/Runtime/Landscape/Private/LandscapeRender.cpp
@@ -890,9 +890,10 @@ FLandscapeComponentSceneProxy::FLandscapeComponentSceneProxy(ULandscapeComponent
 	LODSettings.LastLODScreenSizeSquared = LODScreenRatioSquared[LastLOD];
 	LODSettings.ForcedLOD = ForcedLOD;
 
+	bVirtualTextureRenderWithQuad = InComponent->GetLandscapeProxy()->bVirtualTextureRenderWithQuad;
 	LastVirtualTextureLOD = MaxLOD;
-	FirstVirtualTextureLOD = FMath::Max(MaxLOD - InComponent->GetLandscapeProxy()->VirtualTextureNumLods, 0);
-	VirtualTextureLodBias = InComponent->GetLandscapeProxy()->VirtualTextureLodBias;
+	FirstVirtualTextureLOD = bVirtualTextureRenderWithQuad ? MaxLOD : FMath::Max(MaxLOD - InComponent->GetLandscapeProxy()->VirtualTextureNumLods, 0);
+	VirtualTextureLodBias = bVirtualTextureRenderWithQuad ? 0 : InComponent->GetLandscapeProxy()->VirtualTextureLodBias;
 
 #if WITH_EDITOR || !(UE_BUILD_SHIPPING || UE_BUILD_TEST)
 	LODSettings.DrawCollisionPawnLOD = CollisionResponse.GetResponse(ECC_Pawn) == ECR_Ignore ? -1 : SimpleCollisionMipLevel;
@@ -1545,6 +1546,7 @@ void FLandscapeComponentSceneProxy::OnTransformChanged()
 	LandscapeParams.SubsectionSizeVerts = SubsectionSizeVerts;
 	LandscapeParams.NumSubsections = NumSubsections;
 	LandscapeParams.LastLOD = LastLOD;
+	LandscapeParams.VirtualTexturePerPixelHeight = bVirtualTextureRenderWithQuad ? 1 : 0;
 	LandscapeParams.HeightmapUVScaleBias = HeightmapScaleBias;
 	LandscapeParams.WeightmapUVScaleBias = WeightmapScaleBias;
 	LandscapeParams.LocalToWorldNoScaling = FMatrix44f(LocalToWorldNoScaling);			// LWC_TODO: Precision loss
diff --git a/Engine/Source/Runtime/Landscape/Public/LandscapeRender.h b/Engine/Source/Runtime/Landscape/Public/LandscapeRender.h
--- a/Engine/Source/Runtime/Landscape/Public/LandscapeRender.h
+++ b/Engine/Source/Runtime/Landscape/Public/LandscapeRender.h
@@ -96,6 +96,7 @@ BEGIN_GLOBAL_SHADER_PARAMETER_STRUCT(FLandscapeUniformShaderParameters, LANDSCAP
 	SHADER_PARAMETER(int32, SubsectionSizeVerts)
 	SHADER_PARAMETER(int32, NumSubsections)
 	SHADER_PARAMETER(int32, LastLOD)
+	SHADER_PARAMETER(uint32, VirtualTexturePerPixelHeight)
     SHADER_PARAMETER(FVector4f, HeightmapUVScaleBias)
     SHADER_PARAMETER(FVector4f, WeightmapUVScaleBias)
     SHADER_PARAMETER(FVector4f, LandscapeLightmapScaleBias)
@@ -664,6 +665,7 @@ protected:
 	int8						MaxLOD;		// Maximum LOD level, user override possible
 	int8						NumWeightmapLayerAllocations;
 	uint8						StaticLightingLOD;
+	uint8						bVirtualTextureRenderWithQuad : 1;
 	float						WeightmapSubsectionOffset;
 	TArray<float>				LODScreenRatioSquared;		// Table of valid screen size -> LOD index
 	int32						FirstLOD;	// First LOD we have batch elements for
