Added profile tag

diff --git a/Engine/Source/Runtime/Landscape/Private/Landscape.cpp b/Engine/Source/Runtime/Landscape/Private/Landscape.cpp
--- a/Engine/Source/Runtime/Landscape/Private/Landscape.cpp
+++ b/Engine/Source/Runtime/Landscape/Private/Landscape.cpp
@@ -5313,6 +5313,7 @@ void ALandscapeProxy::InvalidateGeneratedComponentData(const TSet<ULandscapeComp
 
 void ALandscapeProxy::UpdateRenderingMethod()
 {
+	TRACE_CPUPROFILER_EVENT_SCOPE(ALandscapeProxy::UpdateRenderingMethod);
 	if (LandscapeComponents.Num() == 0)
 	{
 		return;
