#UE Expose a pointer to a landscape foliage type to garbage collection

diff --git a/Engine/Source/Runtime/Landscape/Classes/LandscapeComponent.h b/Engine/Source/Runtime/Landscape/Classes/LandscapeComponent.h
--- a/Engine/Source/Runtime/Landscape/Classes/LandscapeComponent.h
+++ b/Engine/Source/Runtime/Landscape/Classes/LandscapeComponent.h
@@ -651,6 +651,8 @@ public:
 	static void DeclareConstructClasses(TArray<FTopLevelAssetPath>& OutConstructClasses, const UClass* SpecificSubclass);
 #endif
 
+	static void AddReferencedObjects(UObject* InThis, FReferenceCollector& Collector);
+
 #if WITH_EDITOR
 	virtual void BeginCacheForCookedPlatformData(const ITargetPlatform* TargetPlatform) override;
 	virtual void PreEditUndo() override;
diff --git a/Engine/Source/Runtime/Landscape/Private/Landscape.cpp b/Engine/Source/Runtime/Landscape/Private/Landscape.cpp
--- a/Engine/Source/Runtime/Landscape/Private/Landscape.cpp
+++ b/Engine/Source/Runtime/Landscape/Private/Landscape.cpp
@@ -51,6 +51,7 @@ Landscape.cpp: Terrain rendering
 #include "Materials/MaterialInstanceDynamic.h"
 #include "ProfilingDebugging/CookStats.h"
 #include "ILandscapeSplineInterface.h"
+#include "LandscapeGrassType.h"
 #include "LandscapeSplineActor.h"
 #include "LandscapeSplinesComponent.h"
 #include "EngineGlobals.h"
@@ -1220,6 +1221,14 @@ void ULandscapeComponent::DeclareConstructClasses(TArray<FTopLevelAssetPath>& Ou
 }
 #endif
 
+void ULandscapeComponent::AddReferencedObjects(UObject* InThis, FReferenceCollector& Collector)
+{
+	Super::AddReferencedObjects(InThis, Collector);
+
+	ThisClass* const TypedThis = Cast<ThisClass>(InThis);
+	Collector.AddReferencedObjects(TypedThis->GrassData->WeightOffsets, TypedThis);
+}
+
 #if WITH_EDITORONLY_DATA
 TArray<ALandscapeProxy*> ALandscapeProxy::LandscapeProxies;
 #endif
