Added ULandscapeSubsystem::PrioritizeGrassCreation method to enable/disable boosting the grass creation process.
Added CVar grass.GrassCreationPrioritizedMultipler defaulted to 4 that affects grass.MaxCreatePerFrame and grass.MaxAsyncTasks when PrioritizeGrassCreation is true.

diff --git a/Engine/Source/Runtime/Landscape/Private/LandscapeGrass.cpp b/Engine/Source/Runtime/Landscape/Private/LandscapeGrass.cpp
--- a/Engine/Source/Runtime/Landscape/Private/LandscapeGrass.cpp
+++ b/Engine/Source/Runtime/Landscape/Private/LandscapeGrass.cpp
@@ -31,6 +31,7 @@
 #include "Engine/MapBuildDataRegistry.h"
 #include "ShadowMap.h"
 #include "LandscapeComponent.h"
+#include "LandscapeSubsystem.h"
 #include "LandscapeVersion.h"
 #include "MaterialShaderType.h"
 #include "MeshMaterialShaderType.h"
@@ -182,6 +183,12 @@ static FAutoConsoleVariableRef CVarGrassMaxCreatePerFrame(
 	GGrassMaxCreatePerFrame,
 	TEXT("Maximum number of Grass components to create per frame"));
 
+static int32 GGrassCreationPrioritizedMultipler = 4;
+static FAutoConsoleVariableRef CVarGrassCreationPrioritizedMultipler(
+	TEXT("grass.GrassCreationPrioritizedMultipler"),
+	GGrassCreationPrioritizedMultipler,
+	TEXT("Multipler applied to MaxCreatePerFrame and MaxAsyncTasks when grass creation is prioritized."));
+
 static int32 GGrassUpdateAllOnRebuild = 0;
 static FAutoConsoleVariableRef CVarUpdateAllOnRebuild(
 	TEXT("grass.UpdateAllOnRebuild"),
@@ -2786,10 +2793,13 @@ void ALandscapeProxy::UpdateGrass(const TArray<FVector>& Cameras, int32& InOutNu
 		}
 	}
 
+	UWorld* World = GetWorld();
+	ULandscapeSubsystem* LandscapeSubsystem = World->GetSubsystem<ULandscapeSubsystem>();
+	const bool bIsGrassCreationPrioritized = LandscapeSubsystem ? LandscapeSubsystem->IsGrassCreationPrioritized() : false;
+	const bool bWaitAsyncTasks = bForceSync || bIsGrassCreationPrioritized;
+
 	if (GGrassEnable > 0)
 	{
-		UWorld* World = GetWorld();
-
 		float GuardBand = GGuardBandMultiplier;
 		float DiscardGuardBand = GGuardBandDiscardMultiplier;
 		bool bCullSubsections = GCullSubsections > 0;
@@ -2797,6 +2807,13 @@ void ALandscapeProxy::UpdateGrass(const TArray<FVector>& Cameras, int32& InOutNu
 		bool bDisableDynamicShadows = GDisableDynamicShadows > 0;
 		int32 MaxInstancesPerComponent = FMath::Max<int32>(1024, GMaxInstancesPerComponent);
 		int32 MaxTasks = GMaxAsyncTasks;
+		int32 GrassMaxCreatePerFrame = GGrassMaxCreatePerFrame;
+		if (bIsGrassCreationPrioritized && GGrassCreationPrioritizedMultipler > 1)
+		{
+			MaxTasks *= GGrassCreationPrioritizedMultipler;
+			GrassMaxCreatePerFrame *= GGrassCreationPrioritizedMultipler;
+		}
+			
 		const float CullDistanceScale = GGrassCullDistanceScale;
 
 		if (World)
@@ -2876,13 +2893,16 @@ void ALandscapeProxy::UpdateGrass(const TArray<FVector>& Cameras, int32& InOutNu
 				SortedLandscapeComponents.Emplace(Component, FMath::Sqrt(MinSqrDistanceToComponent), WorldBounds.GetBox());
 			}
 			
+			// Prioritize components that are closer to camera when grass creation is prioritized
+			bool bShouldSort = bIsGrassCreationPrioritized;
 #if WITH_EDITOR
-			// When editing landscape, prioritize components that are closer to camera for a more reactive update
-			if (GLandscapeEditModeActive)
+			// Also when editing landscape (for a more reactive update)
+			bShouldSort = bShouldSort || GLandscapeEditModeActive;
+#endif
+			if (bShouldSort)
 			{
 				Algo::Sort(SortedLandscapeComponents, [](const SortedLandscapeElement& A, const SortedLandscapeElement& B) { return A.MinDistance < B.MinDistance; });
 			}
-#endif
 
 			for (const SortedLandscapeElement& SortedLandscapeComponent : SortedLandscapeComponents)
 			{
@@ -3044,7 +3064,7 @@ void ALandscapeProxy::UpdateGrass(const TArray<FVector>& Cameras, int32& InOutNu
 											}
 										}
 
-										if (!bRebuildForBoxes && !bForceSync && (InOutNumCompsCreated >= GGrassMaxCreatePerFrame || AsyncFoliageTasks.Num() >= MaxTasks))
+										if (!bRebuildForBoxes && !bForceSync && (InOutNumCompsCreated >= GrassMaxCreatePerFrame || AsyncFoliageTasks.Num() >= MaxTasks))
 										{
 											continue; // one per frame, but we still want to touch the existing ones and we must do the rebuilds because we changed the tag
 										}
@@ -3376,7 +3396,7 @@ void ALandscapeProxy::UpdateGrass(const TArray<FVector>& Cameras, int32& InOutNu
 		for (int32 Index = 0; Index < AsyncFoliageTasks.Num(); Index++)
 		{
 			FAsyncTask<FAsyncGrassTask>* Task = AsyncFoliageTasks[Index];
-			if (bForceSync)
+			if (bWaitAsyncTasks)
 			{
 				Task->EnsureCompletion();
 			}
@@ -3433,7 +3453,7 @@ void ALandscapeProxy::UpdateGrass(const TArray<FVector>& Cameras, int32& InOutNu
 					Existing->Touch();
 				}
 				delete Task;
-				if (!bForceSync)
+				if (!bWaitAsyncTasks)
 				{
 					break; // one per frame is fine
 				}
diff --git a/Engine/Source/Runtime/Landscape/Private/LandscapeSubsystem.cpp b/Engine/Source/Runtime/Landscape/Private/LandscapeSubsystem.cpp
--- a/Engine/Source/Runtime/Landscape/Private/LandscapeSubsystem.cpp
+++ b/Engine/Source/Runtime/Landscape/Private/LandscapeSubsystem.cpp
@@ -39,6 +39,7 @@ DECLARE_CYCLE_STAT(TEXT("LandscapeSubsystem Tick"), STAT_LandscapeSubsystemTick,
 #define LOCTEXT_NAMESPACE "LandscapeSubsystem"
 
 ULandscapeSubsystem::ULandscapeSubsystem()
+: bIsGrassCreationPrioritized(false)
 {
 }
 
diff --git a/Engine/Source/Runtime/Landscape/Public/LandscapeSubsystem.h b/Engine/Source/Runtime/Landscape/Public/LandscapeSubsystem.h
--- a/Engine/Source/Runtime/Landscape/Public/LandscapeSubsystem.h
+++ b/Engine/Source/Runtime/Landscape/Public/LandscapeSubsystem.h
@@ -50,6 +50,9 @@ public:
 	virtual TStatId GetStatId() const override;
 	// End FTickableGameObject overrides
 
+	void PrioritizeGrassCreation(bool bPrioritizeGrassCreation) { bIsGrassCreationPrioritized = bPrioritizeGrassCreation; }
+	bool IsGrassCreationPrioritized() const { return bIsGrassCreationPrioritized; }
+
 #if WITH_EDITOR
 	LANDSCAPE_API void BuildAll();
 	LANDSCAPE_API void BuildGrassMaps();
@@ -81,6 +84,7 @@ private:
 	virtual void Deinitialize() override;
 	// End USubsystem
 
+	bool bIsGrassCreationPrioritized;
 	TArray<TWeakObjectPtr<ALandscapeProxy>> Proxies;
 
 #if WITH_EDITOR
