Fixed static analysis

diff --git a/Engine/Source/Runtime/Landscape/Private/LandscapeGrass.cpp b/Engine/Source/Runtime/Landscape/Private/LandscapeGrass.cpp
--- a/Engine/Source/Runtime/Landscape/Private/LandscapeGrass.cpp
+++ b/Engine/Source/Runtime/Landscape/Private/LandscapeGrass.cpp
@@ -2794,7 +2794,7 @@ void ALandscapeProxy::UpdateGrass(const TArray<FVector>& Cameras, int32& InOutNu
 	}
 
 	UWorld* World = GetWorld();
-	ULandscapeSubsystem* LandscapeSubsystem = World->GetSubsystem<ULandscapeSubsystem>();
+	ULandscapeSubsystem* LandscapeSubsystem = World ? World->GetSubsystem<ULandscapeSubsystem>() : nullptr;
 	const bool bIsGrassCreationPrioritized = LandscapeSubsystem ? LandscapeSubsystem->IsGrassCreationPrioritized() : false;
 	const bool bWaitAsyncTasks = bForceSync || bIsGrassCreationPrioritized;
 
