Fixing grass not spawning in ES3.1 preview mode - bringing in on behalf of Jonathan.Bard

diff --git a/Engine/Source/Runtime/Landscape/Private/LandscapeGrass.cpp b/Engine/Source/Runtime/Landscape/Private/LandscapeGrass.cpp
--- a/Engine/Source/Runtime/Landscape/Private/LandscapeGrass.cpp
+++ b/Engine/Source/Runtime/Landscape/Private/LandscapeGrass.cpp
@@ -962,7 +962,8 @@ bool ULandscapeComponent::MaterialHasGrass() const
 
 bool ULandscapeComponent::IsGrassMapOutdated() const
 {
-	if (GrassData->HasValidData())
+	// Don't mark a component as outdated if we won't be able to render it anyway : 
+	if (GrassData->HasValidData() && CanRenderGrassMap())
 	{
 		// check material / instances haven't changed
 		const auto& MaterialStateIds = GrassData->MaterialStateIds;
@@ -1734,7 +1735,9 @@ void ALandscapeProxy::TickGrass(const TArray<FVector>& Cameras, int32& InOutNumC
 #if WITH_EDITORONLY_DATA
 	if (ALandscape* Landscape = GetLandscapeActor())
 	{
-		if (!Landscape->IsUpToDate() || !Landscape->bGrassUpdateEnabled)
+		// Don't allow grass to tick, except in ES3.1 preview mode, since landscape update is not possible there : IsUpToDate might return false and we'll never see grass in the preview mode as a result :
+		bool bAllowGrassTick = Landscape->IsUpToDate() || (GetWorld()->FeatureLevel < ERHIFeatureLevel::SM5);
+		if (!bAllowGrassTick || !Landscape->bGrassUpdateEnabled)
 		{
 			return;
 		}
