Expose updating heightmap collision data at runtime by moving the neccessary components of CookCollisionData into a runtime accessbile call.

diff --git a/Engine/Source/Runtime/Landscape/Classes/LandscapeHeightfieldCollisionComponent.h b/Engine/Source/Runtime/Landscape/Classes/LandscapeHeightfieldCollisionComponent.h
--- a/Engine/Source/Runtime/Landscape/Classes/LandscapeHeightfieldCollisionComponent.h
+++ b/Engine/Source/Runtime/Landscape/Classes/LandscapeHeightfieldCollisionComponent.h
@@ -229,7 +229,7 @@ public:
 	/** 
 	 * Cooks raw height data into collision object binary stream
 	 */
-	virtual bool CookCollisionData(const FName& Format, bool bUseOnlyDefMaterial, bool bCheckDDC, TArray<uint8>& OutCookedData, TArray<UPhysicalMaterial*>& InOutMaterials) const;
+	virtual bool CookCollisionData(const FName& Format, bool bUseDefaultMaterialOnly, bool bCheckDDC, TArray<uint8>& OutCookedData, TArray<UPhysicalMaterial*>& InOutMaterials) const;
 
 	/** Modify a sub-region of the physics heightfield. Note that this does not update the physical material */
 	void UpdateHeightfieldRegion(int32 ComponentX1, int32 ComponentY1, int32 ComponentX2, int32 ComponentY2);
@@ -237,6 +237,23 @@ public:
 	/** Computes a hash of all the data that will impact final collision */
 	virtual uint32 ComputeCollisionHash() const;
 #endif
+
+	struct FWriteRuntimeDataParams
+	{
+		bool bUseDefaultMaterialOnly = false;
+		TArrayView<const uint16> Heights;
+		TArrayView<const uint16> SimpleHeights;
+		TArrayView<const uint8> DominantLayers;
+		TArrayView<const uint8> SimpleDominantLayers;
+		TArrayView<const uint8> RenderPhysicalMaterialIds;
+		TArrayView<const uint8> SimpleRenderPhysicalMaterialIds;
+	};
+
+	void GetCollisionSampleInfo(int32& OutCollisionSizeVerts, int32& OutSimpleCollisionSizeVerts, int32& OutNumSamples, int32& OutNumSimpleSamples) const;
+
+	// Writes to a cooked data buffer using raw heightfield data
+	LANDSCAPE_API bool WriteRuntimeData(const FWriteRuntimeDataParams& Params, TArray<uint8>& OutHeightfieldData, TArray<UPhysicalMaterial*>& InOutMaterials) const;
+
 	/** Gets the landscape info object for this landscape */
 	ULandscapeInfo* GetLandscapeInfo() const;
 
diff --git a/Engine/Source/Runtime/Landscape/Classes/LandscapeMeshCollisionComponent.h b/Engine/Source/Runtime/Landscape/Classes/LandscapeMeshCollisionComponent.h
--- a/Engine/Source/Runtime/Landscape/Classes/LandscapeMeshCollisionComponent.h
+++ b/Engine/Source/Runtime/Landscape/Classes/LandscapeMeshCollisionComponent.h
@@ -91,7 +91,7 @@ public:
 	virtual void ExportCustomProperties(FOutputDevice& Out, uint32 Indent) override;
 	virtual void ImportCustomProperties(const TCHAR* SourceText, FFeedbackContext* Warn) override;
 
-	virtual bool CookCollisionData(const FName& Format, bool bUseOnlyDefMaterial, bool bCheckDDC, TArray<uint8>& OutCookedData, TArray<UPhysicalMaterial*>& InOutMaterials) const override;
+	virtual bool CookCollisionData(const FName& Format, bool bUseDefaultMaterialOnly, bool bCheckDDC, TArray<uint8>& OutCookedData, TArray<UPhysicalMaterial*>& InOutMaterials) const override;
 	virtual uint32 ComputeCollisionHash() const override { return 0; }
 #endif
 	//~ End UObject Interface.
diff --git a/Engine/Source/Runtime/Landscape/Private/LandscapeCollision.cpp b/Engine/Source/Runtime/Landscape/Private/LandscapeCollision.cpp
--- a/Engine/Source/Runtime/Landscape/Private/LandscapeCollision.cpp
+++ b/Engine/Source/Runtime/Landscape/Private/LandscapeCollision.cpp
@@ -806,7 +806,120 @@ void ULandscapeHeightfieldCollisionComponent::SpeculativelyLoadAsyncDDCCollsionD
 	}
 }
 
-bool ULandscapeHeightfieldCollisionComponent::CookCollisionData(const FName& Format, bool bUseDefMaterial, bool bCheckDDC, TArray<uint8>& OutCookedData, TArray<UPhysicalMaterial*>& InOutMaterials) const
+#endif
+
+void ULandscapeHeightfieldCollisionComponent::GetCollisionSampleInfo(int32& OutCollisionSizeVerts, int32& OutSimpleCollisionSizeVerts, int32& OutNumSamples, int32& OutNumSimpleSamples) const
+{
+	OutCollisionSizeVerts = CollisionSizeQuads + 1;
+	OutSimpleCollisionSizeVerts = SimpleCollisionSizeQuads > 0 ? SimpleCollisionSizeQuads + 1 : 0;
+	OutNumSamples = FMath::Square(OutCollisionSizeVerts);
+	OutNumSimpleSamples = FMath::Square(OutSimpleCollisionSizeVerts);
+}
+
+
+// Writes to a cooked data buffer using raw heightfield data
+bool ULandscapeHeightfieldCollisionComponent::WriteRuntimeData(const FWriteRuntimeDataParams& Params, TArray<uint8>& OutHeightfieldData, TArray<UPhysicalMaterial*>& InOutMaterials) const
+{
+	ALandscapeProxy* Proxy = GetLandscapeProxy();
+	if (!Proxy || !Proxy->GetRootComponent())
+	{
+		return false;
+	}
+
+	UPhysicalMaterial* DefMaterial = Proxy->DefaultPhysMaterial ? Proxy->DefaultPhysMaterial : GEngine->DefaultPhysMaterial;
+
+	// GetComponentTransform() might not be initialized at this point, so use landscape transform
+	const FVector LandscapeScale = Proxy->GetRootComponent()->GetRelativeScale3D();
+	const bool bIsMirrored = (LandscapeScale.X * LandscapeScale.Y * LandscapeScale.Z) < 0.f;
+
+	const bool bGenerateSimpleCollision = SimpleCollisionSizeQuads > 0 && !Params.bUseDefaultMaterialOnly;
+
+	int32 CollisionSizeVerts;
+	int32 SimpleCollisionSizeVerts;
+	int32 NumSamples;
+	int32 NumSimpleSamples;
+	GetCollisionSampleInfo(CollisionSizeVerts, SimpleCollisionSizeVerts, NumSamples, NumSimpleSamples);
+
+	// List of materials which is actually used by heightfield
+	InOutMaterials.Empty();
+
+	// Generate material indices
+	TArray<uint8> MaterialIndices;
+	MaterialIndices.Reserve(NumSamples + NumSimpleSamples);
+
+	auto ResolveMaterials = [&MaterialIndices, &bIsMirrored, &Params, &DefMaterial, &InOutMaterials, this](int32 InCollisionVertExtent, TArrayView<const uint8> InDominantLayers, TArrayView<const uint8> InRenderMaterialIds)
+	{
+		for (int32 RowIndex = 0; RowIndex < InCollisionVertExtent; RowIndex++)
+		{
+			for (int32 ColIndex = 0; ColIndex < InCollisionVertExtent; ColIndex++)
+			{
+				const int32 SrcSampleIndex = (RowIndex * InCollisionVertExtent) + (bIsMirrored ? (InCollisionVertExtent - ColIndex - 1) : ColIndex);
+
+				// Materials are not relevant on the last row/column because they are per-triangle and the last row/column don't own any
+				if (RowIndex < InCollisionVertExtent - 1 &&
+					ColIndex < InCollisionVertExtent - 1)
+				{
+					int32 MaterialIndex = 0; // Default physical material.
+					if (!Params.bUseDefaultMaterialOnly)
+					{
+						uint8 DominantLayerIdx = InDominantLayers.IsEmpty() ? -1 : InDominantLayers[SrcSampleIndex];
+						ULandscapeLayerInfoObject* Layer = ComponentLayerInfos.IsValidIndex(DominantLayerIdx) ? ToRawPtr(ComponentLayerInfos[DominantLayerIdx]) : nullptr;
+
+#if WITH_EDITOR
+						if (Layer == ALandscapeProxy::VisibilityLayer)
+						{
+							// If it's a hole, use the final index
+							MaterialIndex = TNumericLimits<uint8>::Max();
+						}
+						else if (!InRenderMaterialIds.IsEmpty())
+						{
+							uint8 RenderIdx = InRenderMaterialIds[SrcSampleIndex];
+							UPhysicalMaterial* DominantMaterial = RenderIdx > 0 ? ToRawPtr(PhysicalMaterialRenderObjects[RenderIdx - 1]) : DefMaterial;
+							MaterialIndex = InOutMaterials.AddUnique(DominantMaterial);
+						}
+						else
+#endif
+						{
+							UPhysicalMaterial* DominantMaterial = Layer && Layer->PhysMaterial ? ToRawPtr(Layer->PhysMaterial) : DefMaterial;
+							MaterialIndex = InOutMaterials.AddUnique(DominantMaterial);
+						}
+					}
+					MaterialIndices.Add(MaterialIndex);
+				}
+			}
+		}
+	};
+
+	ResolveMaterials(CollisionSizeVerts, Params.DominantLayers, Params.RenderPhysicalMaterialIds);
+	ResolveMaterials(SimpleCollisionSizeVerts, Params.SimpleDominantLayers, Params.SimpleRenderPhysicalMaterialIds);
+
+	TUniquePtr<Chaos::FHeightField> Heightfield = nullptr;
+	TUniquePtr<Chaos::FHeightField> HeightfieldSimple = nullptr;
+
+	FMemoryWriter Writer(OutHeightfieldData);
+	Chaos::FChaosArchive Ar(Writer);
+
+	bool bSerializeGenerateSimpleCollision = bGenerateSimpleCollision;
+	Ar << bSerializeGenerateSimpleCollision;
+
+	const int32 NumCollisionCells = FMath::Square(CollisionSizeQuads);
+	const int32 NumSimpleCollisionCells = FMath::Square(SimpleCollisionSizeQuads);
+
+	TArrayView<uint8> ComplexMaterialIndicesView(MaterialIndices.GetData(), NumCollisionCells);
+	Heightfield = MakeUnique<Chaos::FHeightField>(Params.Heights, ComplexMaterialIndicesView, CollisionSizeVerts, CollisionSizeVerts, Chaos::FVec3(1));
+	Ar << Heightfield;
+	if (bGenerateSimpleCollision)
+	{
+		TArrayView<uint8> SimpleMaterialIndicesView(MaterialIndices.GetData() + NumCollisionCells, NumSimpleCollisionCells);
+		HeightfieldSimple = MakeUnique<Chaos::FHeightField>(Params.SimpleHeights, SimpleMaterialIndicesView, SimpleCollisionSizeVerts, SimpleCollisionSizeVerts, Chaos::FVec3(1));
+		Ar << HeightfieldSimple;
+	}
+
+	return true;
+}
+
+#if WITH_EDITOR
+bool ULandscapeHeightfieldCollisionComponent::CookCollisionData(const FName& Format, bool bUseDefaultMaterialOnly, bool bCheckDDC, TArray<uint8>& OutCookedData, TArray<UPhysicalMaterial*>& InOutMaterials) const
 {
 	TRACE_CPUPROFILER_EVENT_SCOPE(ULandscapeHeightfieldCollisionComponent::CookCollisionData);
 
@@ -823,18 +936,15 @@ bool ULandscapeHeightfieldCollisionComponent::CookCollisionData(const FName& For
 
 	COOK_STAT(auto Timer = LandscapeCollisionCookStats::HeightfieldUsageStats.TimeSyncWork());
 	 
-	bool Succeeded = false;
-	TArray<uint8> OutData;
-
 	// we have 2 versions of collision objects
-	const int32 CookedDataIndex = bUseDefMaterial ? 0 : 1;
+	const int32 CookedDataIndex = bUseDefaultMaterialOnly ? 0 : 1;
 
 	if (!GLandscapeCollisionSkipDDC && bCheckDDC && HeightfieldGuid.IsValid())
 	{
 		// Ensure that content was saved with physical materials before using DDC data
 		if (GetLinkerUEVersion() >= VER_UE4_LANDSCAPE_SERIALIZE_PHYSICS_MATERIALS)
 		{
-			FString DDCKey = GetHFDDCKeyString(Format, bUseDefMaterial, HeightfieldGuid, InOutMaterials);
+			FString DDCKey = GetHFDDCKeyString(Format, bUseDefaultMaterialOnly, HeightfieldGuid, InOutMaterials);
 
 			// Check if the speculatively-loaded data loaded and is what we wanted
 			if (SpeculativeDDCRequest.IsValid() && DDCKey == SpeculativeDDCRequest->GetKey())
@@ -868,26 +978,11 @@ bool ULandscapeHeightfieldCollisionComponent::CookCollisionData(const FName& For
 		}
 	}
 
-	ALandscapeProxy* Proxy = GetLandscapeProxy();
-	if (!Proxy || !Proxy->GetRootComponent())
-	{
-		// We didn't actually build anything, so just track the cycles.
-		COOK_STAT(Timer.TrackCyclesOnly());
-		return false;
-	}
-
-	UPhysicalMaterial* DefMaterial = Proxy->DefaultPhysMaterial ? Proxy->DefaultPhysMaterial : GEngine->DefaultPhysMaterial;
-
-	// GetComponentTransform() might not be initialized at this point, so use landscape transform
-	const FVector LandscapeScale = Proxy->GetRootComponent()->GetRelativeScale3D();
-	const bool bIsMirrored = (LandscapeScale.X*LandscapeScale.Y*LandscapeScale.Z) < 0.f;
-
-	const bool bGenerateSimpleCollision = SimpleCollisionSizeQuads > 0 && !bUseDefMaterial;
-
-	const int32 CollisionSizeVerts = CollisionSizeQuads + 1;
-	const int32 SimpleCollisionSizeVerts = SimpleCollisionSizeQuads > 0 ? SimpleCollisionSizeQuads + 1 : 0;
-	const int32 NumSamples = FMath::Square(CollisionSizeVerts);
-	const int32 NumSimpleSamples = FMath::Square(SimpleCollisionSizeVerts);
+	int32 CollisionSizeVerts;
+	int32 SimpleCollisionSizeVerts;
+	int32 NumSamples;
+	int32 NumSimpleSamples;
+	GetCollisionSampleInfo(CollisionSizeVerts, SimpleCollisionSizeVerts, NumSamples, NumSimpleSamples);
 
 	const uint16* Heights = (const uint16*)CollisionHeightData.LockReadOnly();
 	check(CollisionHeightData.GetElementCount() == NumSamples + NumSimpleSamples);
@@ -913,82 +1008,23 @@ bool ULandscapeHeightfieldCollisionComponent::CookCollisionData(const FName& For
 		SimpleRenderPhysicalMaterialIds = RenderPhysicalMaterialIds + NumSamples;
 	}
 
-	// List of materials which is actually used by heightfield
-	InOutMaterials.Empty();
-
-	// Generate material indices
-	TArray<uint8> MaterialIndices;
-	MaterialIndices.Reserve(NumSamples + NumSimpleSamples);
-
-	auto ResolveMaterials = [&MaterialIndices, &bIsMirrored, &bUseDefMaterial, &DefMaterial, &InOutMaterials, this](int32 InCollisionVertExtent, const uint8* InDominantLayers, const uint8* InRenderMaterialIds)
-	{
-		for(int32 RowIndex = 0; RowIndex < InCollisionVertExtent; RowIndex++)
+	auto MakeSafeArrayView = 
+		[](auto* Data, int32 DataSamples)
 		{
-			for(int32 ColIndex = 0; ColIndex < InCollisionVertExtent; ColIndex++)
-			{
-				const int32 SrcSampleIndex = (RowIndex * InCollisionVertExtent) + (bIsMirrored ? (InCollisionVertExtent - ColIndex - 1) : ColIndex);
+			return Data ? MakeArrayView(Data, DataSamples) : MakeArrayView(Data, 0);
+		};
 
-				// Materials are not relevant on the last row/column because they are per-triangle and the last row/column don't own any
-				if(RowIndex < InCollisionVertExtent - 1 &&
-				   ColIndex < InCollisionVertExtent - 1)
-				{
-					int32 MaterialIndex = 0; // Default physical material.
-					if(!bUseDefMaterial)
-					{
-						uint8 DominantLayerIdx = InDominantLayers ? InDominantLayers[SrcSampleIndex] : -1;
-						ULandscapeLayerInfoObject* Layer = ComponentLayerInfos.IsValidIndex(DominantLayerIdx) ? ToRawPtr(ComponentLayerInfos[DominantLayerIdx]) : nullptr;
+	FWriteRuntimeDataParams WriteParams;
+	WriteParams.bUseDefaultMaterialOnly = bUseDefaultMaterialOnly;
+	WriteParams.Heights = MakeSafeArrayView(Heights, NumSamples);
+	WriteParams.SimpleHeights = MakeSafeArrayView(SimpleHeights, NumSimpleSamples);
+	WriteParams.DominantLayers = MakeSafeArrayView(DominantLayers, NumSamples);
+	WriteParams.SimpleDominantLayers = MakeSafeArrayView(SimpleDominantLayers, NumSimpleSamples);
+	WriteParams.RenderPhysicalMaterialIds = MakeSafeArrayView(RenderPhysicalMaterialIds, NumSamples);
+	WriteParams.SimpleRenderPhysicalMaterialIds = MakeSafeArrayView(SimpleRenderPhysicalMaterialIds, NumSimpleSamples);
 
-						if(Layer == ALandscapeProxy::VisibilityLayer)
-						{
-							// If it's a hole, use the final index
-							MaterialIndex = TNumericLimits<uint8>::Max();
-						}
-						else if(InRenderMaterialIds)
-						{
-							uint8 RenderIdx = InRenderMaterialIds[SrcSampleIndex];
-							UPhysicalMaterial* DominantMaterial = RenderIdx > 0 ? ToRawPtr(PhysicalMaterialRenderObjects[RenderIdx - 1]) : DefMaterial;
-							MaterialIndex = InOutMaterials.AddUnique(DominantMaterial);
-						}
-						else
-						{
-							UPhysicalMaterial* DominantMaterial = Layer && Layer->PhysMaterial ? ToRawPtr(Layer->PhysMaterial) : DefMaterial;
-							MaterialIndex = InOutMaterials.AddUnique(DominantMaterial);
-						}
-					}
-					MaterialIndices.Add(MaterialIndex);
-				}
-			}
-		}
-	};
-
-	ResolveMaterials(CollisionSizeVerts, DominantLayers, RenderPhysicalMaterialIds);
-	ResolveMaterials(SimpleCollisionSizeVerts, SimpleDominantLayers, SimpleRenderPhysicalMaterialIds);
-
-	TUniquePtr<Chaos::FHeightField> Heightfield = nullptr;
-	TUniquePtr<Chaos::FHeightField> HeightfieldSimple = nullptr;
-
-	FMemoryWriter Writer(OutData);
-	Chaos::FChaosArchive Ar(Writer);
-
-	bool bSerializeGenerateSimpleCollision = bGenerateSimpleCollision;
-	Ar << bSerializeGenerateSimpleCollision;
-
-	const int32 NumCollisionCells = FMath::Square(CollisionSizeQuads);
-	const int32 NumSimpleCollisionCells = FMath::Square(SimpleCollisionSizeQuads);
-
-	TArrayView<const uint16> ComplexHeightView(Heights, NumSamples);
-	TArrayView<uint8> ComplexMaterialIndicesView(MaterialIndices.GetData(), NumCollisionCells);
-	Heightfield = MakeUnique<Chaos::FHeightField>(ComplexHeightView, ComplexMaterialIndicesView, CollisionSizeVerts, CollisionSizeVerts, Chaos::FVec3(1));
-	Ar << Heightfield;
-	if(bGenerateSimpleCollision)
-	{
-		TArrayView<const uint16> SimpleHeightView(SimpleHeights, NumSimpleSamples);
-		TArrayView<uint8> SimpleMaterialIndicesView(MaterialIndices.GetData() + NumCollisionCells, NumSimpleCollisionCells);
-		HeightfieldSimple = MakeUnique<Chaos::FHeightField>(SimpleHeightView, SimpleMaterialIndicesView, SimpleCollisionSizeVerts, SimpleCollisionSizeVerts, Chaos::FVec3(1));
-		Ar << HeightfieldSimple;
-	}
-
-	Succeeded = true;
+	TArray<uint8> OutData;
+	bool Succeeded = WriteRuntimeData(WriteParams, OutData, InOutMaterials);
 
 	if (CollisionHeightData.IsLocked())
 	{
@@ -1003,30 +1039,28 @@ bool ULandscapeHeightfieldCollisionComponent::CookCollisionData(const FName& For
 		PhysicalMaterialRenderData.Unlock();
 	}
 
-	if (Succeeded)
+	if (!Succeeded)
 	{
-		COOK_STAT(Timer.AddMiss(OutData.Num()));
-		OutCookedData.SetNumUninitialized(OutData.Num());
-		FMemory::Memcpy(OutCookedData.GetData(), OutData.GetData(), OutData.Num());
-
-		if (!GLandscapeCollisionSkipDDC && bShouldSaveCookedDataToDDC[CookedDataIndex] && HeightfieldGuid.IsValid())
-		{
-			GetDerivedDataCacheRef().Put(*GetHFDDCKeyString(Format, bUseDefMaterial, HeightfieldGuid, InOutMaterials), OutCookedData, GetPathName());
-			bShouldSaveCookedDataToDDC[CookedDataIndex] = false;
-		}
+		// We didn't actually build anything, so just track the cycles.
+		COOK_STAT(Timer.TrackCyclesOnly());
+		return false;
 	}
-	else
+
+	COOK_STAT(Timer.AddMiss(OutData.Num()));
+	OutCookedData.SetNumUninitialized(OutData.Num());
+	FMemory::Memcpy(OutCookedData.GetData(), OutData.GetData(), OutData.Num());
+
+	if (!GLandscapeCollisionSkipDDC && bShouldSaveCookedDataToDDC[CookedDataIndex] && HeightfieldGuid.IsValid())
 	{
-		// if we failed to build the resource, just time the cycles we spent.
-		COOK_STAT(Timer.TrackCyclesOnly());
-		OutCookedData.Empty();
-		InOutMaterials.Empty();
+		GetDerivedDataCacheRef().Put(*GetHFDDCKeyString(Format, bUseDefaultMaterialOnly, HeightfieldGuid, InOutMaterials), OutCookedData, GetPathName());
+		bShouldSaveCookedDataToDDC[CookedDataIndex] = false;
 	}
+	
 
 	return Succeeded;
 }
 
-bool ULandscapeMeshCollisionComponent::CookCollisionData(const FName& Format, bool bUseDefMaterial, bool bCheckDDC, TArray<uint8>& OutCookedData, TArray<UPhysicalMaterial*>& InOutMaterials) const
+bool ULandscapeMeshCollisionComponent::CookCollisionData(const FName& Format, bool bUseDefaultMaterialOnly, bool bCheckDDC, TArray<uint8>& OutCookedData, TArray<UPhysicalMaterial*>& InOutMaterials) const
 {
 	// Use existing cooked data unless !bCheckDDC in which case the data must be rebuilt.
 	if (bCheckDDC && OutCookedData.Num() > 0)
@@ -1036,14 +1070,14 @@ bool ULandscapeMeshCollisionComponent::CookCollisionData(const FName& Format, bo
 
 	COOK_STAT(auto Timer = LandscapeCollisionCookStats::MeshUsageStats.TimeSyncWork());
 	// we have 2 versions of collision objects
-	const int32 CookedDataIndex = bUseDefMaterial ? 0 : 1;
+	const int32 CookedDataIndex = bUseDefaultMaterialOnly ? 0 : 1;
 
 	if (!GLandscapeCollisionSkipDDC && bCheckDDC)
 	{
 		// Ensure that content was saved with physical materials before using DDC data
 		if (GetLinkerUEVersion() >= VER_UE4_LANDSCAPE_SERIALIZE_PHYSICS_MATERIALS && MeshGuid.IsValid())
 		{
-			FString DDCKey = GetHFDDCKeyString(Format, bUseDefMaterial, MeshGuid, InOutMaterials);
+			FString DDCKey = GetHFDDCKeyString(Format, bUseDefaultMaterialOnly, MeshGuid, InOutMaterials);
 
 			// Check if the speculatively-loaded data loaded and is what we wanted
 			if (SpeculativeDDCRequest.IsValid() && DDCKey == SpeculativeDDCRequest->GetKey())
@@ -1128,7 +1162,7 @@ bool ULandscapeMeshCollisionComponent::CookCollisionData(const FName& Format, bo
 			bool bHole = false;
 
 			int32 MaterialIndex = 0; // Default physical material.
-			if (!bUseDefMaterial && DominantLayers)
+			if (!bUseDefaultMaterialOnly && DominantLayers)
 			{
 				uint8 DominantLayerIdx = DominantLayers[DataIdx];
 				if (ComponentLayerInfos.IsValidIndex(DominantLayerIdx))
@@ -1232,7 +1266,7 @@ bool ULandscapeMeshCollisionComponent::CookCollisionData(const FName& Format, bo
 
 		if (!GLandscapeCollisionSkipDDC && bShouldSaveCookedDataToDDC[CookedDataIndex] && MeshGuid.IsValid())
 		{
-			GetDerivedDataCacheRef().Put(*GetHFDDCKeyString(Format, bUseDefMaterial, MeshGuid, InOutMaterials), OutCookedData, GetPathName());
+			GetDerivedDataCacheRef().Put(*GetHFDDCKeyString(Format, bUseDefaultMaterialOnly, MeshGuid, InOutMaterials), OutCookedData, GetPathName());
 			bShouldSaveCookedDataToDDC[CookedDataIndex] = false;
 		}
 	}
