Fix disappearing grass when switching 120fps on/off

#rnx

diff --git a/Engine/Source/Runtime/Landscape/Private/LandscapeGrass.cpp b/Engine/Source/Runtime/Landscape/Private/LandscapeGrass.cpp
--- a/Engine/Source/Runtime/Landscape/Private/LandscapeGrass.cpp
+++ b/Engine/Source/Runtime/Landscape/Private/LandscapeGrass.cpp
@@ -247,13 +247,19 @@ static void GrassCVarSinkFunction()
 	static int32 CachedDetailMode = DetailModeCVar ? DetailModeCVar->GetInt() : 0;
 	int32 DetailMode = DetailModeCVar ? DetailModeCVar->GetInt() : 0;
 
+	static const IConsoleVariable* NaniteEnabledCVar = IConsoleManager::Get().FindConsoleVariable(TEXT("r.Nanite"));
+	static int32 CachedNaniteEnabled = NaniteEnabledCVar ? NaniteEnabledCVar->GetInt() : 0;
+	int32 NaniteEnabled = NaniteEnabledCVar ? NaniteEnabledCVar->GetInt() : 0;
+
 	if (DetailMode != CachedDetailMode || 
 		GrassDensityScale != CachedGrassDensityScale || 
-		GrassCullDistanceScale != CachedGrassCullDistanceScale)
+		GrassCullDistanceScale != CachedGrassCullDistanceScale || 
+		CachedNaniteEnabled != NaniteEnabled)
 	{
 		CachedGrassDensityScale = GrassDensityScale;
 		CachedGrassCullDistanceScale = GrassCullDistanceScale;
 		CachedDetailMode = DetailMode;
+		CachedNaniteEnabled = NaniteEnabled;
 
 		for (auto* Landscape : TObjectRange<ALandscapeProxy>(RF_ClassDefaultObject | RF_ArchetypeObject, true, EInternalObjectFlags::Garbage))
 		{
