Fixed warning about landscape thumbnail material instance shader caching not being issued when in SM6.0

diff --git a/Engine/Source/Runtime/Landscape/Private/LandscapeRender.cpp b/Engine/Source/Runtime/Landscape/Private/LandscapeRender.cpp
--- a/Engine/Source/Runtime/Landscape/Private/LandscapeRender.cpp
+++ b/Engine/Source/Runtime/Landscape/Private/LandscapeRender.cpp
@@ -3107,10 +3107,7 @@ public:
 						}
 						else
 						{
-							if (Platform == EShaderPlatform::SP_PCD3D_SM5)
-							{
-								UE_LOG(LogLandscape, Warning, TEXT("Shader %s unknown by landscape thumbnail material, please add to either AllowedShaderTypes or ExcludedShaderTypes"), ShaderType->GetName());
-							}
+							UE_LOG(LogLandscape, Warning, TEXT("Shader %s unknown by landscape thumbnail material, please add to either AllowedShaderTypes or ExcludedShaderTypes"), ShaderType->GetName());
 							return FMaterialResource::ShouldCache(Platform, ShaderType, VertexFactoryType);
 						}
 					}
