Removed dead code

#rnx

diff --git a/Engine/Source/Editor/LandscapeEditor/Private/LandscapeEdModeTools.h b/Engine/Source/Editor/LandscapeEditor/Private/LandscapeEdModeTools.h
--- a/Engine/Source/Editor/LandscapeEditor/Private/LandscapeEdModeTools.h
+++ b/Engine/Source/Editor/LandscapeEditor/Private/LandscapeEdModeTools.h
@@ -1576,14 +1576,3 @@ protected:
 		return ViewportClient != nullptr && IsShiftDown(ViewportClient->Viewport);
 	}
 };
-
-struct FToolFlattenCustomData
-{
-	FToolFlattenCustomData()
-		: ActiveEyeDropperMode(false)
-		, EyeDropperModeHeight(0.0f)
-	{}
-
-	bool ActiveEyeDropperMode;
-	float EyeDropperModeHeight;
-};
diff --git a/Engine/Source/Runtime/Landscape/Private/Landscape.cpp b/Engine/Source/Runtime/Landscape/Private/Landscape.cpp
--- a/Engine/Source/Runtime/Landscape/Private/Landscape.cpp
+++ b/Engine/Source/Runtime/Landscape/Private/Landscape.cpp
@@ -183,13 +183,7 @@ static FAutoConsoleVariableRef CVarLiveRebuildNaniteOnModification(
 	LiveRebuildNaniteOnModification,
 	TEXT("Trigger a rebuild of Nanite representation immediately when a modification is performed"));
 
-static int32 GLandscapeDirtyOnlyInModeEnabled = 0;
-static FAutoConsoleVariableRef CVarLandscapeDirtyOnlyInModeEnabled(
-	TEXT("landscape.DirtyOnlyInMode"),
-	GLandscapeDirtyOnlyInModeEnabled,
-	TEXT("Set to 1 to avoid Landscape being dirtied when not in Landscape mode."));
-
-#endif
+#endif // WITH_EDITOR
 
 int32 GRenderNaniteLandscape = 1;
 FAutoConsoleVariableRef CVarRenderNaniteLandscape(
