Fixed assert when resizing a ALandscape with a ULandscapeSplinesComponent : upon duplicating the splines component, it would get auto-registered to the ALandscape actor, thus making the call to SetSplinesComponent redundant.

diff --git a/Engine/Source/Runtime/Landscape/Classes/LandscapeProxy.h b/Engine/Source/Runtime/Landscape/Classes/LandscapeProxy.h
--- a/Engine/Source/Runtime/Landscape/Classes/LandscapeProxy.h
+++ b/Engine/Source/Runtime/Landscape/Classes/LandscapeProxy.h
@@ -995,7 +995,7 @@ public:
 	/* Serialize all hashes/guids that record the current state of this proxy */
 	void SerializeStateHashes(FArchive& Ar);
 
-	LANDSCAPE_API void SetSplinesComponent(ULandscapeSplinesComponent* InSplineComponent) { check(!SplineComponent); SplineComponent = InSplineComponent; }
+	LANDSCAPE_API void SetSplinesComponent(ULandscapeSplinesComponent* InSplineComponent) { check(!SplineComponent || (SplineComponent == InSplineComponent)); SplineComponent = InSplineComponent; }
 
 	LANDSCAPE_API virtual bool SupportsForeignSplineMesh() const override { return true; }
 
