Fixed bad merge

diff --git a/Engine/Source/Runtime/Landscape/Private/LandscapeSubsystem.cpp b/Engine/Source/Runtime/Landscape/Private/LandscapeSubsystem.cpp
--- a/Engine/Source/Runtime/Landscape/Private/LandscapeSubsystem.cpp
+++ b/Engine/Source/Runtime/Landscape/Private/LandscapeSubsystem.cpp
@@ -238,8 +238,6 @@ void ULandscapeSubsystem::BuildNanite(TArrayView<ALandscapeProxy*> InProxiesToBu
 		return;
 	}
 
-	FScopedSlowTask SlowTask(static_cast<float>(Proxies.Num()), (LOCTEXT("Landscape_BuildNanite", "Building Nanite Landscape Meshes")));
-
 	TArray<ALandscapeProxy*> FinalProxiesToBuild;
 	if (InProxiesToBuild.IsEmpty())
 	{
