World Partition: add support for projects to force landscape proxies and landscape dependant objects non-spatially loaded by default.

diff --git a/Engine/Source/Runtime/Landscape/Classes/LandscapeInfo.h b/Engine/Source/Runtime/Landscape/Classes/LandscapeInfo.h
--- a/Engine/Source/Runtime/Landscape/Classes/LandscapeInfo.h
+++ b/Engine/Source/Runtime/Landscape/Classes/LandscapeInfo.h
@@ -406,10 +406,12 @@ public:
 
 	LANDSCAPE_API void RequestSplineLayerUpdate();
 	LANDSCAPE_API void ForceLayersFullUpdate();
-#endif
 
-#if WITH_EDITOR
+	static void SetForceNonSpatiallyLoadedByDefault(bool bInForceNonSpatiallyLoadedByDefault) { bForceNonSpatiallyLoadedByDefault = bInForceNonSpatiallyLoadedByDefault; }
+
 private:
+	inline static bool bForceNonSpatiallyLoadedByDefault = false;
+
 	bool ApplySplinesInternal(bool bOnlySelected, TScriptInterface<ILandscapeSplineInterface> SplineOwner, TSet<TObjectPtr<ULandscapeComponent>>* OutModifiedComponents, bool bMarkPackageDirty, int32 LandscapeMinX, int32 LandscapeMinY, int32 LandscapeMaxX, int32 LandscapeMaxY, TFunctionRef<TSharedPtr<FModulateAlpha>(ULandscapeLayerInfoObject*)> GetOrCreateModulate);
 	void MoveSegment(ULandscapeSplineSegment* InSegment, TScriptInterface<ILandscapeSplineInterface> From, TScriptInterface<ILandscapeSplineInterface> To);
 	void MoveControlPoint(ULandscapeSplineControlPoint* InControlPoint, TScriptInterface<ILandscapeSplineInterface> From, TScriptInterface<ILandscapeSplineInterface> To);
diff --git a/Engine/Source/Runtime/Landscape/Private/LandscapeEdit.cpp b/Engine/Source/Runtime/Landscape/Private/LandscapeEdit.cpp
--- a/Engine/Source/Runtime/Landscape/Private/LandscapeEdit.cpp
+++ b/Engine/Source/Runtime/Landscape/Private/LandscapeEdit.cpp
@@ -4806,9 +4806,12 @@ uint32 ULandscapeInfo::GetGridSize(uint32 InGridSizeInComponents) const
 
 bool ULandscapeInfo::AreNewLandscapeActorsSpatiallyLoaded() const
 {
-	if (ALandscape* Landscape = LandscapeActor.Get())
+	if (!bForceNonSpatiallyLoadedByDefault)
 	{
-		return Landscape->bAreNewLandscapeActorsSpatiallyLoaded;
+		if (ALandscape* Landscape = LandscapeActor.Get())
+		{
+			return Landscape->bAreNewLandscapeActorsSpatiallyLoaded;
+		}
 	}
 	
 	return false;
