Added single method in ULandscapeSubsystem to force regenerate the grass at runtime. This allows to (optionally) flush all grass on landscape components and update them all (synchronously, optionally), using an optional list of camera locations in case one wants to regenerate grass around another location than the streaming manager's view point location.

#skipci

diff --git a/Engine/Plugins/MovieScene/MovieRenderPipeline/Source/MovieRenderPipelineCore/Private/MoviePipelineRendering.cpp b/Engine/Plugins/MovieScene/MovieRenderPipeline/Source/MovieRenderPipelineCore/Private/MoviePipelineRendering.cpp
--- a/Engine/Plugins/MovieScene/MovieRenderPipeline/Source/MovieRenderPipelineCore/Private/MoviePipelineRendering.cpp
+++ b/Engine/Plugins/MovieScene/MovieRenderPipeline/Source/MovieRenderPipelineCore/Private/MoviePipelineRendering.cpp
@@ -29,7 +29,7 @@
 
 // For flushing async systems
 #include "RendererInterface.h"
-#include "LandscapeProxy.h"
+#include "LandscapeSubsystem.h"
 #include "EngineModule.h"
 #include "DistanceFieldAtlas.h"
 #include "MeshCardRepresentation.h"
@@ -547,14 +547,9 @@ void UMoviePipeline::FlushAsyncEngineSystems()
 		UMoviePipelineGameOverrideSetting* GameOverrides = FindOrAddSettingForShot<UMoviePipelineGameOverrideSetting>(ActiveShotList[CurrentShotIndex]);
 		if (GameOverrides && GameOverrides->bFlushGrassStreaming)
 		{
-			for (TActorIterator<ALandscapeProxy> It(GetWorld()); It; ++It)
+			if (ULandscapeSubsystem* LandscapeSubsystem = GetWorld()->GetSubsystem<ULandscapeSubsystem>())
 			{
-				ALandscapeProxy* LandscapeProxy = (*It);
-				if (LandscapeProxy)
-				{
-					TArray<FVector> CameraList;
-					LandscapeProxy->UpdateGrass(CameraList, true);
-				}
+				LandscapeSubsystem->RegenerateGrass(/*bInFlushGrass = */false, /*bInForceSync = */true, /*InOptionalCameraLocations = */MakeArrayView(TArray<FVector>()));
 			}
 		}
 	}
diff --git a/Engine/Source/Runtime/Landscape/Private/LandscapeGrass.cpp b/Engine/Source/Runtime/Landscape/Private/LandscapeGrass.cpp
--- a/Engine/Source/Runtime/Landscape/Private/LandscapeGrass.cpp
+++ b/Engine/Source/Runtime/Landscape/Private/LandscapeGrass.cpp
@@ -261,9 +261,12 @@ static void GrassCVarSinkFunction()
 		CachedDetailMode = DetailMode;
 		CachedNaniteEnabled = NaniteEnabled;
 
-		for (auto* Landscape : TObjectRange<ALandscapeProxy>(RF_ClassDefaultObject | RF_ArchetypeObject, true, EInternalObjectFlags::Garbage))
+		for (UWorld* World : TObjectRange<UWorld>(RF_ClassDefaultObject | RF_ArchetypeObject, true, EInternalObjectFlags::Garbage))
 		{
-			Landscape->FlushGrassComponents(nullptr, false);
+			if (ULandscapeSubsystem* LandscapeSubsystem = World->GetSubsystem<ULandscapeSubsystem>())
+			{
+				LandscapeSubsystem->RegenerateGrass(/*bInFlushGrass = */true, /*bInForceSync = */true);
+			}
 		}
 	}
 }
diff --git a/Engine/Source/Runtime/Landscape/Private/LandscapeSubsystem.cpp b/Engine/Source/Runtime/Landscape/Private/LandscapeSubsystem.cpp
--- a/Engine/Source/Runtime/Landscape/Private/LandscapeSubsystem.cpp
+++ b/Engine/Source/Runtime/Landscape/Private/LandscapeSubsystem.cpp
@@ -42,7 +42,7 @@ DECLARE_CYCLE_STAT(TEXT("LandscapeSubsystem Tick"), STAT_LandscapeSubsystemTick,
 #define LOCTEXT_NAMESPACE "LandscapeSubsystem"
 
 ULandscapeSubsystem::ULandscapeSubsystem()
-: bIsGrassCreationPrioritized(false)
+	: bIsGrassCreationPrioritized(false)
 {
 }
 
@@ -90,6 +90,65 @@ TStatId ULandscapeSubsystem::GetStatId() const
 	RETURN_QUICK_DECLARE_CYCLE_STAT(ULandscapeSubsystem, STATGROUP_Tickables);
 }
 
+void ULandscapeSubsystem::RegenerateGrass(bool bInFlushGrass, bool bInForceSync, TOptional<TArrayView<FVector>> InOptionalCameraLocations)
+{
+	TRACE_CPUPROFILER_EVENT_SCOPE(ULandscapeSubsystem::RegenerateGrass);
+
+	if (Proxies.IsEmpty())
+	{
+		return;
+	}
+
+	UWorld* World = GetWorld();
+
+	if (bInFlushGrass)
+	{
+		TRACE_CPUPROFILER_EVENT_SCOPE(FlushGrass);
+		for (TWeakObjectPtr<ALandscapeProxy> ProxyPtr : Proxies)
+		{
+			if (ALandscapeProxy* Proxy = ProxyPtr.Get())
+			{
+				Proxy->FlushGrassComponents(/*OnlyForComponents = */nullptr, /*bFlushGrassMaps = */false);
+			}
+		}
+	}
+
+	{
+		TRACE_CPUPROFILER_EVENT_SCOPE(UpdateGrass);
+
+		TArray<FVector> CameraLocations;
+		if (InOptionalCameraLocations.IsSet())
+		{
+			CameraLocations = *InOptionalCameraLocations;
+		}
+		else
+		{
+			if (GUseStreamingManagerForCameras == 0)
+			{
+				CameraLocations = World->ViewLocationsRenderedLastFrame;
+			}
+			else if (int32 Num = IStreamingManager::Get().GetNumViews())
+			{
+				CameraLocations.Reserve(Num);
+				for (int32 Index = 0; Index < Num; Index++)
+				{
+					const FStreamingViewInfo& ViewInfo = IStreamingManager::Get().GetViewInformation(Index);
+					CameraLocations.Add(ViewInfo.ViewOrigin);
+				}
+			}
+		}
+
+		// Update the grass near the specified location(s) : 
+		for (TWeakObjectPtr<ALandscapeProxy> ProxyPtr : Proxies)
+		{
+			if (ALandscapeProxy* Proxy = ProxyPtr.Get())
+			{
+				Proxy->UpdateGrass(CameraLocations, bInForceSync);
+			}
+		}
+	}
+}
+
 ETickableTickType ULandscapeSubsystem::GetTickableTickType() const
 {
 	return HasAnyFlags(RF_ClassDefaultObject) || !GetWorld() || GetWorld()->IsNetMode(NM_DedicatedServer) ? ETickableTickType::Never : ETickableTickType::Always;
diff --git a/Engine/Source/Runtime/Landscape/Public/LandscapeSubsystem.h b/Engine/Source/Runtime/Landscape/Public/LandscapeSubsystem.h
--- a/Engine/Source/Runtime/Landscape/Public/LandscapeSubsystem.h
+++ b/Engine/Source/Runtime/Landscape/Public/LandscapeSubsystem.h
@@ -53,6 +53,14 @@ public:
 	void PrioritizeGrassCreation(bool bPrioritizeGrassCreation) { bIsGrassCreationPrioritized = bPrioritizeGrassCreation; }
 	bool IsGrassCreationPrioritized() const { return bIsGrassCreationPrioritized; }
 
+	/**
+	 * Can be called at runtime : (optionally) flushes grass on all landscape components and updates them
+	 * @param bInFlushGrass : flushes all grass from landscape components prior to updating them
+	 * @param bInForceSync : synchronously updates grass on all landscape components 
+	 * @param InOptionalCameraLocations : (optional) camera locations that should be used when updating the grass. If not specified, the usual (streaming manager-based) view locations will be used
+	 */
+	LANDSCAPE_API void RegenerateGrass(bool bInFlushGrass, bool bInForceSync, TOptional<TArrayView<FVector>> InOptionalCameraLocations = TOptional<TArrayView<FVector>>());
+
 #if WITH_EDITOR
 	LANDSCAPE_API void BuildAll();
 	LANDSCAPE_API void BuildGrassMaps();
