Removed usage of CYCLE stats in landscape tool code, we now use CPUPROFILER

#rnx

diff --git a/Engine/Source/Runtime/Landscape/Private/Landscape.cpp b/Engine/Source/Runtime/Landscape/Private/Landscape.cpp
--- a/Engine/Source/Runtime/Landscape/Private/Landscape.cpp
+++ b/Engine/Source/Runtime/Landscape/Private/Landscape.cpp
@@ -103,37 +103,13 @@ Landscape.cpp: Terrain rendering
 /** Landscape stats */
 
 DEFINE_STAT(STAT_LandscapeDynamicDrawTime);
-DEFINE_STAT(STAT_LandscapeStaticDrawLODTime);
 DEFINE_STAT(STAT_LandscapeVFDrawTimeVS);
-DEFINE_STAT(STAT_LandscapeInitViewCustomData);
-DEFINE_STAT(STAT_LandscapePostInitViewCustomData);
-DEFINE_STAT(STAT_LandscapeComputeCustomMeshBatchLOD);
-DEFINE_STAT(STAT_LandscapeComputeCustomShadowMeshBatchLOD);
 DEFINE_STAT(STAT_LandscapeVFDrawTimePS);
 DEFINE_STAT(STAT_LandscapeComponentRenderPasses);
-DEFINE_STAT(STAT_LandscapeTessellatedShadowCascade);
-DEFINE_STAT(STAT_LandscapeTessellatedComponents);
-DEFINE_STAT(STAT_LandscapeComponentUsingSubSectionDrawCalls);
 DEFINE_STAT(STAT_LandscapeDrawCalls);
 DEFINE_STAT(STAT_LandscapeTriangles);
-
-DEFINE_STAT(STAT_LandscapeLayersRegenerate_RenderThread);
 DEFINE_STAT(STAT_LandscapeLayersRegenerateDrawCalls);
 
-DEFINE_STAT(STAT_LandscapeLayersRegenerateHeightmaps);
-DEFINE_STAT(STAT_LandscapeLayersResolveHeightmaps);
-DEFINE_STAT(STAT_LandscapeLayersResolveTexture);
-
-DEFINE_STAT(STAT_LandscapeLayersUpdateMaterialInstance);
-DEFINE_STAT(STAT_LandscapeLayersReallocateWeightmaps);
-
-DEFINE_STAT(STAT_LandscapeLayersResolveWeightmaps);
-DEFINE_STAT(STAT_LandscapeLayersRegenerateWeightmaps);
-
-DEFINE_STAT(STAT_LandscapeVertexMem);
-DEFINE_STAT(STAT_LandscapeHoleMem);
-DEFINE_STAT(STAT_LandscapeComponentMem);
-
 #if ENABLE_COOK_STATS
 namespace LandscapeCookStats
 {
diff --git a/Engine/Source/Runtime/Landscape/Private/LandscapeEditLayers.cpp b/Engine/Source/Runtime/Landscape/Private/LandscapeEditLayers.cpp
--- a/Engine/Source/Runtime/Landscape/Private/LandscapeEditLayers.cpp
+++ b/Engine/Source/Runtime/Landscape/Private/LandscapeEditLayers.cpp
@@ -4887,7 +4887,6 @@ int32 ALandscape::PerformLayersHeightmapsGlobalMerge(const FUpdateLayersContentC
 int32 ALandscape::RegenerateLayersHeightmaps(const FUpdateLayersContentContext& InUpdateLayersContentContext)
 {
 	TRACE_CPUPROFILER_EVENT_SCOPE(LandscapeLayers_RegenerateLayersHeightmaps);
-	SCOPE_CYCLE_COUNTER(STAT_LandscapeLayersRegenerateHeightmaps);
 	ULandscapeInfo* Info = GetLandscapeInfo();
 
 	const int32 AllHeightmapUpdateModes = (ELandscapeLayerUpdateMode::Update_Heightmap_All | ELandscapeLayerUpdateMode::Update_Heightmap_Editing | ELandscapeLayerUpdateMode::Update_Heightmap_Editing_NoCollision);
@@ -4998,7 +4997,6 @@ void ALandscape::ResolveLayersHeightmapTexture(
 	TArray<FLandscapeEditLayerComponentReadbackResult>& InOutComponentReadbackResults)
 {
 	TRACE_CPUPROFILER_EVENT_SCOPE(LandscapeLayers_ResolveLayersHeightmapTexture);
-	SCOPE_CYCLE_COUNTER(STAT_LandscapeLayersResolveHeightmaps);
 
 	ULandscapeInfo* Info = GetLandscapeInfo();
 	if (Info == nullptr)
@@ -5208,7 +5206,6 @@ bool ALandscape::ResolveLayersTexture(
 	bool bIsWeightmap)
 {
 	TRACE_CPUPROFILER_EVENT_SCOPE(LandscapeLayers_ResolveLayersTexture);
-	SCOPE_CYCLE_COUNTER(STAT_LandscapeLayersResolveTexture);
 
 	if (bFlushRender)
 	{
@@ -5579,7 +5576,6 @@ void ALandscape::PrepareComponentDataToPackMaterialLayersCS(int32 InCurrentWeigh
 void ALandscape::ReallocateLayersWeightmaps(FUpdateLayersContentContext& InUpdateLayersContentContext, const TArray<ULandscapeLayerInfoObject*>& InBrushRequiredAllocations)
 {
 	TRACE_CPUPROFILER_EVENT_SCOPE(LandscapeLayers_ReallocateLayersWeightmaps);
-	SCOPE_CYCLE_COUNTER(STAT_LandscapeLayersReallocateWeightmaps);
 
 	ULandscapeInfo* Info = GetLandscapeInfo();
 
@@ -7142,7 +7138,6 @@ int32 ALandscape::PerformLayersWeightmapsGlobalMerge(FUpdateLayersContentContext
 int32 ALandscape::RegenerateLayersWeightmaps(FUpdateLayersContentContext& InUpdateLayersContentContext)
 {
 	TRACE_CPUPROFILER_EVENT_SCOPE(LandscapeLayers_RegenerateLayersWeightmaps);
-	SCOPE_CYCLE_COUNTER(STAT_LandscapeLayersRegenerateWeightmaps);
 	const int32 AllWeightmapUpdateModes = (ELandscapeLayerUpdateMode::Update_Weightmap_All | ELandscapeLayerUpdateMode::Update_Weightmap_Editing | ELandscapeLayerUpdateMode::Update_Weightmap_Editing_NoCollision);
 	const int32 WeightmapUpdateModes = LayerContentUpdateModes & AllWeightmapUpdateModes;
 	const bool bSkipBrush = CVarLandscapeLayerBrushOptim.GetValueOnAnyThread() == 1 && ((WeightmapUpdateModes & AllWeightmapUpdateModes) == ELandscapeLayerUpdateMode::Update_Weightmap_Editing);
@@ -7269,7 +7264,6 @@ uint32 ULandscapeComponent::ComputeWeightmapsHash()
 void ALandscape::UpdateLayersMaterialInstances(const TArray<ULandscapeComponent*>& InLandscapeComponents)
 {
 	TRACE_CPUPROFILER_EVENT_SCOPE(LandscapeLayers_UpdateLayersMaterialInstances);
-	SCOPE_CYCLE_COUNTER(STAT_LandscapeLayersUpdateMaterialInstance);
 	TArray<ULandscapeComponent*> ComponentsToUpdate;
 
 	// Compute Weightmap usage changes
@@ -7439,7 +7433,6 @@ void ALandscape::ResolveLayersWeightmapTexture(
 	TArray<FLandscapeEditLayerComponentReadbackResult>& InOutComponentReadbackResults)
 {
 	TRACE_CPUPROFILER_EVENT_SCOPE(LandscapeLayers_ResolveLayersWeightmapTexture);
-	SCOPE_CYCLE_COUNTER(STAT_LandscapeLayersResolveWeightmaps);
 
 	ULandscapeInfo* Info = GetLandscapeInfo();
 	if (Info == nullptr)
diff --git a/Engine/Source/Runtime/Landscape/Private/LandscapePrivate.h b/Engine/Source/Runtime/Landscape/Private/LandscapePrivate.h
--- a/Engine/Source/Runtime/Landscape/Private/LandscapePrivate.h
+++ b/Engine/Source/Runtime/Landscape/Private/LandscapePrivate.h
@@ -13,34 +13,10 @@ LANDSCAPE_API DECLARE_LOG_CATEGORY_EXTERN(LogLandscapeBP, Display, All);
  * Landscape stats
  */
 DECLARE_CYCLE_STAT_EXTERN(TEXT("Dynamic Draw Time"), STAT_LandscapeDynamicDrawTime, STATGROUP_Landscape, );
-DECLARE_CYCLE_STAT_EXTERN(TEXT("Static Draw LOD Time"), STAT_LandscapeStaticDrawLODTime, STATGROUP_Landscape, );
 DECLARE_CYCLE_STAT_EXTERN(TEXT("Render SetMesh Draw Time VS"), STAT_LandscapeVFDrawTimeVS, STATGROUP_Landscape, );
 DECLARE_CYCLE_STAT_EXTERN(TEXT("Render SetMesh Draw Time PS"), STAT_LandscapeVFDrawTimePS, STATGROUP_Landscape, );
-DECLARE_CYCLE_STAT_EXTERN(TEXT("Init View Custom Data"), STAT_LandscapeInitViewCustomData, STATGROUP_Landscape, );
-DECLARE_CYCLE_STAT_EXTERN(TEXT("PostInit View Custom Data"), STAT_LandscapePostInitViewCustomData, STATGROUP_Landscape, );
-DECLARE_CYCLE_STAT_EXTERN(TEXT("Compute Custom Mesh Batch LOD"), STAT_LandscapeComputeCustomMeshBatchLOD, STATGROUP_Landscape, );
-DECLARE_CYCLE_STAT_EXTERN(TEXT("Compute Custom Shadow Mesh Batch LOD"), STAT_LandscapeComputeCustomShadowMeshBatchLOD, STATGROUP_Landscape, );
-
-DECLARE_CYCLE_STAT_EXTERN(TEXT("Regenerate Layers (RenderThread)"), STAT_LandscapeLayersRegenerate_RenderThread, STATGROUP_Landscape, );
-DECLARE_CYCLE_STAT_EXTERN(TEXT("Regenerate Layers Heightmap (GameThread)"), STAT_LandscapeLayersRegenerateHeightmaps, STATGROUP_Landscape, );
-DECLARE_CYCLE_STAT_EXTERN(TEXT("Resolve Layers Heightmaps"), STAT_LandscapeLayersResolveHeightmaps, STATGROUP_Landscape, );
-
-DECLARE_CYCLE_STAT_EXTERN(TEXT("Resolve Layers Texture"), STAT_LandscapeLayersResolveTexture, STATGROUP_Landscape, );
-DECLARE_CYCLE_STAT_EXTERN(TEXT("Regenerate Layers Weightmap (GameThread)"), STAT_LandscapeLayersRegenerateWeightmaps, STATGROUP_Landscape, );
-
-DECLARE_CYCLE_STAT_EXTERN(TEXT("Resolve Layers Weightmaps"), STAT_LandscapeLayersResolveWeightmaps, STATGROUP_Landscape, );
-DECLARE_CYCLE_STAT_EXTERN(TEXT("Reallocate Layers Weightmaps"), STAT_LandscapeLayersReallocateWeightmaps, STATGROUP_Landscape, );
-DECLARE_CYCLE_STAT_EXTERN(TEXT("Update Layers Material Instances"), STAT_LandscapeLayersUpdateMaterialInstance, STATGROUP_Landscape, );
 
 DECLARE_DWORD_COUNTER_STAT_EXTERN(TEXT("Regenerate Layers DrawCalls"), STAT_LandscapeLayersRegenerateDrawCalls, STATGROUP_Landscape, );
-
-DECLARE_DWORD_COUNTER_STAT_EXTERN(TEXT("Components Using SubSection DrawCall"), STAT_LandscapeComponentUsingSubSectionDrawCalls, STATGROUP_Landscape, );
-DECLARE_DWORD_COUNTER_STAT_EXTERN(TEXT("Tessellated Shadow Cascade"), STAT_LandscapeTessellatedShadowCascade, STATGROUP_Landscape, );
-DECLARE_DWORD_COUNTER_STAT_EXTERN(TEXT("Tessellated Components"), STAT_LandscapeTessellatedComponents, STATGROUP_Landscape, );
 DECLARE_DWORD_COUNTER_STAT_EXTERN(TEXT("Processed Triangles"), STAT_LandscapeTriangles, STATGROUP_Landscape, );
 DECLARE_DWORD_COUNTER_STAT_EXTERN(TEXT("Render Passes"), STAT_LandscapeComponentRenderPasses, STATGROUP_Landscape, );
 DECLARE_DWORD_COUNTER_STAT_EXTERN(TEXT("DrawCalls"), STAT_LandscapeDrawCalls, STATGROUP_Landscape, );
-
-DECLARE_MEMORY_STAT_EXTERN(TEXT("Vertex Mem"), STAT_LandscapeVertexMem, STATGROUP_Landscape, );
-DECLARE_MEMORY_STAT_EXTERN(TEXT("Hole Mem"), STAT_LandscapeHoleMem, STATGROUP_Landscape, );
-DECLARE_DWORD_COUNTER_STAT_EXTERN(TEXT("Component Mem"), STAT_LandscapeComponentMem, STATGROUP_Landscape, );
