Landscape edit layers : optimized mip upload by unlocking all mips at once, to avoid recomputing FTextureSource's hash for each mip we write to (on a landscape with ~256 components : 170ms -> 50ms for reading back heightmaps/weightmaps)

diff --git a/Engine/Source/Runtime/Landscape/Private/LandscapeEditLayers.cpp b/Engine/Source/Runtime/Landscape/Private/LandscapeEditLayers.cpp
--- a/Engine/Source/Runtime/Landscape/Private/LandscapeEditLayers.cpp
+++ b/Engine/Source/Runtime/Landscape/Private/LandscapeEditLayers.cpp
@@ -92,13 +92,13 @@ int32 RenderCaptureLayersNextHeightmapDraws = 0;
 static FAutoConsoleVariableRef CVarRenderCaptureLayersNextHeightmapDraws(
 	TEXT("landscape.RenderCaptureLayersNextHeightmapDraws"),
 	RenderCaptureLayersNextHeightmapDraws,
-	TEXT("Trigger N render capture during the next heightmap draw calls."));
+	TEXT("Trigger N render captures during the next heightmap draw calls."));
 
 int32 RenderCaptureLayersNextWeightmapDraws = 0;
 static FAutoConsoleVariableRef CVarRenderCaptureLayersNextWeightmapDraws(
 	TEXT("landscape.RenderCaptureLayersNextWeightmapDraws"),
 	RenderCaptureLayersNextWeightmapDraws,
-	TEXT("Trigger N render capture during the next weightmap draw calls."));
+	TEXT("Trigger N render captures during the next weightmap draw calls."));
 
 static TAutoConsoleVariable<int32> CVarOutputLayersRTContent(
 	TEXT("landscape.OutputLayersRTContent"),
@@ -1539,6 +1539,7 @@ public:
 	{
 		SCOPED_GPU_STAT(InRHICmdList, LandscapeLayers_Clear);
 		SCOPED_DRAW_EVENTF(InRHICmdList, LandscapeLayers, TEXT("LandscapeLayers_Clear %s"), DebugName.Len() > 0 ? *DebugName : TEXT(""));
+		TRACE_CPUPROFILER_EVENT_SCOPE(LandscapeLayersWeightmapClear_RenderThread::Clear);
 
 		check(IsInRenderingThread());
 
@@ -1582,6 +1583,7 @@ public:
 		SCOPED_GPU_STAT(InRHICmdList, LandscapeLayers_Render);
 		SCOPED_DRAW_EVENTF(InRHICmdList, LandscapeLayers, TEXT("LandscapeLayers_Render %s"), DebugName.Len() > 0 ? *DebugName : TEXT(""));
 		INC_DWORD_STAT(STAT_LandscapeLayersRegenerateDrawCalls);
+		TRACE_CPUPROFILER_EVENT_SCOPE(FLandscapeLayersRender_RenderThread::Render);
 
 		check(IsInRenderingThread());
 
@@ -2690,8 +2692,8 @@ void ALandscape::DrawWeightmapComponentsToRenderTarget(const FString& InDebugNam
 	ENQUEUE_RENDER_COMMAND(LandscapeLayers_Cmd_RenderWeightmap)(
 		[LayersRender, InDebugName, InDrawType, InClearRTWrite](FRHICommandListImmediate& RHICmdList) mutable
 	{
-		TRACE_CPUPROFILER_EVENT_SCOPE(LandscapeLayers_RT_RenderWeightmap);
 		SCOPED_DRAW_EVENTF(RHICmdList, LandscapeLayers, TEXT("DrawWeightmapComponentsToRenderTarget %s (%s)"), *InDebugName, *StaticEnum<EWeightmapRTType>()->GetDisplayValueAsText(InDrawType).ToString());
+		TRACE_CPUPROFILER_EVENT_SCOPE(LandscapeLayers_RT_RenderWeightmap);
 		LayersRender.Render(RHICmdList, InClearRTWrite);
 	});
 
@@ -2867,8 +2869,8 @@ void ALandscape::DrawHeightmapComponentsToRenderTarget(const FString& InDebugNam
 	ENQUEUE_RENDER_COMMAND(LandscapeLayers_Cmd_RenderHeightmap)(
 		[LayersRender, InDebugName, InDrawType, InClearRTWrite](FRHICommandListImmediate& RHICmdList) mutable
 	{
-		TRACE_CPUPROFILER_EVENT_SCOPE(LandscapeLayers_RT_RenderHeightmap);
 		SCOPED_DRAW_EVENTF(RHICmdList, LandscapeLayers, TEXT("DrawHeightmapComponentsToRenderTarget %s (%s)"), *InDebugName, *StaticEnum<EHeightmapRTType>()->GetDisplayValueAsText(InDrawType).ToString());
+		TRACE_CPUPROFILER_EVENT_SCOPE(LandscapeLayers_RT_RenderHeightmap);
 		LayersRender.Render(RHICmdList, InClearRTWrite);
 	});
 
@@ -5283,7 +5285,14 @@ bool ALandscape::ResolveLayersTexture(
 				}
 
 				FMemory::Memcpy(TextureData, OutMipsData[MipIndex].GetData(), OutMipsData[MipIndex].Num() * sizeof(FColor));
+			}
+		}
 
+		// Unlock all mips at once because there's a lock counter in FTextureSource that recomputes the content hash when reaching 0 (which means we'd recompute the hash several times over if we Lock/Unlock/Lock/Unloc/... for each mip ):
+		for (int8 MipIndex = 0; MipIndex < OutMipsData.Num(); ++MipIndex)
+		{
+			if (OutMipsData[MipIndex].Num() > 0)
+			{
 				InOutputTexture->Source.UnlockMip(MipIndex);
 			}
 		}
