Added comments to landscape grass exclusion box system to improve discoverability

#rnx
#skipci

diff --git a/Engine/Source/Runtime/Landscape/Classes/LandscapeProxy.h b/Engine/Source/Runtime/Landscape/Classes/LandscapeProxy.h
--- a/Engine/Source/Runtime/Landscape/Classes/LandscapeProxy.h
+++ b/Engine/Source/Runtime/Landscape/Classes/LandscapeProxy.h
@@ -896,9 +896,23 @@ public:
 	LANDSCAPE_API void UpdateGrass(const TArray<FVector>& Cameras, int32& InOutNumComponentsCreated, bool bForceSync = false);
 	LANDSCAPE_API void UpdateGrass(const TArray<FVector>& Cameras, bool bForceSync = false);
 
-	// TODO [jonathan.bard] : Rename to "AddGrassExlusionBox" + no reason for any of this to be static
+
+	/**
+	 * Registers an axis-aligned bounding box that will act as an exclusion volume for all landscape grass overlapping it
+	 * @param Owning UObject of the box. Acts as an identifier of the exclusion volume. Also, when the UObject becomes stale, the box will be automatically unregistered from landscapes
+	 * @param BoxToRemove AABBox (excluded volume)
+	 */
+	// TODO [jonathan.bard] : Rename to "AddGrassExclusionBox" + no reason for any of this to be static
 	LANDSCAPE_API static void AddExclusionBox(FWeakObjectPtr Owner, const FBox& BoxToRemove);
+	/**
+	 * Unregisters a previously-registered exclusion box. Landscape grass will be able to be spawned again in this area after the operation
+	 * @param Owner Identifier of the box to remove (see AddExclusionBox)
+	 */
 	LANDSCAPE_API static void RemoveExclusionBox(FWeakObjectPtr Owner);
+
+	/**
+	 * Unregisters all existing exclusion boxes.
+	 */
 	LANDSCAPE_API static void RemoveAllExclusionBoxes();
 
 
