Landscape: Import Gimzo usability improvements

* Added snap modes (texel and component)
* Added reset to landscape minimum corner.

diff --git a/Engine/Source/Editor/LandscapeEditor/Private/LandscapeEdMode.cpp b/Engine/Source/Editor/LandscapeEditor/Private/LandscapeEdMode.cpp
--- a/Engine/Source/Editor/LandscapeEditor/Private/LandscapeEdMode.cpp
+++ b/Engine/Source/Editor/LandscapeEditor/Private/LandscapeEdMode.cpp
@@ -510,8 +510,7 @@ void FEdModeLandscape::Enter()
 	if (CurrentGizmoActor.IsValid())
 	{
 		CurrentGizmoActor->SetTargetLandscape(CurrentToolTarget.LandscapeInfo.Get());
-
-		CurrentGizmoActor.Get()->bSnapToLandscapeGrid = UISettings->bSnapGizmo;
+		CurrentGizmoActor->SnapType = UISettings->SnapMode;
 	}
 
 	int32 SquaredDataTex = ALandscapeGizmoActiveActor::DataTexSize * ALandscapeGizmoActiveActor::DataTexSize;
diff --git a/Engine/Source/Editor/LandscapeEditor/Private/LandscapeEdModeComponentTools.cpp b/Engine/Source/Editor/LandscapeEditor/Private/LandscapeEdModeComponentTools.cpp
--- a/Engine/Source/Editor/LandscapeEditor/Private/LandscapeEdModeComponentTools.cpp
+++ b/Engine/Source/Editor/LandscapeEditor/Private/LandscapeEdModeComponentTools.cpp
@@ -1738,11 +1738,20 @@ public:
 		ALandscapeGizmoActiveActor* Gizmo = this->EdMode->CurrentGizmoActor.Get();
 		if (Gizmo)
 		{
+			Gizmo->bFollowTerrainHeight = false;
 			GEditor->SelectNone(false, true);
 			GEditor->SelectActor(Gizmo, true, true, true);
 		}
 	}
 
+	virtual void ExitTool() override
+	{
+		if (ALandscapeGizmoActiveActor* Gizmo = this->EdMode->CurrentGizmoActor.Get())
+		{
+			Gizmo->bFollowTerrainHeight = true;
+		}
+	}
+
 	// Copy tool doesn't use any view information, so just do it as one function
 	void Copy()
 	{
@@ -2613,6 +2622,11 @@ public:
 			}
 		}
 	}
+	
+	virtual EAxisList::Type GetWidgetAxisToDraw(UE::Widget::EWidgetMode InWidgetMode) const override
+	{
+		return EAxisList::XY;
+	}
 };
 
 //////////////////////////////////////////////////////////////////////////
diff --git a/Engine/Source/Editor/LandscapeEditor/Private/LandscapeEditorDetailCustomization_ImportExport.cpp b/Engine/Source/Editor/LandscapeEditor/Private/LandscapeEditorDetailCustomization_ImportExport.cpp
--- a/Engine/Source/Editor/LandscapeEditor/Private/LandscapeEditorDetailCustomization_ImportExport.cpp
+++ b/Engine/Source/Editor/LandscapeEditor/Private/LandscapeEditorDetailCustomization_ImportExport.cpp
@@ -18,6 +18,7 @@
 #include "Framework/Application/SlateApplication.h"
 #include "Framework/MultiBox/MultiBoxBuilder.h"
 
+#include "ScopedTransaction.h"
 #include "DesktopPlatformModule.h"
 
 #include "LandscapeEditorObject.h"
@@ -241,6 +242,21 @@ void FLandscapeEditorDetailCustomization_ImportExport::CustomizeDetails(IDetailL
 	TSharedRef<IPropertyHandle> PropertyHandle_Layers = DetailBuilder.GetProperty(GET_MEMBER_NAME_CHECKED(ULandscapeEditorObject, ImportLandscape_Layers));
 	ImportExportCategory.AddProperty(PropertyHandle_Layers);
 
+	ImportExportCategory.AddCustomRow(FText::GetEmpty())
+	[
+		SNew(SHorizontalBox)
+		+ SHorizontalBox::Slot()
+		.FillWidth(1)
+		+ SHorizontalBox::Slot()
+		.AutoWidth()
+		[
+			SNew(SButton)
+			.Text(LOCTEXT("OriginReset", "Reset Import Location"))
+			.OnClicked_Static(&OnOriginResetButtonClicked)
+		]
+	]
+	.Visibility(TAttribute<EVisibility>(this, &FLandscapeEditorDetailCustomization_ImportExport::GetImportingVisibility));
+	
 	ImportExportCategory.AddCustomRow(FText::GetEmpty())
 	[
 		SNew(SHorizontalBox)
@@ -509,6 +525,29 @@ FText FLandscapeEditorDetailCustomization_ImportExport::GetImportExportButtonTex
 	return LOCTEXT("Export", "Export");
 }
 
+FReply FLandscapeEditorDetailCustomization_ImportExport::OnOriginResetButtonClicked()
+{
+	FScopedTransaction Transaction(LOCTEXT("LandscapeEdResetImportLocation", "Reset Import Location"));
+	
+	FEdModeLandscape* LandscapeEdMode = GetEditorMode();
+	check(LandscapeEdMode);
+	
+	TWeakObjectPtr<ULandscapeInfo> LandscapeInfo = LandscapeEdMode->CurrentToolTarget.LandscapeInfo;
+	TWeakObjectPtr<ALandscapeGizmoActiveActor> CurrentGizmoActor = LandscapeEdMode->CurrentGizmoActor;
+
+	check(LandscapeInfo.Get());
+	check(CurrentGizmoActor.Get());
+	
+	FIntRect LandscapeExtent; 
+	LandscapeInfo->GetLandscapeExtent(LandscapeExtent);
+	const FTransform LandscapeToWorldTransform = LandscapeInfo->GetLandscapeProxy()->LandscapeActorToWorld();
+	const FVector LandscapeOrigin = LandscapeToWorldTransform.TransformPosition(FVector(LandscapeExtent.Min.X,LandscapeExtent.Min.Y,0.0f));
+
+	CurrentGizmoActor->SetActorLocation(LandscapeOrigin, false);
+	
+	return FReply::Handled();
+}
+
 FReply FLandscapeEditorDetailCustomization_ImportExport::OnImportExportButtonClicked()
 {
 	FEdModeLandscape* LandscapeEdMode = GetEditorMode();
@@ -690,6 +729,11 @@ bool FLandscapeEditorDetailCustomization_ImportExport::GetImportExportButtonIsEn
 	return bHasOneSelection && LandscapeEdMode->IsLandscapeResolutionCompliant();
 }
 
+EVisibility FLandscapeEditorDetailCustomization_ImportExport::GetImportingVisibility() const
+{
+	return IsImporting() ? EVisibility::Visible : EVisibility::Hidden;
+}
+
 EVisibility FLandscapeEditorDetailCustomization_ImportExport::GetImportExportLandscapeErrorVisibility() const
 {
 	FEdModeLandscape* EdMode = GetEditorMode();
diff --git a/Engine/Source/Editor/LandscapeEditor/Private/LandscapeEditorDetailCustomization_ImportExport.h b/Engine/Source/Editor/LandscapeEditor/Private/LandscapeEditorDetailCustomization_ImportExport.h
--- a/Engine/Source/Editor/LandscapeEditor/Private/LandscapeEditorDetailCustomization_ImportExport.h
+++ b/Engine/Source/Editor/LandscapeEditor/Private/LandscapeEditorDetailCustomization_ImportExport.h
@@ -50,8 +50,11 @@ private:
 	static FText GetImportExportButtonText();
 	static FReply OnImportExportButtonClicked();
 	static bool GetImportExportButtonIsEnabled();
+	
+	static FReply OnOriginResetButtonClicked();
 
 	EVisibility GetImportExportLandscapeErrorVisibility() const;
+	EVisibility GetImportingVisibility() const;
 	FText GetImportExportLandscapeErrorText() const;
 
 	static TSharedRef<SWidget> GetImportLandscapeResolutionMenu();
diff --git a/Engine/Source/Editor/LandscapeEditor/Private/LandscapeEditorObject.cpp b/Engine/Source/Editor/LandscapeEditor/Private/LandscapeEditorObject.cpp
--- a/Engine/Source/Editor/LandscapeEditor/Private/LandscapeEditorObject.cpp
+++ b/Engine/Source/Editor/LandscapeEditor/Private/LandscapeEditorObject.cpp
@@ -68,7 +68,7 @@ ULandscapeEditorObject::ULandscapeEditorObject(const FObjectInitializer& ObjectI
 
 	, PasteMode(ELandscapeToolPasteMode::Both)
 	, bApplyToAllTargets(true)
-	, bSnapGizmo(false)
+	, SnapMode(ELandscapeGizmoSnapType::None)
 	, bSmoothGizmoBrush(true)
 
 	, MirrorPoint(FVector::ZeroVector)
@@ -139,7 +139,7 @@ void ULandscapeEditorObject::PostEditChangeProperty(FPropertyChangedEvent& Prope
 	SetbUseSelectedRegion(bUseSelectedRegion);
 	SetbUseNegativeMask(bUseNegativeMask);
 	SetPasteMode(PasteMode);
-	SetbSnapGizmo(bSnapGizmo);
+	SetGizmoSnapMode(SnapMode);
 
 	if (PropertyChangedEvent.MemberProperty == nullptr ||
 		PropertyChangedEvent.MemberProperty->GetFName() == GET_MEMBER_NAME_CHECKED(ULandscapeEditorObject, AlphaTexture) ||
@@ -476,16 +476,16 @@ void ULandscapeEditorObject::SetPasteMode(ELandscapeToolPasteMode InPasteMode)
 	PasteMode = InPasteMode;
 }
 
-void ULandscapeEditorObject::SetbSnapGizmo(bool InbSnapGizmo)
+void ULandscapeEditorObject::SetGizmoSnapMode(ELandscapeGizmoSnapType InSnapMode)
 {
-	bSnapGizmo = InbSnapGizmo;
+	SnapMode = InSnapMode;
 
 	if (ParentMode->CurrentGizmoActor.IsValid())
 	{
-		ParentMode->CurrentGizmoActor->bSnapToLandscapeGrid = bSnapGizmo;
+		ParentMode->CurrentGizmoActor->SnapType = SnapMode;
 	}
 
-	if (bSnapGizmo)
+	if (SnapMode != ELandscapeGizmoSnapType::None)
 	{
 		if (ParentMode->CurrentGizmoActor.IsValid())
 		{
diff --git a/Engine/Source/Editor/LandscapeEditor/Public/LandscapeEditorObject.h b/Engine/Source/Editor/LandscapeEditor/Public/LandscapeEditorObject.h
--- a/Engine/Source/Editor/LandscapeEditor/Public/LandscapeEditorObject.h
+++ b/Engine/Source/Editor/LandscapeEditor/Public/LandscapeEditorObject.h
@@ -444,9 +444,8 @@ class ULandscapeEditorObject : public UObject
 	UPROPERTY(Category="Tool Settings", EditAnywhere, NonTransactional, meta=(DisplayName="Gizmo copy/paste all layers", ShowForTools="CopyPaste"))
 	bool bApplyToAllTargets;
 
-	// Makes sure the gizmo is snapped perfectly to the landscape so that the sample points line up, which makes copy/paste less blurry. Irrelevant if gizmo is scaled
-	UPROPERTY(Category="Tool Settings", EditAnywhere, NonTransactional, meta=(DisplayName="Snap Gizmo to Landscape grid", ShowForTools="CopyPaste"))
-	bool bSnapGizmo;
+	UPROPERTY(Category="Tool Settings", EditAnywhere, NonTransactional, meta=(DisplayName="Snap Gizmo to Landscape grid", ShowForTools="CopyPaste,ImportExport"))
+	ELandscapeGizmoSnapType SnapMode;
 
 	// Smooths the edges of the gizmo data into the landscape. Without this, the edges of the pasted data will be sharp
 	UPROPERTY(Category="Tool Settings", EditAnywhere, NonTransactional, meta=(DisplayName="Use Smooth Gizmo Brush", ShowForTools="CopyPaste"))
@@ -765,7 +764,7 @@ public:
 		}
 	}
 
-	void SetbSnapGizmo(bool InbSnapGizmo);
+	void SetGizmoSnapMode(ELandscapeGizmoSnapType SnapMode);
 
 	void SetParent(FEdModeLandscape* LandscapeParent)
 	{
diff --git a/Engine/Source/Runtime/Landscape/Classes/LandscapeGizmoActiveActor.h b/Engine/Source/Runtime/Landscape/Classes/LandscapeGizmoActiveActor.h
--- a/Engine/Source/Runtime/Landscape/Classes/LandscapeGizmoActiveActor.h
+++ b/Engine/Source/Runtime/Landscape/Classes/LandscapeGizmoActiveActor.h
@@ -24,6 +24,14 @@ enum ELandscapeGizmoType
 	LGT_MAX,
 };
 
+UENUM()
+enum class ELandscapeGizmoSnapType
+{
+	None,
+	Component,
+	Texel
+};
+
 USTRUCT()
 struct FGizmoSelectData
 {
@@ -102,7 +110,10 @@ class ALandscapeGizmoActiveActor : public ALandscapeGizmoActor
 	TArray<TObjectPtr<ULandscapeLayerInfoObject>> LayerInfos;
 
 	UPROPERTY(transient)
-	bool bSnapToLandscapeGrid;
+	ELandscapeGizmoSnapType SnapType = ELandscapeGizmoSnapType::None;
+
+	UPROPERTY(transient)
+	bool bFollowTerrainHeight = true;
 
 	UPROPERTY(transient)
 	FRotator UnsnappedRotation;
diff --git a/Engine/Source/Runtime/Landscape/Private/LandscapeGizmoActor.cpp b/Engine/Source/Runtime/Landscape/Private/LandscapeGizmoActor.cpp
--- a/Engine/Source/Runtime/Landscape/Private/LandscapeGizmoActor.cpp
+++ b/Engine/Source/Runtime/Landscape/Private/LandscapeGizmoActor.cpp
@@ -152,7 +152,7 @@ public:
 			&& !FMath::IsNearlyZero(Gizmo->CachedHeight))
 		{
 			ULandscapeInfo* LandscapeInfo = Gizmo->TargetLandscapeInfo;
-			if (LandscapeInfo && LandscapeInfo->GetLandscapeProxy() && !FMath::IsNearlyZero(LandscapeInfo->DrawScale.X))
+ 			if (LandscapeInfo && LandscapeInfo->GetLandscapeProxy() && !FMath::IsNearlyZero(LandscapeInfo->DrawScale.X))
 			{
 				const float ScaleXY = LandscapeInfo->DrawScale.X;
 				SampleSizeX = Gizmo->SampleSizeX;
@@ -567,26 +567,52 @@ void ALandscapeGizmoActiveActor::PostEditMove(bool bFinished)
 FVector ALandscapeGizmoActiveActor::SnapToLandscapeGrid(const FVector& GizmoLocation) const
 {
 	check(TargetLandscapeInfo);
-	const FTransform LToW = TargetLandscapeInfo->GetLandscapeProxy()->LandscapeActorToWorld();
-	const FVector LandscapeSpaceLocation = LToW.InverseTransformPosition(GizmoLocation);
-	const FVector SnappedLandscapeSpaceLocation = LandscapeSpaceLocation.GridSnap(1);
-	const FVector ResultLocation = LToW.TransformPosition(SnappedLandscapeSpaceLocation);
-	return ResultLocation;
-}
+	const ALandscapeProxy* LandscapeProxy = TargetLandscapeInfo->GetLandscapeProxy();
 
-void ALandscapeGizmoActiveActor::EditorApplyTranslation(const FVector& DeltaTranslation, bool bAltDown, bool bShiftDown, bool bCtrlDown)
-{
-	if (bSnapToLandscapeGrid)
+	const FTransform LToW = LandscapeProxy->LandscapeActorToWorld();
+	const FVector LandscapeSpaceLocation = LToW.InverseTransformPosition(GizmoLocation);
+	float SnapDimension;
+	switch(SnapType)
 	{
-		const FVector GizmoLocation = GetActorLocation() + DeltaTranslation;
-		const FVector ResultLocation = SnapToLandscapeGrid(GizmoLocation);
+	case ELandscapeGizmoSnapType::Component:
+		SnapDimension = LandscapeProxy->ComponentSizeQuads;
+		break;
+	case ELandscapeGizmoSnapType::Texel:
+		SnapDimension = 1.0f;
+		break;
+	case ELandscapeGizmoSnapType::None:
+		SnapDimension = 0.0f; 
+		break;
+	default:
+		SnapDimension = 0.0f;
+		break;
+	}
 
-		SetActorLocation(ResultLocation, false);
+	FVector ResultLocation;
+	if (SnapDimension > 0.0f)
+	{
+		const FVector SnappedLandscapeSpaceLocation = LandscapeSpaceLocation.GridSnap(SnapDimension);
+		ResultLocation = LToW.TransformPosition(SnappedLandscapeSpaceLocation);	
 	}
 	else
 	{
-		Super::EditorApplyTranslation(DeltaTranslation, bAltDown, bShiftDown, bCtrlDown);
+		ResultLocation = GizmoLocation;
+	}
+
+	if (bFollowTerrainHeight)
+	{
+		ResultLocation.Z = LandscapeProxy->GetHeightAtLocation(ResultLocation, EHeightfieldSource::Editor).Get(GizmoLocation.Z);
 	}
+	
+	return ResultLocation;
+}
+
+void ALandscapeGizmoActiveActor::EditorApplyTranslation(const FVector& DeltaTranslation, bool bAltDown, bool bShiftDown, bool bCtrlDown)
+{
+	const FVector GizmoLocation = GetActorLocation() + DeltaTranslation;
+	const FVector ResultLocation = SnapToLandscapeGrid(GizmoLocation);
+
+	SetActorLocation(ResultLocation, false);
 
 	ReregisterAllComponents();
 }
@@ -607,7 +633,7 @@ FRotator ALandscapeGizmoActiveActor::SnapToLandscapeGrid(const FRotator& GizmoRo
 
 void ALandscapeGizmoActiveActor::EditorApplyRotation(const FRotator& DeltaRotation, bool bAltDown, bool bShiftDown, bool bCtrlDown)
 {
-	if (bSnapToLandscapeGrid)
+	if (SnapType == ELandscapeGizmoSnapType::Texel || SnapType == ELandscapeGizmoSnapType::Component)
 	{
 		// Based on AActor::EditorApplyRotation
 		FRotator GizmoRotation = GetActorRotation() + UnsnappedRotation;
