Landscape : Only generate LandscapeCollision SceneProxies if any view has a collision view flag enabled.

diff --git a/Engine/Source/Runtime/Landscape/Classes/LandscapeHeightfieldCollisionComponent.h b/Engine/Source/Runtime/Landscape/Classes/LandscapeHeightfieldCollisionComponent.h
--- a/Engine/Source/Runtime/Landscape/Classes/LandscapeHeightfieldCollisionComponent.h
+++ b/Engine/Source/Runtime/Landscape/Classes/LandscapeHeightfieldCollisionComponent.h
@@ -30,9 +30,10 @@ namespace Chaos
 
 enum class EHeightfieldSource
 {
-	Simple = 0,
-	Complex = 1,
-	Editor = 2
+	None,
+	Simple,
+	Complex,
+	Editor
 };
 
 UCLASS(MinimalAPI, Within=LandscapeProxy)
diff --git a/Engine/Source/Runtime/Landscape/Private/LandscapeCollision.cpp b/Engine/Source/Runtime/Landscape/Private/LandscapeCollision.cpp
--- a/Engine/Source/Runtime/Landscape/Private/LandscapeCollision.cpp
+++ b/Engine/Source/Runtime/Landscape/Private/LandscapeCollision.cpp
@@ -18,6 +18,8 @@
 #include "Misc/SecureHash.h"
 #include "CollisionQueryParams.h"
 #include "Engine/World.h"
+#include "LandscapeSubsystem.h"
+#include "LandscapeRender.h"
 #include "LandscapeProxy.h"
 #include "LandscapeInfo.h"
 #include "Interfaces/Interface_CollisionDataProvider.h"
@@ -66,11 +68,11 @@ using namespace PhysicsInterfaceTypes;
 // Global switch for whether to read/write to DDC for landscape cooked data
 // It's a lot faster to compute than to request from DDC, so always skip.
 bool GLandscapeCollisionSkipDDC = true;
-static int32 CVarLandscapeShowCollisionMeshCurrentValue = 0;
+static int32 CVarLandscapeShowCollisionMeshCurrentValue = static_cast<int>(EHeightfieldSource::Simple);
 static TAutoConsoleVariable<int32> CVarLandscapeShowCollisionMesh(
 	TEXT("landscape.ShowCollisionMesh"),
 	CVarLandscapeShowCollisionMeshCurrentValue,
-	TEXT("Selects which heightfield to visualize when ShowFlags.Collision is used. 0 for simple, 1 for complex, 2 for editor only."),
+	TEXT("Selects which heightfield to visualize when ShowFlags.Collision is used. 0 to disable, 1 for simple, 2 for complex, 3 for editor only."),
 	ECVF_RenderThreadSafe
 );
 
@@ -621,7 +623,12 @@ FPrimitiveSceneProxy* ULandscapeHeightfieldCollisionComponent::CreateSceneProxy(
 	};
 
 	FLandscapeHeightfieldCollisionComponentSceneProxy* Proxy = nullptr;
-
+	if (ULandscapeSubsystem* LandscapeSubsystem = this->GetWorld()->GetSubsystem<ULandscapeSubsystem>(); 
+		LandscapeSubsystem && !LandscapeSubsystem->AnyViewShowCollisions())
+	{
+		return Proxy;
+	}
+	
 	if (HeightfieldRef.IsValid() && IsValidRef(HeightfieldRef))
 	{
 		const Chaos::FHeightField* LocalHeightfield = nullptr;
@@ -629,6 +636,9 @@ FPrimitiveSceneProxy* ULandscapeHeightfieldCollisionComponent::CreateSceneProxy(
 
 		switch (static_cast<EHeightfieldSource>(CVarLandscapeShowCollisionMesh.GetValueOnGameThread()))
 		{
+		case EHeightfieldSource::None:
+			WireframeColor = FColor(0, 0, 0, 0);
+			break;
 		case EHeightfieldSource::Simple:
 			if (HeightfieldRef->HeightfieldSimple.IsValid())
 			{
@@ -2656,6 +2666,8 @@ TOptional<float> ULandscapeHeightfieldCollisionComponent::GetHeight(float X, flo
 	
 	switch(HeightFieldSource)
 	{
+	case EHeightfieldSource::None:
+		break;
 	case EHeightfieldSource::Simple:
 		HeightField = HeightfieldRef->HeightfieldSimple.Get(); 
 		break;
diff --git a/Engine/Source/Runtime/Landscape/Private/LandscapeRender.cpp b/Engine/Source/Runtime/Landscape/Private/LandscapeRender.cpp
--- a/Engine/Source/Runtime/Landscape/Private/LandscapeRender.cpp
+++ b/Engine/Source/Runtime/Landscape/Private/LandscapeRender.cpp
@@ -606,12 +606,23 @@ void FLandscapeRenderSystem::UpdateBuffers()
 //
 FLandscapeSceneViewExtension::FLandscapeSceneViewExtension(const FAutoRegister& AutoReg) : FSceneViewExtensionBase(AutoReg)
 {
+
+	FCoreDelegates::OnEndFrame.AddRaw(this, &FLandscapeSceneViewExtension::EndFrame_GameThread);
 	FCoreDelegates::OnEndFrameRT.AddRaw(this, &FLandscapeSceneViewExtension::EndFrame_RenderThread);
 }
 
 FLandscapeSceneViewExtension::~FLandscapeSceneViewExtension()
 {
 	FCoreDelegates::OnEndFrameRT.RemoveAll(this);
+	FCoreDelegates::OnEndFrame.RemoveAll(this);
+}
+
+void FLandscapeSceneViewExtension::BeginRenderViewFamily(FSceneViewFamily& InViewFamily) 
+{
+	if (InViewFamily.EngineShowFlags.Collision)
+	{
+		NumViewsWithShowCollisionAcc++;
+	}
 }
 
 void FLandscapeSceneViewExtension::PreRenderView_RenderThread(FRDGBuilder& GraphBuilder, FSceneView& InView)
@@ -693,6 +704,12 @@ void FLandscapeSceneViewExtension::PreInitViews_RenderThread(FRDGBuilder& GraphB
 	LandscapeViews.Reset();
 }
 
+void FLandscapeSceneViewExtension::EndFrame_GameThread()
+{
+	NumViewsWithShowCollision = NumViewsWithShowCollisionAcc;
+	NumViewsWithShowCollisionAcc = 0;
+}
+
 // TODO [jonathan.bard] Ideally this should be symmetrical with FLandscapeSceneViewExtension::PreRenderView_RenderThread and should be called in FLandscapeSceneViewExtension::PostRenderView_RenderThread
 //  but PostRenderView_RenderThread is not currently called on all code paths (scene/reflection captures, etc.)
 void FLandscapeSceneViewExtension::EndFrame_RenderThread()
diff --git a/Engine/Source/Runtime/Landscape/Private/LandscapeSubsystem.cpp b/Engine/Source/Runtime/Landscape/Private/LandscapeSubsystem.cpp
--- a/Engine/Source/Runtime/Landscape/Private/LandscapeSubsystem.cpp
+++ b/Engine/Source/Runtime/Landscape/Private/LandscapeSubsystem.cpp
@@ -11,6 +11,9 @@
 #include "LandscapeStreamingProxy.h"
 #include "LandscapeInfo.h"
 #include "LandscapeInfoMap.h"
+#include "LandscapeModule.h"
+#include "LandscapeRender.h"
+#include "LandscapePrivate.h"
 #include "ProfilingDebugging/CsvProfiler.h"
 #include "WorldPartition/WorldPartitionSubsystem.h"
 #include "ActorPartition/ActorPartitionSubsystem.h"
@@ -170,6 +173,23 @@ void ULandscapeSubsystem::Tick(float DeltaTime)
 
 	Super::Tick(DeltaTime);
 
+#if WITH_EDITOR
+	//Check if we need to start or stop creating Collision SceneProxies
+	ILandscapeModule& LandscapeModule = FModuleManager::GetModuleChecked<ILandscapeModule>("Landscape");
+	int32 NumViewsWithShowCollision = LandscapeModule.GetLandscapeSceneViewExtension()->GetNumViewsWithShowCollision();
+	bool bNewShowCollisions = NumViewsWithShowCollision > 0;
+    bool bCollisionChanged = bNewShowCollisions != bAnyViewShowCollisions;
+    bAnyViewShowCollisions = bNewShowCollisions;
+        
+	if (bCollisionChanged)
+	{
+		for (ULandscapeHeightfieldCollisionComponent* LandscapeHeightfieldCollisionComponent : TObjectRange<ULandscapeHeightfieldCollisionComponent>(RF_ClassDefaultObject | RF_ArchetypeObject, true, EInternalObjectFlags::Garbage))
+		{
+			LandscapeHeightfieldCollisionComponent->MarkRenderStateDirty();
+		}
+	}
+#endif
+	
 	UWorld* World = GetWorld();
 
 	static TArray<FVector> OldCameras;
diff --git a/Engine/Source/Runtime/Landscape/Public/LandscapeRender.h b/Engine/Source/Runtime/Landscape/Public/LandscapeRender.h
--- a/Engine/Source/Runtime/Landscape/Public/LandscapeRender.h
+++ b/Engine/Source/Runtime/Landscape/Public/LandscapeRender.h
@@ -490,16 +490,18 @@ public:
 	FLandscapeSceneViewExtension(const FAutoRegister& AutoReg);
 	virtual ~FLandscapeSceneViewExtension();
 
+	void EndFrame_GameThread();
 	void EndFrame_RenderThread();
 
 	virtual void SetupViewFamily(FSceneViewFamily& InViewFamily) override {}
 	virtual void SetupView(FSceneViewFamily& InViewFamily, FSceneView& InView) override {}
-	virtual void BeginRenderViewFamily(FSceneViewFamily& InViewFamily) override {}
+	virtual void BeginRenderViewFamily(FSceneViewFamily& InViewFamily) override;
 
 	virtual void PreRenderView_RenderThread(FRDGBuilder& GraphBuilder, FSceneView& InView) override;
 	virtual void PreInitViews_RenderThread(FRDGBuilder& GraphBuilder) override;
 
 	LANDSCAPE_API const TMap<uint32, FLandscapeRenderSystem*>& GetLandscapeRenderSystems() const;
+	LANDSCAPE_API int32 GetNumViewsWithShowCollision() const { return NumViewsWithShowCollision; }
 private:
 	FBufferRHIRef LandscapeLODDataBuffer;
 	FBufferRHIRef LandscapeIndirectionBuffer;
@@ -519,6 +521,8 @@ private:
 
 	TArray<FLandscapeViewData> LandscapeViews;
 	UE::Tasks::FTask LandscapeSetupTask;
+	int32 NumViewsWithShowCollision = 0;    // Last frame number of views with collision enabled.
+	int32 NumViewsWithShowCollisionAcc = 0; // Accumulate the number of views with collision
 };
 
 
diff --git a/Engine/Source/Runtime/Landscape/Public/LandscapeSubsystem.h b/Engine/Source/Runtime/Landscape/Public/LandscapeSubsystem.h
--- a/Engine/Source/Runtime/Landscape/Public/LandscapeSubsystem.h
+++ b/Engine/Source/Runtime/Landscape/Public/LandscapeSubsystem.h
@@ -85,6 +85,7 @@ public:
 	LANDSCAPE_API static bool IsDirtyOnlyInModeEnabled();
 	FLandscapeNotificationManager* GetNotificationManager() { return NotificationManager; }
 	FOnHeightmapStreamedDelegate& GetOnHeightmapStreamedDelegate() { return OnHeightmapStreamed; }
+	LANDSCAPE_API bool AnyViewShowCollisions() const { return bAnyViewShowCollisions; }  //! Returns true if any view has view collisions enabled.
 #endif // WITH_EDITOR
 
 private:
@@ -107,5 +108,6 @@ private:
 	
 	FLandscapeNotificationManager* NotificationManager;
 	FOnHeightmapStreamedDelegate OnHeightmapStreamed;
+	bool bAnyViewShowCollisions = false;
 #endif // WITH_EDITOR
 };
