Fixed static analysis warnings

#rnx

diff --git a/Engine/Source/Runtime/Engine/Private/Components/SplineMeshComponent.cpp b/Engine/Source/Runtime/Engine/Private/Components/SplineMeshComponent.cpp
--- a/Engine/Source/Runtime/Engine/Private/Components/SplineMeshComponent.cpp
+++ b/Engine/Source/Runtime/Engine/Private/Components/SplineMeshComponent.cpp
@@ -1312,7 +1312,7 @@ void USplineMeshComponent::PostEditChangeProperty(FPropertyChangedEvent& Propert
 		Utilities->HandleActorModified(GetOwner());
 	}
 
-	if (MemberPropertyThatChanged->GetFName() == GET_MEMBER_NAME_CHECKED(USplineMeshComponent, bNeverNeedsCookedCollisionData))
+	if (MemberPropertyThatChanged && (MemberPropertyThatChanged->GetFName() == GET_MEMBER_NAME_CHECKED(USplineMeshComponent, bNeverNeedsCookedCollisionData)))
 	{
 		// TODO [jonathan.bard] : this is currently needed because Setter doesn't correctly do its job in the details panel but eventually this could be removed : 
 		SetbNeverNeedsCookedCollisionData(bNeverNeedsCookedCollisionData);
diff --git a/Engine/Source/Runtime/Landscape/Private/Landscape.cpp b/Engine/Source/Runtime/Landscape/Private/Landscape.cpp
--- a/Engine/Source/Runtime/Landscape/Private/Landscape.cpp
+++ b/Engine/Source/Runtime/Landscape/Private/Landscape.cpp
@@ -2993,14 +2993,21 @@ void ALandscapeProxy::Serialize(FArchive& Ar)
 		}
 	}
 
-	// Fixup Nanite meshes which were using the wrong material and didn't have proper UVs :
-	// Also remove cooked collision data from Nanite landscape meshes, since collisions are handled by ULandscapeHeighfieldCollisionComponent :
-	if (Ar.IsLoading() 
-		&& ((Ar.CustomVer(FFortniteReleaseBranchCustomObjectVersion::GUID) < FFortniteReleaseBranchCustomObjectVersion::FixupNaniteLandscapeMeshes)
-			|| (Ar.CustomVer(FFortniteReleaseBranchCustomObjectVersion::GUID) < FFortniteReleaseBranchCustomObjectVersion::RemoveUselessLandscapeMeshesCookedCollisionData)))
-	{
-		// This will force the Nanite meshes to be properly regenerated during the next save :
-		InvalidateNaniteRepresentation(/* bCheckContentId = */ false);
+	if (Ar.IsLoading())
+	{
+		// Fixup Nanite meshes which were using the wrong material and didn't have proper UVs :
+		if (Ar.CustomVer(FFortniteReleaseBranchCustomObjectVersion::GUID) < FFortniteReleaseBranchCustomObjectVersion::FixupNaniteLandscapeMeshes)
+		{
+			// This will force the Nanite meshes to be properly regenerated during the next save :
+			InvalidateNaniteRepresentation(/* bCheckContentId = */ false);
+		}
+
+		// Remove cooked collision data from Nanite landscape meshes, since collisions are handled by ULandscapeHeighfieldCollisionComponent :		
+		if (Ar.CustomVer(FFortniteReleaseBranchCustomObjectVersion::GUID) < FFortniteReleaseBranchCustomObjectVersion::RemoveUselessLandscapeMeshesCookedCollisionData)
+		{
+			// This will force the Nanite meshes to be properly regenerated during the next save :
+			InvalidateNaniteRepresentation(/* bCheckContentId = */ false);
+		}
 	}
 #endif
 }
