* Unified caption formatting for landscape related material expressions
* Landscape Layer Weight/Sample/Switch's title can now be edited directly in the graph, similar to Texture sample material expressions

diff --git a/Engine/Source/Runtime/Landscape/Classes/Materials/MaterialExpressionLandscapeLayerSample.h b/Engine/Source/Runtime/Landscape/Classes/Materials/MaterialExpressionLandscapeLayerSample.h
--- a/Engine/Source/Runtime/Landscape/Classes/Materials/MaterialExpressionLandscapeLayerSample.h
+++ b/Engine/Source/Runtime/Landscape/Classes/Materials/MaterialExpressionLandscapeLayerSample.h
@@ -27,6 +27,9 @@ public:
 
 	//~ Begin UMaterialExpression Interface
 #if WITH_EDITOR
+	virtual bool CanRenameNode() const override { return true; }
+	virtual FString GetEditableName() const override;
+	virtual void SetEditableName(const FString& NewName) override;
 	virtual int32 Compile(class FMaterialCompiler* Compiler, int32 OutputIndex) override;
 	virtual void GetCaption(TArray<FString>& OutCaptions) const override;
 	virtual bool MatchesSearchQuery(const TCHAR* SearchQuery) override;
diff --git a/Engine/Source/Runtime/Landscape/Classes/Materials/MaterialExpressionLandscapeLayerSwitch.h b/Engine/Source/Runtime/Landscape/Classes/Materials/MaterialExpressionLandscapeLayerSwitch.h
--- a/Engine/Source/Runtime/Landscape/Classes/Materials/MaterialExpressionLandscapeLayerSwitch.h
+++ b/Engine/Source/Runtime/Landscape/Classes/Materials/MaterialExpressionLandscapeLayerSwitch.h
@@ -40,6 +40,9 @@ public:
 #if WITH_EDITOR
 	virtual bool IsResultMaterialAttributes(int32 OutputIndex) override;
 	virtual int32 Compile(class FMaterialCompiler* Compiler, int32 OutputIndex) override;
+	virtual bool CanRenameNode() const override { return true; }
+	virtual FString GetEditableName() const override;
+	virtual void SetEditableName(const FString& NewName) override;
 	virtual void GetCaption(TArray<FString>& OutCaptions) const override;
 	virtual bool MatchesSearchQuery(const TCHAR* SearchQuery) override;
 #endif
diff --git a/Engine/Source/Runtime/Landscape/Classes/Materials/MaterialExpressionLandscapeLayerWeight.h b/Engine/Source/Runtime/Landscape/Classes/Materials/MaterialExpressionLandscapeLayerWeight.h
--- a/Engine/Source/Runtime/Landscape/Classes/Materials/MaterialExpressionLandscapeLayerWeight.h
+++ b/Engine/Source/Runtime/Landscape/Classes/Materials/MaterialExpressionLandscapeLayerWeight.h
@@ -40,6 +40,9 @@ public:
 #if WITH_EDITOR
 	virtual bool IsResultMaterialAttributes(int32 OutputIndex) override;
 	virtual int32 Compile(class FMaterialCompiler* Compiler, int32 OutputIndex) override;
+	virtual bool CanRenameNode() const override { return true; }
+	virtual FString GetEditableName() const override;
+	virtual void SetEditableName(const FString& NewName) override;
 	virtual void GetCaption(TArray<FString>& OutCaptions) const override;
 	virtual bool MatchesSearchQuery(const TCHAR* SearchQuery) override;
 #endif
diff --git a/Engine/Source/Runtime/Landscape/Private/LandscapeGrass.cpp b/Engine/Source/Runtime/Landscape/Private/LandscapeGrass.cpp
--- a/Engine/Source/Runtime/Landscape/Private/LandscapeGrass.cpp
+++ b/Engine/Source/Runtime/Landscape/Private/LandscapeGrass.cpp
@@ -1245,7 +1245,7 @@ int32 UMaterialExpressionLandscapeGrassOutput::Compile(class FMaterialCompiler*
 
 void UMaterialExpressionLandscapeGrassOutput::GetCaption(TArray<FString>& OutCaptions) const
 {
-	OutCaptions.Add(TEXT("Grass"));
+	OutCaptions.Add(TEXT("Landscape Grass"));
 }
 
 const TArray<FExpressionInput*> UMaterialExpressionLandscapeGrassOutput::GetInputs()
diff --git a/Engine/Source/Runtime/Landscape/Private/Materials/MaterialExpressionLandscapeLayerBlend.cpp b/Engine/Source/Runtime/Landscape/Private/Materials/MaterialExpressionLandscapeLayerBlend.cpp
--- a/Engine/Source/Runtime/Landscape/Private/Materials/MaterialExpressionLandscapeLayerBlend.cpp
+++ b/Engine/Source/Runtime/Landscape/Private/Materials/MaterialExpressionLandscapeLayerBlend.cpp
@@ -295,7 +295,7 @@ UObject* UMaterialExpressionLandscapeLayerBlend::GetReferencedTexture() const
 #if WITH_EDITOR
 void UMaterialExpressionLandscapeLayerBlend::GetCaption(TArray<FString>& OutCaptions) const
 {
-	OutCaptions.Add(FString(TEXT("Layer Blend")));
+	OutCaptions.Add(FString(TEXT("Landscape Layer Blend")));
 }
 
 void UMaterialExpressionLandscapeLayerBlend::PostEditChangeProperty(FPropertyChangedEvent& PropertyChangedEvent)
diff --git a/Engine/Source/Runtime/Landscape/Private/Materials/MaterialExpressionLandscapeLayerSample.cpp b/Engine/Source/Runtime/Landscape/Private/Materials/MaterialExpressionLandscapeLayerSample.cpp
--- a/Engine/Source/Runtime/Landscape/Private/Materials/MaterialExpressionLandscapeLayerSample.cpp
+++ b/Engine/Source/Runtime/Landscape/Private/Materials/MaterialExpressionLandscapeLayerSample.cpp
@@ -56,9 +56,20 @@ UObject* UMaterialExpressionLandscapeLayerSample::GetReferencedTexture() const
 }
 
 #if WITH_EDITOR
+FString UMaterialExpressionLandscapeLayerSample::GetEditableName() const
+{
+	return ParameterName.ToString();
+}
+
+void UMaterialExpressionLandscapeLayerSample::SetEditableName(const FString& NewName)
+{
+	ParameterName = *NewName;
+}
+
 void UMaterialExpressionLandscapeLayerSample::GetCaption(TArray<FString>& OutCaptions) const
 {
-	OutCaptions.Add(FString::Printf(TEXT("Sample '%s'"), *ParameterName.ToString()));
+	OutCaptions.Add(TEXT("Landscape Layer Sample"));
+	OutCaptions.Add(FString::Printf(TEXT("'%s'"), *ParameterName.ToString()));
 }
 
 bool UMaterialExpressionLandscapeLayerSample::MatchesSearchQuery(const TCHAR* SearchQuery)
diff --git a/Engine/Source/Runtime/Landscape/Private/Materials/MaterialExpressionLandscapeLayerSwitch.cpp b/Engine/Source/Runtime/Landscape/Private/Materials/MaterialExpressionLandscapeLayerSwitch.cpp
--- a/Engine/Source/Runtime/Landscape/Private/Materials/MaterialExpressionLandscapeLayerSwitch.cpp
+++ b/Engine/Source/Runtime/Landscape/Private/Materials/MaterialExpressionLandscapeLayerSwitch.cpp
@@ -86,9 +86,19 @@ UObject* UMaterialExpressionLandscapeLayerSwitch::GetReferencedTexture() const
 }
 
 #if WITH_EDITOR
+FString UMaterialExpressionLandscapeLayerSwitch::GetEditableName() const
+{
+	return ParameterName.ToString();
+}
+
+void UMaterialExpressionLandscapeLayerSwitch::SetEditableName(const FString& NewName)
+{
+	ParameterName = *NewName;
+}
+
 void UMaterialExpressionLandscapeLayerSwitch::GetCaption(TArray<FString>& OutCaptions) const
 {
-	OutCaptions.Add(TEXT("Layer Switch"));
+	OutCaptions.Add(TEXT("Landscape Layer Switch"));
 	OutCaptions.Add(FString::Printf(TEXT("'%s'"), *ParameterName.ToString()));
 }
 
diff --git a/Engine/Source/Runtime/Landscape/Private/Materials/MaterialExpressionLandscapeLayerWeight.cpp b/Engine/Source/Runtime/Landscape/Private/Materials/MaterialExpressionLandscapeLayerWeight.cpp
--- a/Engine/Source/Runtime/Landscape/Private/Materials/MaterialExpressionLandscapeLayerWeight.cpp
+++ b/Engine/Source/Runtime/Landscape/Private/Materials/MaterialExpressionLandscapeLayerWeight.cpp
@@ -96,9 +96,20 @@ UObject* UMaterialExpressionLandscapeLayerWeight::GetReferencedTexture() const
 }
 
 #if WITH_EDITOR
+FString UMaterialExpressionLandscapeLayerWeight::GetEditableName() const
+{
+	return ParameterName.ToString();
+}
+
+void UMaterialExpressionLandscapeLayerWeight::SetEditableName(const FString& NewName)
+{
+	ParameterName = *NewName;
+}
+
 void UMaterialExpressionLandscapeLayerWeight::GetCaption(TArray<FString>& OutCaptions) const
 {
-	OutCaptions.Add(FString::Printf(TEXT("Layer '%s'"), *ParameterName.ToString()));
+	OutCaptions.Add(TEXT("Landscape Layer Weight"));
+	OutCaptions.Add(FString::Printf(TEXT("'%s'"), *ParameterName.ToString()));
 }
 
 bool UMaterialExpressionLandscapeLayerWeight::MatchesSearchQuery(const TCHAR* SearchQuery)
