Fixed crash when force deleting a material which has (landscape) material instances : in that case, the UMaterialInstanceConstant's Parent field is force-serialized to null. Instead of crashing when rebuilding the primitive render proxy when this happens, we now fallback on the default material instead, like is done elsewhere.

diff --git a/Engine/Source/Runtime/Engine/Private/Materials/MaterialInterface.cpp b/Engine/Source/Runtime/Engine/Private/Materials/MaterialInterface.cpp
--- a/Engine/Source/Runtime/Engine/Private/Materials/MaterialInterface.cpp
+++ b/Engine/Source/Runtime/Engine/Private/Materials/MaterialInterface.cpp
@@ -363,12 +363,13 @@ FMaterialRelevance UMaterialInterface::GetRelevance_Internal(const UMaterial* Ma
 		const FMaterialResource* MaterialResource = GetMaterialResource(InFeatureLevel);
 
 		// If material is invalid e.g. unparented instance, fallback to the passed in material
-		if (!MaterialResource && Material)
+		bool bIsValidMaterialResource = (MaterialResource != nullptr) && (MaterialResource->GetMaterial() != nullptr);
+		if (!bIsValidMaterialResource && (Material != nullptr))
 		{
-			MaterialResource = Material->GetMaterialResource(InFeatureLevel);	
+			MaterialResource = Material->GetMaterialResource(InFeatureLevel);
 		}
 
-		if (!MaterialResource)
+		if (MaterialResource == nullptr)
 		{
 			return FMaterialRelevance();
 		}
diff --git a/Engine/Source/Runtime/Landscape/Private/LandscapeEditLayers.cpp b/Engine/Source/Runtime/Landscape/Private/LandscapeEditLayers.cpp
--- a/Engine/Source/Runtime/Landscape/Private/LandscapeEditLayers.cpp
+++ b/Engine/Source/Runtime/Landscape/Private/LandscapeEditLayers.cpp
@@ -7344,8 +7344,7 @@ void ALandscape::UpdateLayersMaterialInstances(const TArray<ULandscapeComponent*
 
 		Component->MaterialPerLOD = NewMaterialPerLOD;
 
-		Component->MaterialInstances.SetNumZeroed(Component->MaterialPerLOD.Num() * 2); // over allocate in case we are using tessellation
-		int8 TessellatedMaterialCount = 0;
+		Component->MaterialInstances.SetNumZeroed(Component->MaterialPerLOD.Num()); 
 		int8 MaterialIndex = 0;
 
 		const TArray<FWeightmapLayerAllocationInfo>& WeightmapBaseLayerAllocation = Component->GetWeightmapLayerAllocations();
diff --git a/Engine/Source/Runtime/Landscape/Private/LandscapeRender.cpp b/Engine/Source/Runtime/Landscape/Private/LandscapeRender.cpp
--- a/Engine/Source/Runtime/Landscape/Private/LandscapeRender.cpp
+++ b/Engine/Source/Runtime/Landscape/Private/LandscapeRender.cpp
@@ -49,7 +49,7 @@ LandscapeRender.cpp: New terrain rendering
 #include "NaniteSceneProxy.h"
 #include "Rendering/Texture2DResource.h"
 #include "RenderCore.h"
-
+#include "Algo/Transform.h"
 
 IMPLEMENT_GLOBAL_SHADER_PARAMETER_STRUCT(FLandscapeUniformShaderParameters, "LandscapeParameters");
 IMPLEMENT_GLOBAL_SHADER_PARAMETER_STRUCT(FLandscapeFixedGridUniformShaderParameters, "LandscapeFixedGrid");
@@ -958,18 +958,34 @@ FLandscapeComponentSceneProxy::FLandscapeComponentSceneProxy(ULandscapeComponent
 	const bool bHasStaticLighting = ComponentLightInfo->GetLightMap() || ComponentLightInfo->GetShadowMap();
 
 	check(AvailableMaterialInterfaces.Num() == AvailableMaterials.Num());
-	// Check material usage
+	// Check material usage and validity. Replace invalid entries by default material so that indexing AvailableMaterials with LODIndexToMaterialIndex still works :
 	if (ensure(AvailableMaterialInterfaces.Num() > 0))
 	{
 		for(int Index = 0; Index < AvailableMaterialInterfaces.Num(); ++Index)
 		{
-			UMaterialInterface*& MaterialInterface = AvailableMaterialInterfaces[Index];
-			FMaterialRenderProxy*& MaterialProxy = AvailableMaterials[Index];
-			if (MaterialInterface == nullptr ||
-				(bHasStaticLighting && !MaterialInterface->CheckMaterialUsage_Concurrent(MATUSAGE_StaticLighting)))
+			bool bIsValidMaterial = false;
+			UMaterialInterface* MaterialInterface = AvailableMaterialInterfaces[Index];
+			if (MaterialInterface != nullptr)
+			{
+				bIsValidMaterial = true;
+
+				const UMaterial* LandscapeMaterial = MaterialInterface->GetMaterial_Concurrent();
+
+				// In some case it's possible that the Material Instance we have and the Material are not related, for example, in case where content was force deleted, we can have a MIC with no parent, so GetMaterial will fallback to the default material.
+				// and since the MIC is not really valid, fallback to 
+				UMaterialInstance* MaterialInstance = Cast<UMaterialInstance>(MaterialInterface);
+				bIsValidMaterial &= (MaterialInstance == nullptr) || MaterialInstance->IsChildOf(LandscapeMaterial);
+
+				// Check usage flags : 
+				bIsValidMaterial &= !bHasStaticLighting || MaterialInterface->CheckMaterialUsage_Concurrent(MATUSAGE_StaticLighting);
+			}
+
+			if (!bIsValidMaterial)
 			{
+				// Replace the landscape material by the default material : 
 				MaterialInterface = UMaterial::GetDefaultMaterial(MD_Surface);
-				MaterialProxy = MaterialInterface->GetRenderProxy();
+				AvailableMaterialInterfaces[Index] = MaterialInterface;
+				AvailableMaterials[Index] = MaterialInterface->GetRenderProxy();
 			}
 		}
 	}
@@ -979,24 +995,7 @@ FLandscapeComponentSceneProxy::FLandscapeComponentSceneProxy(ULandscapeComponent
 		AvailableMaterials.Add(AvailableMaterialInterfaces.Last()->GetRenderProxy());
 	}
 
-	MaterialRelevances.Reserve(AvailableMaterials.Num());
-
-	for (UMaterialInterface*& MaterialInterface : AvailableMaterialInterfaces)
-	{
-		const UMaterial* LandscapeMaterial = MaterialInterface != nullptr ? MaterialInterface->GetMaterial_Concurrent() : nullptr;
-
-		if (LandscapeMaterial != nullptr)
-		{
-			UMaterialInstance* MaterialInstance = Cast<UMaterialInstance>(MaterialInterface);
-
-			// In some case it's possible that the Material Instance we have and the Material are not related, for example, in case where content was force deleted, we can have a MIC with no parent, so GetMaterial will fallback to the default material.
-			// and since the MIC is not really valid, dont generate the relevance.
-			if (MaterialInstance == nullptr || MaterialInstance->IsChildOf(LandscapeMaterial))
-			{
-				MaterialRelevances.Add(MaterialInterface->GetRelevance_Concurrent(FeatureLevel));
-			}
-		}
-	}
+	Algo::Transform(AvailableMaterialInterfaces, MaterialRelevances, [FeatureLevel](UMaterialInterface* InMaterialInterface) { check(InMaterialInterface != nullptr); return InMaterialInterface->GetRelevance_Concurrent(FeatureLevel); });
 
 #if !(UE_BUILD_SHIPPING || UE_BUILD_TEST) || (UE_BUILD_SHIPPING && WITH_EDITOR)
 	if (GIsEditor)
