Fixed landscape actors becoming dirty on load when using Nanite landscape

diff --git a/Engine/Source/Runtime/Landscape/Private/Landscape.cpp b/Engine/Source/Runtime/Landscape/Private/Landscape.cpp
--- a/Engine/Source/Runtime/Landscape/Private/Landscape.cpp
+++ b/Engine/Source/Runtime/Landscape/Private/Landscape.cpp
@@ -356,7 +356,8 @@ void ALandscapeProxy::InvalidateNaniteRepresentation(bool bInCheckContentId)
 	{
 		if (!bInCheckContentId || NaniteComponent->GetProxyContentId() != GetNaniteContentId())
 		{
-			NaniteComponent->DetachFromComponent(FDetachmentTransformRules::KeepRelativeTransform);
+			// Don't call modify when detaching the nanite component, this is non-transactional "derived data", regenerated any time the source landscape data changes. This prevents needlessly dirtying the package :
+			NaniteComponent->DetachFromComponent(FDetachmentTransformRules(EDetachmentRule::KeepRelative, /*bInCallModify = */false));
 			NaniteComponent->DestroyComponent();
 			NaniteComponent = nullptr;
 		}
