Fix for rare crashes -- it is not allowed to call IsDefaultTexture() or access other UTexture state from the render thread.

diff --git a/Engine/Source/Runtime/Landscape/Private/LandscapeRender.cpp b/Engine/Source/Runtime/Landscape/Private/LandscapeRender.cpp
--- a/Engine/Source/Runtime/Landscape/Private/LandscapeRender.cpp
+++ b/Engine/Source/Runtime/Landscape/Private/LandscapeRender.cpp
@@ -1596,14 +1596,15 @@ void FLandscapeComponentSceneProxy::OnTransformChanged()
 		0
 	);
 
-	if (HeightmapTexture)
+	FTextureResource* HeightmapResource = HeightmapTexture ? HeightmapTexture->GetResource() : nullptr;
+	if (HeightmapResource)
 	{
-		const float SizeX = FMath::Max(HeightmapTexture->GetSizeX(), 1);
-		const float SizeY = FMath::Max(HeightmapTexture->GetSizeY(), 1);
+		const float SizeX = FMath::Max(HeightmapResource->GetSizeX(), 1u);
+		const float SizeY = FMath::Max(HeightmapResource->GetSizeY(), 1u);
 		LandscapeParams.HeightmapTextureSize = FVector4f(SizeX, SizeY, 1.f / SizeX, 1.f / SizeY);
-		LandscapeParams.HeightmapTexture = HeightmapTexture->TextureReference.TextureReferenceRHI;
+		LandscapeParams.HeightmapTexture = HeightmapResource->TextureRHI.GetReference();
 		LandscapeParams.HeightmapTextureSampler = TStaticSamplerState<SF_Point>::GetRHI();
-		LandscapeParams.NormalmapTexture = HeightmapTexture->TextureReference.TextureReferenceRHI;
+		LandscapeParams.NormalmapTexture = HeightmapResource->TextureRHI.GetReference();
 		LandscapeParams.NormalmapTextureSampler = TStaticSamplerState<SF_Bilinear>::GetRHI();
 	}
 	else
@@ -1615,9 +1616,10 @@ void FLandscapeComponentSceneProxy::OnTransformChanged()
 		LandscapeParams.NormalmapTextureSampler = GBlackTexture->SamplerStateRHI;
 	}
 
-	if (XYOffsetmapTexture)
+	FTextureResource* XYOffsetmapResource = XYOffsetmapTexture ? XYOffsetmapTexture->GetResource() : nullptr;
+	if (XYOffsetmapResource)
 	{
-		LandscapeParams.XYOffsetmapTexture = XYOffsetmapTexture->TextureReference.TextureReferenceRHI;
+		LandscapeParams.XYOffsetmapTexture = XYOffsetmapResource->TextureRHI.GetReference();
 		LandscapeParams.XYOffsetmapTextureSampler = TStaticSamplerState<SF_Point>::GetRHI();
 	}
 	else
@@ -3592,22 +3594,30 @@ void FLandscapeComponentSceneProxy::ChangeComponentScreenSizeToUseSubSections_Re
 bool FLandscapeComponentSceneProxy::HeightfieldHasPendingStreaming() const
 {
 	bool bHeightmapTextureStreaming = false;
-
 	if (HeightmapTexture)
 	{
+		// this is technically a game thread value and not render-thread safe, but it shouldn't ever crash, may just be out of date.
+		// there doesn't appear to be any render thread equivalent, the render thread is ignorant of streaming state.
+		// in general, HeightfieldHasPendingStreaming() should only be used if the code is ok with a slightly out of date value being returned.
 		bHeightmapTextureStreaming |= HeightmapTexture->bHasStreamingUpdatePending;
 #if WITH_EDITOR
-		bHeightmapTextureStreaming |= HeightmapTexture->IsCompiling();
+		if (const FTexture2DResource* HeightmapTextureResource = (const FTexture2DResource*)HeightmapTexture->GetResource())
+		{
+			bHeightmapTextureStreaming |= HeightmapTextureResource->IsProxy();
+		}
 #endif
 	}
 
 	bool bVisibilityTextureStreaming = false;
-
 	if (VisibilityWeightmapTexture)
 	{
+		// again, not render thread safe (see above)
 		bVisibilityTextureStreaming |= VisibilityWeightmapTexture->bHasStreamingUpdatePending;
 #if WITH_EDITOR
-		bVisibilityTextureStreaming |= VisibilityWeightmapTexture->IsCompiling();
+		if (const FTexture2DResource* VisibilityTextureResource = (const FTexture2DResource*)VisibilityWeightmapTexture->GetResource())
+		{
+			bVisibilityTextureStreaming |= VisibilityTextureResource->IsProxy();
+		}
 #endif
 	}
 
