Guard against nullptr access in ALandscapeSplineActor::PostEditMove()

diff --git a/Engine/Source/Runtime/Landscape/Private/LandscapeSplineActor.cpp b/Engine/Source/Runtime/Landscape/Private/LandscapeSplineActor.cpp
--- a/Engine/Source/Runtime/Landscape/Private/LandscapeSplineActor.cpp
+++ b/Engine/Source/Runtime/Landscape/Private/LandscapeSplineActor.cpp
@@ -134,7 +134,10 @@ void ALandscapeSplineActor::PostEditMove(bool bFinished)
 
 	if (bFinished)
 	{
-		GetLandscapeInfo()->RequestSplineLayerUpdate();
+		if (ULandscapeInfo* LandscapeInfo = GetLandscapeInfo())
+		{
+			LandscapeInfo->RequestSplineLayerUpdate();
+		}
 	}
 }
 
