Multi-view:  Landscape ray tracing fix.  BLAS data for Landscape (FLandscapeRayTracingState) needs to be cached per view, to avoid LOD thrashing which causes the BLAS to be completely rebuilt every frame for each view.  This can cost 30-50 ms CPU per view, or 60-100 ms CPU for two views.  An array of per-view ray tracing state is stored in the scene proxy, in a newly added opaque implementation structure to allow moving structure definitions to the source file.  To handle deletion of ray tracing state when a view gets deleted, a linked list of ray tracing state is maintained in FSceneViewState, with an opaque type hiding the details there as well.

#rnx

diff --git a/Engine/Source/Runtime/Landscape/Private/LandscapeRender.cpp b/Engine/Source/Runtime/Landscape/Private/LandscapeRender.cpp
--- a/Engine/Source/Runtime/Landscape/Private/LandscapeRender.cpp
+++ b/Engine/Source/Runtime/Landscape/Private/LandscapeRender.cpp
@@ -38,6 +38,7 @@ LandscapeRender.cpp: New terrain rendering
 #include "PrimitiveSceneInfo.h"
 #include "SceneView.h"
 #include "SceneCore.h"
+#include "ScenePrivate.h"
 #include "LandscapeProxy.h"
 #include "HAL/LowLevelMemTracker.h"
 #include "MeshMaterialShader.h"
@@ -331,6 +332,95 @@ TMap<uint32, FLandscapeRenderSystem*> LandscapeRenderSystems;
 
 TBitArray<> FLandscapeRenderSystem::LandscapeIndexAllocator;
 
+#if RHI_RAYTRACING
+
+struct FLandscapeSectionRayTracingState
+{
+	int8 CurrentLOD;
+	float FractionalLOD;
+	float HeightmapLODBias;
+	uint32 ReferencedTextureRHIHash;
+
+	FRayTracingGeometry Geometry;
+	FRWBuffer RayTracingDynamicVertexBuffer;
+	FLandscapeVertexFactoryMVFUniformBufferRef UniformBuffer;
+
+	FLandscapeSectionRayTracingState()
+		: CurrentLOD(-1)
+		, FractionalLOD(-1000.0f)
+		, HeightmapLODBias(-1000.0f)
+		, ReferencedTextureRHIHash(0) {}
+};
+
+// Where we are rendering multiple views, we need to branch the landscape ray tracing state (BLAS data) per view, for performance reasons.
+// Without this, the BLAS data ends up getting rebuilt from scratch every frame due to LOD thrashing, costing 30-50 ms per view, or 60-100 ms
+// for a scene with two views.  This structure represents the state for a single view.
+struct FLandscapeRayTracingState
+{
+	FLandscapeRayTracingState()
+		: Pimpl(nullptr), ViewStateNext(nullptr), ViewStatePrev(nullptr), ViewKey(-1), NumSubsections(0) {}
+	~FLandscapeRayTracingState();
+
+	// Parent structure that holds this ray tracing state, needed for deletion of this item if its view is deleted
+	FLandscapeRayTracingImpl* Pimpl;
+
+	// Linked list pointers for FLandscapeRayTracingStateList, referenced from FSceneViewState, iterated over if view gets deleted
+	FLandscapeRayTracingState* ViewStateNext;
+	FLandscapeRayTracingState** ViewStatePrev;
+
+	// View state key from FSceneViewState.  Zero is used if FSceneViewState is null (view state keys start at 1, so 0 is invalid for an actual view).
+	uint32 ViewKey;
+
+	// Rendering data
+	int32 NumSubsections;
+	TStaticArray<FLandscapeSectionRayTracingState, FLandscapeComponentSceneProxy::MAX_SUBSECTION_COUNT> Sections;
+};
+
+// This wrapper holds the ray tracing state for a single scene proxy for all views
+struct FLandscapeRayTracingImpl
+{
+	// Needs to be indirect array, because elements are added to a linked list 
+	TIndirectArray<FLandscapeRayTracingState> PerViewRayTracingState;
+
+	// ViewStateInterface pointer can be NULL
+	FLandscapeRayTracingState* FindOrCreateRayTracingState(FSceneViewStateInterface* ViewStateInterface, int32 NumSubsections, int32 SubsectionSizeVerts);
+};
+
+// When views get deleted, we need to clean up the per view ray tracing data, which is handled by this class.  This is just a doubly linked list
+// of all the ray tracing data associated with a given view, pointed to by the FSceneViewState, with the destructor emptying the items from the list.
+class FLandscapeRayTracingStateList
+{
+public:
+	FLandscapeRayTracingState* ListHead;
+
+	FLandscapeRayTracingStateList() : ListHead(nullptr) {}
+
+	~FLandscapeRayTracingStateList()
+	{
+		// Pop items from the list head until the list is empty
+		while (ListHead)
+		{
+			// Find the item in its parent array
+			FLandscapeRayTracingState* ToRemove = ListHead;
+			FLandscapeRayTracingImpl* Pimpl = ToRemove->Pimpl;
+
+			for (int32 RemoveIndex = 0; RemoveIndex < Pimpl->PerViewRayTracingState.Num(); RemoveIndex++)
+			{
+				if (&Pimpl->PerViewRayTracingState[RemoveIndex] == ToRemove)
+				{
+					// Remove it -- the destructor will also unlink it from the list, updating ListHead
+					Pimpl->PerViewRayTracingState.RemoveAtSwap(RemoveIndex);
+					break;
+				}
+			}
+
+			// Make sure the item was successfully removed (ListHead updated)
+			check(ListHead != ToRemove);
+		}
+	}
+};
+#endif	// RHI_RAYTRACING
+
 
 //
 // FLandscapeRenderSystem
@@ -630,6 +720,18 @@ void FLandscapeSceneViewExtension::PreRenderView_RenderThread(FRDGBuilder& Graph
 {
 	LandscapeViews.Emplace(InView);
 
+#if RHI_RAYTRACING
+	if (InView.State)
+	{
+		// Create the ray tracing state list class if necessary
+		FSceneViewState* ViewState = InView.State->GetConcreteViewState();
+		if (!ViewState->LandscapeRayTracingStates.IsValid())
+		{
+			ViewState->LandscapeRayTracingStates = MakePimpl<FLandscapeRayTracingStateList>();
+		}
+	}
+#endif	// RHI_RAYTRACING
+
 	// Kick the job once all views have been collected.
 	if (!LandscapeRenderSystems.IsEmpty() && LandscapeViews.Num() == InView.Family->Views.Num())
 	{
@@ -1186,42 +1288,114 @@ void FLandscapeComponentSceneProxy::CreateRenderThreadResources()
 		}
 	}
 #endif
+}
 
 #if RHI_RAYTRACING
-	if (IsRayTracingAllowed())
+FLandscapeRayTracingState* FLandscapeRayTracingImpl::FindOrCreateRayTracingState(FSceneViewStateInterface* ViewStateInterface, int32 NumSubsections, int32 SubsectionSizeVerts)
+{
+	// Default view key of zero if there's no ViewStateInterface provided.  View keys start at 1, so 0 wouldn't be a valid key on an actual view.
+	uint32 ViewKey = 0;
+	FSceneViewState* ViewState = nullptr;
+	if (ViewStateInterface)
 	{
-		for (int32 SubY = 0; SubY < NumSubsections; SubY++)
+		ViewState = ViewStateInterface->GetConcreteViewState();
+		ViewKey = ViewState->UniqueID;
+	}
+
+	// Check for existing state for this view.  We're just doing a linear search of the array, because practical applications won't have
+	// more than two or three views running ray tracing, for overall frame performance reasons.  If this assumption changes, we could
+	// implement a more efficient lookup in the future.
+	for (FLandscapeRayTracingState& PerView : PerViewRayTracingState)
+	{
+		if (PerView.ViewKey == ViewKey)
 		{
-			for (int32 SubX = 0; SubX < NumSubsections; SubX++)
-			{
-				const int8 SubSectionIdx = SubX + SubY * NumSubsections;
-
-				FRayTracingGeometryInitializer Initializer;
-				static const FName DebugName("FLandscapeComponentSceneProxy");
-				static int32 DebugNumber = 0;
-				Initializer.DebugName = FDebugName(DebugName, DebugNumber++);
-				Initializer.IndexBuffer = nullptr;
-				Initializer.GeometryType = RTGT_Triangles;
-				Initializer.bFastBuild = true;
-				Initializer.bAllowUpdate = true;
-				FRayTracingGeometrySegment Segment;
-				Segment.VertexBuffer = nullptr;
-				Segment.VertexBufferStride = sizeof(FVector3f);
-				Segment.VertexBufferElementType = VET_Float3;
-				Segment.MaxVertices = FMath::Square(SubsectionSizeVerts);
-				Initializer.Segments.Add(Segment);
-				SectionRayTracingStates[SubSectionIdx].Geometry.SetInitializer(Initializer);
-				SectionRayTracingStates[SubSectionIdx].Geometry.InitResource();
-
-				FLandscapeVertexFactoryMVFParameters UniformBufferParams;
-				UniformBufferParams.SubXY = FIntPoint(SubX, SubY);
-				SectionRayTracingStates[SubSectionIdx].UniformBuffer = FLandscapeVertexFactoryMVFUniformBufferRef::CreateUniformBufferImmediate(UniformBufferParams, UniformBuffer_MultiFrame);
-			}
+			return &PerView;
 		}
 	}
-#endif
+
+	// Need to create a new one
+	FLandscapeRayTracingState* RayTracingState = new FLandscapeRayTracingState();
+
+	PerViewRayTracingState.Add(RayTracingState);
+
+	RayTracingState->Pimpl = this;
+	RayTracingState->ViewKey = ViewKey;
+
+	if (ViewState)
+	{
+		// Link into the view state's linked list, so it can be cleaned up if the view gets deleted
+		FLandscapeRayTracingStateList* StateList = ViewState->LandscapeRayTracingStates.Get();
+		check(StateList);
+
+		if (StateList->ListHead)
+		{
+			StateList->ListHead->ViewStatePrev = &RayTracingState->ViewStateNext;
+		}
+		RayTracingState->ViewStateNext = StateList->ListHead;
+		RayTracingState->ViewStatePrev = &StateList->ListHead;
+		StateList->ListHead = RayTracingState;
+	}
+
+	// Initialize rendering data
+	RayTracingState->NumSubsections = NumSubsections;
+
+	for (int32 SubY = 0; SubY < NumSubsections; SubY++)
+	{
+		for (int32 SubX = 0; SubX < NumSubsections; SubX++)
+		{
+			const int8 SubSectionIdx = SubX + SubY * NumSubsections;
+
+			FRayTracingGeometryInitializer Initializer;
+			static const FName DebugName("FLandscapeComponentSceneProxy");
+			static int32 DebugNumber = 0;
+			Initializer.DebugName = FDebugName(DebugName, DebugNumber++);
+			Initializer.IndexBuffer = nullptr;
+			Initializer.GeometryType = RTGT_Triangles;
+			Initializer.bFastBuild = true;
+			Initializer.bAllowUpdate = true;
+			FRayTracingGeometrySegment Segment;
+			Segment.VertexBuffer = nullptr;
+			Segment.VertexBufferStride = sizeof(FVector3f);
+			Segment.VertexBufferElementType = VET_Float3;
+			Segment.MaxVertices = FMath::Square(SubsectionSizeVerts);
+			Initializer.Segments.Add(Segment);
+			RayTracingState->Sections[SubSectionIdx].Geometry.SetInitializer(Initializer);
+			RayTracingState->Sections[SubSectionIdx].Geometry.InitResource();
+
+			FLandscapeVertexFactoryMVFParameters UniformBufferParams;
+			UniformBufferParams.SubXY = FIntPoint(SubX, SubY);
+			RayTracingState->Sections[SubSectionIdx].UniformBuffer = FLandscapeVertexFactoryMVFUniformBufferRef::CreateUniformBufferImmediate(UniformBufferParams, UniformBuffer_MultiFrame);
+		}
+	}
+
+	return RayTracingState;
 }
 
+FLandscapeRayTracingState::~FLandscapeRayTracingState()
+{
+	// Unlink this from the view state linked list
+	if (ViewStatePrev)
+	{
+		(*ViewStatePrev) = ViewStateNext;
+	}
+	if (ViewStateNext)
+	{
+		ViewStateNext->ViewStatePrev = ViewStatePrev;
+	}
+
+	// And clean up the contents
+	for (int32 SubY = 0; SubY < NumSubsections; SubY++)
+	{
+		for (int32 SubX = 0; SubX < NumSubsections; SubX++)
+		{
+			const int8 SubSectionIdx = SubX + SubY * NumSubsections;
+			Sections[SubSectionIdx].Geometry.ReleaseResource();
+			Sections[SubSectionIdx].RayTracingDynamicVertexBuffer.Release();
+		}
+	}
+}
+#endif	// RHI_RAYTRACING
+
 void FLandscapeComponentSceneProxy::DestroyRenderThreadResources()
 {
 	FPrimitiveSceneProxy::DestroyRenderThreadResources();
@@ -1269,18 +1443,6 @@ FLandscapeComponentSceneProxy::~FLandscapeComponentSceneProxy()
 		}
 		SharedBuffers = nullptr;
 	}
-
-#if RHI_RAYTRACING
-	for (int32 SubY = 0; SubY < NumSubsections; SubY++)
-	{
-		for (int32 SubX = 0; SubX < NumSubsections; SubX++)
-		{
-			const int8 SubSectionIdx = SubX + SubY * NumSubsections;
-			SectionRayTracingStates[SubSectionIdx].Geometry.ReleaseResource();
-			SectionRayTracingStates[SubSectionIdx].RayTracingDynamicVertexBuffer.Release();
-		}
-	}
-#endif
 }
 
 bool FLandscapeComponentSceneProxy::CanBeOccluded() const
@@ -2254,6 +2416,12 @@ void FLandscapeComponentSceneProxy::GetDynamicRayTracingInstances(FRayTracingMat
 	const FSceneView& SceneView = *Context.ReferenceView;
 	const FLandscapeRenderSystem& RenderSystem = *LandscapeRenderSystems.FindChecked(LandscapeKey);
 
+	if (!RayTracingImpl.IsValid())
+	{
+		RayTracingImpl = MakePimpl<FLandscapeRayTracingImpl>();
+	}
+	FLandscapeRayTracingState* RayTracingState = RayTracingImpl.Get()->FindOrCreateRayTracingState(SceneView.State, NumSubsections, SubsectionSizeVerts);
+
 	int32 LODToRender = RenderSystem.GetSectionLODValue(SceneView, ComponentBase);
 
 	FLandscapeElementParamArray& ParameterArray = Context.RayTracingMeshResourceCollector.AllocateOneFrameResource<FLandscapeElementParamArray>();
@@ -2322,9 +2490,9 @@ void FLandscapeComponentSceneProxy::GetDynamicRayTracingInstances(FRayTracingMat
 
 			MeshBatch.Elements.Add(BatchElement);
 
-			SectionRayTracingStates[SubSectionIdx].Geometry.Initializer.IndexBuffer = BatchElement.IndexBuffer->IndexBufferRHI;
+			RayTracingState->Sections[SubSectionIdx].Geometry.Initializer.IndexBuffer = BatchElement.IndexBuffer->IndexBufferRHI;
 
-			BatchElementParams.LandscapeVertexFactoryMVFUniformBuffer = SectionRayTracingStates[SubSectionIdx].UniformBuffer;
+			BatchElementParams.LandscapeVertexFactoryMVFUniformBuffer = RayTracingState->Sections[SubSectionIdx].UniformBuffer;
 
 			bool bNeedsRayTracingGeometryUpdate = false;
 
@@ -2333,22 +2501,22 @@ void FLandscapeComponentSceneProxy::GetDynamicRayTracingInstances(FRayTracingMat
 
 			// Detect continuous LOD parameter changes. This is for far-away high LODs - they change rarely yet the BLAS refit time is not ideal, even if they contains tiny amount of triangles
 			{
-				if (SectionRayTracingStates[SubSectionIdx].CurrentLOD != CurrentLOD)
+				if (RayTracingState->Sections[SubSectionIdx].CurrentLOD != CurrentLOD)
 				{
 					bNeedsRayTracingGeometryUpdate = true;
-					SectionRayTracingStates[SubSectionIdx].CurrentLOD = CurrentLOD;
-					SectionRayTracingStates[SubSectionIdx].RayTracingDynamicVertexBuffer.Release();
+					RayTracingState->Sections[SubSectionIdx].CurrentLOD = CurrentLOD;
+					RayTracingState->Sections[SubSectionIdx].RayTracingDynamicVertexBuffer.Release();
 				}
-				if (SectionRayTracingStates[SubSectionIdx].HeightmapLODBias != RenderSystem.GetSectionLODBias(ComponentBase))
+				if (RayTracingState->Sections[SubSectionIdx].HeightmapLODBias != RenderSystem.GetSectionLODBias(ComponentBase))
 				{
 					bNeedsRayTracingGeometryUpdate = true;
-					SectionRayTracingStates[SubSectionIdx].HeightmapLODBias = RenderSystem.GetSectionLODBias(ComponentBase);
+					RayTracingState->Sections[SubSectionIdx].HeightmapLODBias = RenderSystem.GetSectionLODBias(ComponentBase);
 				}
 
-				if (SectionRayTracingStates[SubSectionIdx].FractionalLOD != RenderSystem.GetSectionLODValue(SceneView, ComponentBase))
+				if (RayTracingState->Sections[SubSectionIdx].FractionalLOD != RenderSystem.GetSectionLODValue(SceneView, ComponentBase))
 				{
 					bNeedsRayTracingGeometryUpdate = true;
-					SectionRayTracingStates[SubSectionIdx].FractionalLOD = RenderSystem.GetSectionLODValue(SceneView, ComponentBase);
+					RayTracingState->Sections[SubSectionIdx].FractionalLOD = RenderSystem.GetSectionLODValue(SceneView, ComponentBase);
 				}
 			}
 
@@ -2366,16 +2534,16 @@ void FLandscapeComponentSceneProxy::GetDynamicRayTracingInstances(FRayTracingMat
 					const FUniformExpressionSet& UniformExpressionSet = Material.GetRenderingThreadShaderMap()->GetUniformExpressionSet();
 					const uint32 Hash = UniformExpressionSet.GetReferencedTexture2DRHIHash(MaterialRenderContext);
 
-					if (SectionRayTracingStates[SubSectionIdx].ReferencedTextureRHIHash != Hash)
+					if (RayTracingState->Sections[SubSectionIdx].ReferencedTextureRHIHash != Hash)
 					{
 						bNeedsRayTracingGeometryUpdate = true;
-						SectionRayTracingStates[SubSectionIdx].ReferencedTextureRHIHash = Hash;
+						RayTracingState->Sections[SubSectionIdx].ReferencedTextureRHIHash = Hash;
 					}
 				}
 			}
 
 			FRayTracingInstance RayTracingInstance;
-			RayTracingInstance.Geometry = &SectionRayTracingStates[SubSectionIdx].Geometry;
+			RayTracingInstance.Geometry = &RayTracingState->Sections[SubSectionIdx].Geometry;
 			RayTracingInstance.InstanceTransforms.Add(GetLocalToWorld());
 			RayTracingInstance.Materials.Add(MeshBatch);
 			OutRayTracingInstances.Add(RayTracingInstance);
@@ -2393,8 +2561,8 @@ void FLandscapeComponentSceneProxy::GetDynamicRayTracingInstances(FRayTracingMat
 						(uint32)FMath::Square(LodSubsectionSizeVerts),
 						FMath::Square(LodSubsectionSizeVerts) * (uint32)sizeof(FVector3f),
 						(uint32)FMath::Square(LodSubsectionSizeVerts - 1) * 2,
-						&SectionRayTracingStates[SubSectionIdx].Geometry,
-						&SectionRayTracingStates[SubSectionIdx].RayTracingDynamicVertexBuffer,
+						&RayTracingState->Sections[SubSectionIdx].Geometry,
+						&RayTracingState->Sections[SubSectionIdx].RayTracingDynamicVertexBuffer,
 						true
 					}
 				);
diff --git a/Engine/Source/Runtime/Landscape/Public/LandscapeRender.h b/Engine/Source/Runtime/Landscape/Public/LandscapeRender.h
--- a/Engine/Source/Runtime/Landscape/Public/LandscapeRender.h
+++ b/Engine/Source/Runtime/Landscape/Public/LandscapeRender.h
@@ -39,6 +39,10 @@ LandscapeRender.h: New terrain rendering
 class FLandscapeComponentSceneProxy;
 enum class ERuntimeVirtualTextureMaterialType : uint8;
 
+#if RHI_RAYTRACING
+struct FLandscapeRayTracingImpl;
+#endif
+
 #if WITH_EDITOR
 namespace ELandscapeViewMode
 {
@@ -638,25 +642,7 @@ public:
 	static const int8 MAX_SUBSECTION_COUNT = 2*2;
 
 #if RHI_RAYTRACING
-	struct FLandscapeSectionRayTracingState
-	{
-		int8 CurrentLOD;
-		float FractionalLOD;
-		float HeightmapLODBias;
-		uint32 ReferencedTextureRHIHash;
-
-		FRayTracingGeometry Geometry;
-		FRWBuffer RayTracingDynamicVertexBuffer;
-		FLandscapeVertexFactoryMVFUniformBufferRef UniformBuffer;
-
-		FLandscapeSectionRayTracingState() 
-			: CurrentLOD(-1)
-			, FractionalLOD(-1000.0f)
-			, HeightmapLODBias(-1000.0f)
-			, ReferencedTextureRHIHash(0) {}
-	};
-
-	TStaticArray<FLandscapeSectionRayTracingState, MAX_SUBSECTION_COUNT> SectionRayTracingStates;
+	TPimplPtr<FLandscapeRayTracingImpl> RayTracingImpl;
 #endif
 
 	friend FLandscapeRenderSystem;
diff --git a/Engine/Source/Runtime/Renderer/Private/ScenePrivate.h b/Engine/Source/Runtime/Renderer/Private/ScenePrivate.h
--- a/Engine/Source/Runtime/Renderer/Private/ScenePrivate.h
+++ b/Engine/Source/Runtime/Renderer/Private/ScenePrivate.h
@@ -94,6 +94,7 @@ class FVirtualShadowMapArrayCacheManager;
 struct FHairStrandsInstance;
 struct FPathTracingState;
 class FSparseVolumeTextureViewerSceneProxy;
+class FLandscapeRayTracingStateList;
 
 /** Holds information about a single primitive's occlusion. */
 class FPrimitiveOcclusionHistory
@@ -997,6 +998,9 @@ public:
 	// Last valid RTPSO is saved, so it could be used as fallback in future frames if background PSO compilation is enabled.
 	// This RTPSO can be used only if the only difference from previous PSO is the material hit shaders.
 	FRayTracingPipelineStateSignature LastRayTracingMaterialPipelineSignature;
+
+	// List of landscape ray tracing state associated with this view, so it can be cleaned up if the view gets deleted.
+	TPimplPtr<FLandscapeRayTracingStateList> LandscapeRayTracingStates;
 #endif
 
 	TUniquePtr<FForwardLightingViewResources> ForwardLightingResources;
