Landscape edit layers :
* Fixed missing barrier
* Added some GPU tags to improve readability of render captures

#rnx

diff --git a/Engine/Source/Runtime/Landscape/Private/LandscapeEditLayers.cpp b/Engine/Source/Runtime/Landscape/Private/LandscapeEditLayers.cpp
--- a/Engine/Source/Runtime/Landscape/Private/LandscapeEditLayers.cpp
+++ b/Engine/Source/Runtime/Landscape/Private/LandscapeEditLayers.cpp
@@ -884,6 +884,9 @@ public:
 		SetUAVParameter(RHICmdList, RHICmdList.GetBoundComputeShader(), AtlasPaintListsParam, InParams.AtlasWeightmapsPerLayer->GetTextureUAV(/*InMipLevel = */0));
 		SetSRVParameter(RHICmdList, RHICmdList.GetBoundComputeShader(), ComponentsDataParam, InParams.ComputeShaderResource->ComponentsDataSRV);
 		SetShaderValue(RHICmdList, RHICmdList.GetBoundComputeShader(), ComponentSizeParam, InParams.ComponentSize);
+		
+		// In case the CS is executed twice in a row, we need a barrier since we want to prevent UAV overlaps:
+		RHICmdList.Transition(FRHITransitionInfo(InParams.AtlasWeightmapsPerLayer->TextureRHI, ERHIAccess::UAVMask, ERHIAccess::UAVMask));
 	}
 
 	void UnsetParameters(FRHICommandList& RHICmdList)
@@ -2477,7 +2480,8 @@ void ExecuteCopyLayersTexture(TArray<FLandscapeLayersCopyTextureParams>&& InCopy
 	ENQUEUE_RENDER_COMMAND(LandscapeLayers_Cmd_CopyTexture)(
 		[CopyTextureParams = MoveTemp(InCopyTextureParams)](FRHICommandListImmediate& RHICmdList) mutable
 	{
-		TRACE_CPUPROFILER_EVENT_SCOPE(LandscapeLayers_RT_CopyTexture);
+		TRACE_CPUPROFILER_EVENT_SCOPE(LandscapeLayers_RT_CopyTextures);
+		SCOPED_DRAW_EVENTF(RHICmdList, LandscapeLayers, TEXT("LandscapeLayers : Copy %d texture regions"), CopyTextureParams.Num());
 
 		for (const FLandscapeLayersCopyTextureParams& Params : CopyTextureParams)
 		{
@@ -2522,6 +2526,8 @@ struct FLandscapeLayersCopyReadbackTextureParams
 
 void ExecuteCopyToReadbackTexture(TArray<FLandscapeLayersCopyReadbackTextureParams>& InParams)
 {
+	SCOPED_DRAW_EVENTF_GAMETHREAD(LandscapeLayers, TEXT("Copy to readback textures (%d copies)"), InParams.Num());
+
 	for (FLandscapeLayersCopyReadbackTextureParams& Params : InParams)
 	{
 		Params.Dest->Enqueue(Params.Source, MoveTemp(Params.Context));
