Adding missing include.

diff --git a/Engine/Source/Runtime/Landscape/Private/LandscapeEditLayers.cpp b/Engine/Source/Runtime/Landscape/Private/LandscapeEditLayers.cpp
--- a/Engine/Source/Runtime/Landscape/Private/LandscapeEditLayers.cpp
+++ b/Engine/Source/Runtime/Landscape/Private/LandscapeEditLayers.cpp
@@ -30,6 +30,7 @@ LandscapeEditLayers.cpp: Landscape editing layers mode
 #include "PixelShaderUtils.h"
 #include "SystemTextures.h"
 #include "Rendering/Texture2DResource.h"
+#include "SceneView.h"
 #include "MaterialCachedData.h"
 
 #if WITH_EDITOR
