Removes TLazyObjectPtr usage in ALandscapeStreamingProxy, ULandscapeInfo, ULandscapeComponent, ULandscapeHeightfieldCollisionComponent.
Adds package deprecation and fixup mechanism to allow users to update fixedup packages.
Fixes Guid Remapping inTLazyObjectPtr asserting when multiple PIE sessions

diff --git a/Engine/Source/Editor/FoliageEdit/Private/FoliageEdMode.cpp b/Engine/Source/Editor/FoliageEdit/Private/FoliageEdMode.cpp
--- a/Engine/Source/Editor/FoliageEdit/Private/FoliageEdMode.cpp
+++ b/Engine/Source/Editor/FoliageEdit/Private/FoliageEdMode.cpp
@@ -1161,7 +1161,7 @@ bool GetMaxHitWeight(const FVector& Location, UActorComponent* Component, const
 	float MaxHitWeight = 0.f;
 	if (ULandscapeHeightfieldCollisionComponent* HitLandscapeCollision = Cast<ULandscapeHeightfieldCollisionComponent>(Component))
 	{
-		if (ULandscapeComponent* HitLandscape = HitLandscapeCollision->RenderComponent.Get())
+		if (ULandscapeComponent* HitLandscape = HitLandscapeCollision->GetRenderComponent())
 		{
 			for (const FName& LandscapeLayerName : LandscapeLayersArray)
 			{
diff --git a/Engine/Source/Editor/LandscapeEditor/Private/LandscapeEdMode.cpp b/Engine/Source/Editor/LandscapeEditor/Private/LandscapeEdMode.cpp
--- a/Engine/Source/Editor/LandscapeEditor/Private/LandscapeEdMode.cpp
+++ b/Engine/Source/Editor/LandscapeEditor/Private/LandscapeEdMode.cpp
@@ -3260,7 +3260,7 @@ void FEdModeLandscape::DeleteLandscapeComponents(ULandscapeInfo* LandscapeInfo,
 	for (ULandscapeComponent* Component : ComponentsToDelete)
 	{
 		Component->Modify();
-		ULandscapeHeightfieldCollisionComponent* CollisionComp = Component->CollisionComponent.Get();
+		ULandscapeHeightfieldCollisionComponent* CollisionComp = Component->GetCollisionComponent();
 		if (CollisionComp)
 		{
 			CollisionComp->Modify();
@@ -3319,7 +3319,7 @@ void FEdModeLandscape::DeleteLandscapeComponents(ULandscapeInfo* LandscapeInfo,
 	// Remove attached foliage
 	for (ULandscapeComponent* Component : ComponentsToDelete)
 	{
-		ULandscapeHeightfieldCollisionComponent* CollisionComp = Component->CollisionComponent.Get();
+		ULandscapeHeightfieldCollisionComponent* CollisionComp = Component->GetCollisionComponent();
 		if (CollisionComp)
 		{
 			AInstancedFoliageActor::DeleteInstancesForComponent(Proxy->GetWorld(), CollisionComp);
@@ -3427,7 +3427,7 @@ void FEdModeLandscape::DeleteLandscapeComponents(ULandscapeInfo* LandscapeInfo,
 			Component->XYOffsetmapTexture->ClearFlags(RF_Standalone);
 		}
 
-		ULandscapeHeightfieldCollisionComponent* CollisionComp = Component->CollisionComponent.Get();
+		ULandscapeHeightfieldCollisionComponent* CollisionComp = Component->GetCollisionComponent();
 		if (CollisionComp)
 		{
 			CollisionComp->DestroyComponent();
@@ -3725,8 +3725,8 @@ ALandscape* FEdModeLandscape::ChangeComponentSetting(int32 NumComponentsX, int32
 					ULandscapeComponent* NewComponent = NewLandscapeInfo->XYtoComponentMap.FindRef(Entry.Key);
 					if (NewComponent)
 					{
-						ULandscapeHeightfieldCollisionComponent* OldCollisionComponent = Entry.Value->CollisionComponent.Get();
-						ULandscapeHeightfieldCollisionComponent* NewCollisionComponent = NewComponent->CollisionComponent.Get();
+						ULandscapeHeightfieldCollisionComponent* OldCollisionComponent = Entry.Value->GetCollisionComponent();
+						ULandscapeHeightfieldCollisionComponent* NewCollisionComponent = NewComponent->GetCollisionComponent();
 
 						if (OldCollisionComponent && NewCollisionComponent)
 						{
@@ -3757,7 +3757,7 @@ ALandscape* FEdModeLandscape::ChangeComponentSetting(int32 NumComponentsX, int32
 				// Move instances
 				for (const TPair<FIntPoint, ULandscapeComponent*>& OldEntry : LandscapeInfo->XYtoComponentMap)
 				{
-					ULandscapeHeightfieldCollisionComponent* OldCollisionComponent = OldEntry.Value->CollisionComponent.Get();
+					ULandscapeHeightfieldCollisionComponent* OldCollisionComponent = OldEntry.Value->GetCollisionComponent();
 
 					if (OldCollisionComponent)
 					{
@@ -3765,7 +3765,7 @@ ALandscape* FEdModeLandscape::ChangeComponentSetting(int32 NumComponentsX, int32
 
 						for (const TPair<FIntPoint, ULandscapeComponent*>& NewEntry : NewLandscapeInfo->XYtoComponentMap)
 						{
-							ULandscapeHeightfieldCollisionComponent* NewCollisionComponent = NewEntry.Value->CollisionComponent.Get();
+							ULandscapeHeightfieldCollisionComponent* NewCollisionComponent = NewEntry.Value->GetCollisionComponent();
 
 							if (NewCollisionComponent && FBoxSphereBounds::BoxesIntersect(NewCollisionComponent->Bounds, OldCollisionComponent->Bounds))
 							{
@@ -3782,7 +3782,7 @@ ALandscape* FEdModeLandscape::ChangeComponentSetting(int32 NumComponentsX, int32
 				// Snap them to the bounds
 				for (const TPair<FIntPoint, ULandscapeComponent*>& NewEntry : NewLandscapeInfo->XYtoComponentMap)
 				{
-					ULandscapeHeightfieldCollisionComponent* NewCollisionComponent = NewEntry.Value->CollisionComponent.Get();
+					ULandscapeHeightfieldCollisionComponent* NewCollisionComponent = NewEntry.Value->GetCollisionComponent();
 
 					if (NewCollisionComponent)
 					{
@@ -3818,7 +3818,7 @@ ALandscape* FEdModeLandscape::ChangeComponentSetting(int32 NumComponentsX, int32
 			// Delete the old Landscape and all its proxies
 			for (ALandscapeStreamingProxy* Proxy : TActorRange<ALandscapeStreamingProxy>(OldLandscape->GetWorld()))
 			{
-				if (Proxy->LandscapeActor == OldLandscapeActor)
+				if (Proxy->GetLandscapeActor() == OldLandscapeActor)
 				{
 					Proxy->Destroy();
 				}
diff --git a/Engine/Source/Editor/LandscapeEditor/Private/LandscapeEdModeComponentTools.cpp b/Engine/Source/Editor/LandscapeEditor/Private/LandscapeEdModeComponentTools.cpp
--- a/Engine/Source/Editor/LandscapeEditor/Private/LandscapeEdModeComponentTools.cpp
+++ b/Engine/Source/Editor/LandscapeEditor/Private/LandscapeEdModeComponentTools.cpp
@@ -652,7 +652,7 @@ public:
 
 				if (!bHasLandscapeLayersContent)
 				{
-					ULandscapeHeightfieldCollisionComponent* CollisionComp = NewComponent->CollisionComponent.Get();
+					ULandscapeHeightfieldCollisionComponent* CollisionComp = NewComponent->GetCollisionComponent();
 					if (CollisionComp && !bHasXYOffset)
 					{
 						CollisionComp->MarkRenderStateDirty();
diff --git a/Engine/Source/Editor/LandscapeEditor/Private/LandscapeEdModeMirrorTool.cpp b/Engine/Source/Editor/LandscapeEditor/Private/LandscapeEdModeMirrorTool.cpp
--- a/Engine/Source/Editor/LandscapeEditor/Private/LandscapeEdModeMirrorTool.cpp
+++ b/Engine/Source/Editor/LandscapeEditor/Private/LandscapeEdModeMirrorTool.cpp
@@ -654,7 +654,7 @@ public:
 				for (ULandscapeComponent* Component : Components)
 				{
 					// Recreate collision for modified components and update the navmesh
-					ULandscapeHeightfieldCollisionComponent* CollisionComponent = Component->CollisionComponent.Get();
+					ULandscapeHeightfieldCollisionComponent* CollisionComponent = Component->GetCollisionComponent();
 					if (CollisionComponent)
 					{
 						CollisionComponent->RecreateCollision();
diff --git a/Engine/Source/Editor/LandscapeEditor/Private/LandscapeEdModeRampTool.cpp b/Engine/Source/Editor/LandscapeEditor/Private/LandscapeEdModeRampTool.cpp
--- a/Engine/Source/Editor/LandscapeEditor/Private/LandscapeEdModeRampTool.cpp
+++ b/Engine/Source/Editor/LandscapeEditor/Private/LandscapeEdModeRampTool.cpp
@@ -598,7 +598,7 @@ public:
 					for (ULandscapeComponent* Component : Components)
 					{
 						// Recreate collision for modified components and update the navmesh
-						ULandscapeHeightfieldCollisionComponent* CollisionComponent = Component->CollisionComponent.Get();
+						ULandscapeHeightfieldCollisionComponent* CollisionComponent = Component->GetCollisionComponent();
 						if (CollisionComponent)
 						{
 							CollisionComponent->RecreateCollision();
diff --git a/Engine/Source/Editor/LandscapeEditor/Private/LandscapeEdModeTools.h b/Engine/Source/Editor/LandscapeEditor/Private/LandscapeEdModeTools.h
--- a/Engine/Source/Editor/LandscapeEditor/Private/LandscapeEdModeTools.h
+++ b/Engine/Source/Editor/LandscapeEditor/Private/LandscapeEdModeTools.h
@@ -782,7 +782,7 @@ struct FXYOffsetmapAccessor
 				bUpdateNormals = true;
 				for (ULandscapeComponent* Component : Components)
 				{
-					ULandscapeHeightfieldCollisionComponent* CollisionComponent = Component->CollisionComponent.Get();
+					ULandscapeHeightfieldCollisionComponent* CollisionComponent = Component->GetCollisionComponent();
 					if (CollisionComponent && AInstancedFoliageActor::HasFoliageAttached(CollisionComponent))
 					{
 						bUpdateFoliage = true;
@@ -801,7 +801,7 @@ struct FXYOffsetmapAccessor
 
 				for (ULandscapeComponent* Component : Components)
 				{
-					CollisionComponents.Add(Component->CollisionComponent.Get());
+					CollisionComponents.Add(Component->GetCollisionComponent());
 					PreUpdateLocalBoxes.Add(FBox(FVector((float)X1, (float)Y1, Component->CachedLocalBox.Min.Z), FVector((float)X2, (float)Y2, Component->CachedLocalBox.Max.Z)));
 				}
 
@@ -1116,7 +1116,7 @@ struct FFullWeightmapAccessor
 			// Recreate collision for modified components to update the physical materials
 			for (ULandscapeComponent* Component : ModifiedComponents)
 			{
-				ULandscapeHeightfieldCollisionComponent* CollisionComponent = Component->CollisionComponent.Get();
+				ULandscapeHeightfieldCollisionComponent* CollisionComponent = Component->GetCollisionComponent();
 				if (CollisionComponent)
 				{
 					CollisionComponent->RecreateCollision();
diff --git a/Engine/Source/Editor/UnrealEd/Private/Editor/ActorPositioning.cpp b/Engine/Source/Editor/UnrealEd/Private/Editor/ActorPositioning.cpp
--- a/Engine/Source/Editor/UnrealEd/Private/Editor/ActorPositioning.cpp
+++ b/Engine/Source/Editor/UnrealEd/Private/Editor/ActorPositioning.cpp
@@ -70,7 +70,7 @@ bool IsHitIgnored(const FHitResult& InHit, const FSceneView& InSceneView)
 	}
 	if (PrimitiveComponent && PrimitiveComponent->IsA(ULandscapeHeightfieldCollisionComponent::StaticClass()))
 	{
-		PrimitiveComponent = CastChecked<ULandscapeHeightfieldCollisionComponent>(PrimitiveComponent)->RenderComponent.Get();
+		PrimitiveComponent = CastChecked<ULandscapeHeightfieldCollisionComponent>(PrimitiveComponent)->GetRenderComponent();
 	}
 
 	if (InHit.bStartPenetrating || !PrimitiveComponent)
diff --git a/Engine/Source/Runtime/Core/Private/Misc/MapErrors.cpp b/Engine/Source/Runtime/Core/Private/Misc/MapErrors.cpp
--- a/Engine/Source/Runtime/Core/Private/Misc/MapErrors.cpp
+++ b/Engine/Source/Runtime/Core/Private/Misc/MapErrors.cpp
@@ -37,6 +37,7 @@ FLazyName FMapErrors::FixedUpInvalidLandscapeMaterialInstances(TEXT("FixedUpInva
 FLazyName FMapErrors::FixedUpIncorrectLayerWeightmap(TEXT("FixedUpIncorrectLayerWeightmap"));
 FLazyName FMapErrors::FixedUpSharedLayerWeightmap(TEXT("FixedUpSharedLayerWeightmap"));
 FLazyName FMapErrors::LandscapeComponentPostLoad_Warning(TEXT("LandscapeComponentPostLoad_Warning"));
+FLazyName FMapErrors::LandscapeLazyObjectPtrDeprecation_Warning(TEXT("LandscapeLazyObjectPtrDeprecation_Warning"));
 FLazyName FMapErrors::DuplicateLevelInfo(TEXT("DuplicateLevelInfo"));
 FLazyName FMapErrors::NoKillZ(TEXT("NoKillZ"));
 FLazyName FMapErrors::LightComponentNull(TEXT("LightComponentNull"));
diff --git a/Engine/Source/Runtime/Core/Public/Misc/MapErrors.h b/Engine/Source/Runtime/Core/Public/Misc/MapErrors.h
--- a/Engine/Source/Runtime/Core/Public/Misc/MapErrors.h
+++ b/Engine/Source/Runtime/Core/Public/Misc/MapErrors.h
@@ -123,6 +123,9 @@ struct CORE_API FMapErrors
 	/** Landscape ({ProxyName}) has overlapping render components at location ({X, Y}): */
 	static FLazyName LandscapeComponentPostLoad_Warning;
 
+	/** Landscape proxy {ProxyName} of {LevelName} points to a LandscapeActor that is not currently loaded. This will lose the property upon save.  */
+	static FLazyName LandscapeLazyObjectPtrDeprecation_Warning;
+
 
     /**  Level */
 
diff --git a/Engine/Source/Runtime/Landscape/Classes/LandscapeComponent.h b/Engine/Source/Runtime/Landscape/Classes/LandscapeComponent.h
--- a/Engine/Source/Runtime/Landscape/Classes/LandscapeComponent.h
+++ b/Engine/Source/Runtime/Landscape/Classes/LandscapeComponent.h
@@ -420,7 +420,7 @@ class ULandscapeComponent : public UPrimitiveComponent
 
 	UPROPERTY()
 	TObjectPtr<UMaterialInstanceConstant> MaterialInstance_DEPRECATED;
-#endif
+#endif // WITH_EDITORONLY_DATA
 
 	UPROPERTY(TextExportTransient)
 	TArray<TObjectPtr<UMaterialInstanceConstant>> MaterialInstances;
@@ -452,11 +452,16 @@ class ULandscapeComponent : public UPrimitiveComponent
 	UPROPERTY()
 	FBox CachedLocalBox;
 
-	/** Reference to associated collision component */
+#if WITH_EDITORONLY_DATA
 	UPROPERTY()
-	TLazyObjectPtr<ULandscapeHeightfieldCollisionComponent> CollisionComponent;
+	TLazyObjectPtr<ULandscapeHeightfieldCollisionComponent> CollisionComponent_DEPRECATED;
+#endif // !WITH_EDITORONLY_DATA
 
 private:
+	/** Reference to associated collision component */
+	UPROPERTY()
+	TObjectPtr<ULandscapeHeightfieldCollisionComponent> CollisionComponentRef;
+
 
 	/** Store  */ 
 	UPROPERTY(Transient)
@@ -512,10 +517,6 @@ public:
 	UPROPERTY()
 	FGuid MapBuildDataId;
 
-	/**	Legacy irrelevant lights */
-	UPROPERTY()
-	TArray<FGuid> IrrelevantLights_DEPRECATED;
-
 	/** Heightfield mipmap used to generate collision */
 	UPROPERTY(EditAnywhere, Category=LandscapeComponent)
 	int32 CollisionMipLevel;
@@ -562,6 +563,10 @@ public:
 	TObjectPtr<UTexture2D> GIBakedBaseColorTexture;
 
 #if WITH_EDITORONLY_DATA
+	/**	Legacy irrelevant lights */
+	UPROPERTY()
+	TArray<FGuid> IrrelevantLights_DEPRECATED;
+
 	/** LOD level Bias to use when lighting building via lightmass, -1 Means automatic LOD calculation based on ForcedLOD + LODBias */
 	UPROPERTY(EditAnywhere, Category=LandscapeComponent)
 	int32 LightingLODBias;
@@ -595,7 +600,7 @@ public:
 	/** Represents last saved hash for PhysicalMaterialTask */
 	UPROPERTY(Transient)
 	uint32 LastSavedPhysicalMaterialHash;
-#endif
+#endif // WITH_EDITORONLY_DATA
 
 	UPROPERTY(NonPIEDuplicateTransient)
 	TObjectPtr<UMaterialInterface> MobileMaterialInterface_DEPRECATED;
@@ -620,7 +625,7 @@ public:
 
 	UPROPERTY(NonPIEDuplicateTransient)
 	TObjectPtr<UMaterialInstanceConstant> MobileCombinationMaterialInstance_DEPRECATED;
-#endif
+#endif // WITH_EDITORONLY_DATA
 
 public:
 	/** Grass data for generation **/
@@ -1095,6 +1100,9 @@ public:
 		return bNaniteActive;
 	}
 
+	LANDSCAPE_API ULandscapeHeightfieldCollisionComponent* GetCollisionComponent() const { return CollisionComponentRef.Get(); }
+	LANDSCAPE_API void SetCollisionComponent(ULandscapeHeightfieldCollisionComponent* InCollisionComponent) { CollisionComponentRef = InCollisionComponent; }
+
 	void SetUserTriggeredChangeRequested(bool bInUserTriggeredChangeRequested)
 	{
 		bUserTriggeredChangeRequested = bInUserTriggeredChangeRequested;
diff --git a/Engine/Source/Runtime/Landscape/Classes/LandscapeHeightfieldCollisionComponent.h b/Engine/Source/Runtime/Landscape/Classes/LandscapeHeightfieldCollisionComponent.h
--- a/Engine/Source/Runtime/Landscape/Classes/LandscapeHeightfieldCollisionComponent.h
+++ b/Engine/Source/Runtime/Landscape/Classes/LandscapeHeightfieldCollisionComponent.h
@@ -81,13 +81,21 @@ class ULandscapeHeightfieldCollisionComponent : public UPrimitiveComponent
 	UPROPERTY()
 	FBox CachedLocalBox;
 
+#if WITH_EDITORONLY_DATA
 	/** Reference to render component */
 	UPROPERTY()
-	TLazyObjectPtr<ULandscapeComponent> RenderComponent;
+	TLazyObjectPtr<ULandscapeComponent> RenderComponent_DEPRECATED;
+#endif // !WITH_EDITORONLY_DATA
 
+private:
+	/** Reference to render component */
+	UPROPERTY()
+	TObjectPtr<ULandscapeComponent> RenderComponentRef;
+
+public:
 	/** Returns associated landscape component */
 	UFUNCTION(BlueprintCallable, Category = "Landscape")
-	ULandscapeComponent* GetRenderComponent() const;
+	LANDSCAPE_API ULandscapeComponent* GetRenderComponent() const;
 
 	struct FHeightfieldGeometryRef : public FRefCountedObject
 	{
@@ -285,6 +293,8 @@ public:
 	LANDSCAPE_API void SnapFoliageInstances();
 #endif
 
+	LANDSCAPE_API void SetRenderComponent(ULandscapeComponent* InRenderComponent) { RenderComponentRef = InRenderComponent; }
+
 public:
 	TOptional<float> GetHeight(float X, float Y, EHeightfieldSource HeightFieldSource);
 
diff --git a/Engine/Source/Runtime/Landscape/Classes/LandscapeInfo.h b/Engine/Source/Runtime/Landscape/Classes/LandscapeInfo.h
--- a/Engine/Source/Runtime/Landscape/Classes/LandscapeInfo.h
+++ b/Engine/Source/Runtime/Landscape/Classes/LandscapeInfo.h
@@ -124,7 +124,7 @@ class ULandscapeInfo : public UObject
 	GENERATED_UCLASS_BODY()
 
 	UPROPERTY()
-	TLazyObjectPtr<ALandscape> LandscapeActor;
+	TWeakObjectPtr<ALandscape> LandscapeActor;
 
 	UPROPERTY()
 	FGuid LandscapeGuid;
@@ -166,6 +166,8 @@ public:
 
 private:
 #if WITH_EDITORONLY_DATA
+	TSet<TWeakObjectPtr<UObject>> PackagesToDeprecate;
+
 	UPROPERTY()
 	TArray<TScriptInterface<ILandscapeSplineInterface>> SplineActors;
 
@@ -260,6 +262,10 @@ public:
 
 	LANDSCAPE_API void ModifyObject(UObject* InObject, bool bAlwaysMarkDirty = true);
 	LANDSCAPE_API void MarkObjectDirty(UObject* InObject);
+
+	LANDSCAPE_API void RequestPackageDeprecation(UObject* InObject);
+	LANDSCAPE_API void ShowForceUpdateMapCheckError(UObject* InOutermost);
+	LANDSCAPE_API void ForceUpdateDeprecatedObjects();
 #endif //WITH_EDITOR
 
 	/**
diff --git a/Engine/Source/Runtime/Landscape/Classes/LandscapeProxy.h b/Engine/Source/Runtime/Landscape/Classes/LandscapeProxy.h
--- a/Engine/Source/Runtime/Landscape/Classes/LandscapeProxy.h
+++ b/Engine/Source/Runtime/Landscape/Classes/LandscapeProxy.h
@@ -774,6 +774,7 @@ public:
 
 #if WITH_EDITOR
 	FOnLandscapeProxyComponentDataChanged OnComponentDataChanged;
+	bool bPackageDeprecationNeeded = false;
 #endif
 
 	// Blueprint functions
@@ -1309,6 +1310,8 @@ public:
 	LANDSCAPE_API bool CanHaveLayersContent() const;
 
 	LANDSCAPE_API virtual void UpdateCachedHasLayersContent(bool InCheckComponentDataIntegrity = false);
+	
+	LANDSCAPE_API void RequestPackageDeprecation() { bPackageDeprecationNeeded = true; };
 
 protected:
 	friend class ALandscape;
@@ -1333,7 +1336,7 @@ protected:
 protected:
 	FLandscapeMaterialChangedDelegate LandscapeMaterialChangedDelegate;
 
-#endif
+#endif // WITH_EDITOR
 private:
 	/** Returns Grass Update interval */
 	FORCEINLINE int32 GetGrassUpdateInterval() const 
diff --git a/Engine/Source/Runtime/Landscape/Classes/LandscapeStreamingProxy.h b/Engine/Source/Runtime/Landscape/Classes/LandscapeStreamingProxy.h
--- a/Engine/Source/Runtime/Landscape/Classes/LandscapeStreamingProxy.h
+++ b/Engine/Source/Runtime/Landscape/Classes/LandscapeStreamingProxy.h
@@ -23,9 +23,16 @@ class ALandscapeStreamingProxy : public ALandscapeProxy
 public:
 	ALandscapeStreamingProxy(const FObjectInitializer& ObjectInitializer);
 
-	UPROPERTY(EditAnywhere, Category=LandscapeProxy)
-	TLazyObjectPtr<ALandscape> LandscapeActor;
-
+#if WITH_EDITORONLY_DATA
+	UPROPERTY()
+	TLazyObjectPtr<ALandscape> LandscapeActor_DEPRECATED;
+#endif // WITH_EDITORONLY_DATA
+
+private:
+	UPROPERTY(EditAnywhere, Category = LandscapeProxy, Meta = (DisplayName = "Landscape Actor"))
+	TSoftObjectPtr<ALandscape> LandscapeActorRef;
+	
+public:
 	//~ Begin UObject Interface
 #if WITH_EDITOR
 	virtual bool ShouldExport() override { return false;  }
@@ -43,6 +50,7 @@ public:
 	//~ Begin ALandscapeBase Interface
 	virtual ALandscape* GetLandscapeActor() override;
 	virtual const ALandscape* GetLandscapeActor() const override;
+	void LANDSCAPE_API SetLandscapeActor(ALandscape* InLandscape);
 #if WITH_EDITOR
 	virtual UMaterialInterface* GetLandscapeMaterial(int8 InLODIndex = INDEX_NONE) const override;
 	virtual UMaterialInterface* GetLandscapeHoleMaterial() const override;
diff --git a/Engine/Source/Runtime/Landscape/Private/Landscape.cpp b/Engine/Source/Runtime/Landscape/Private/Landscape.cpp
--- a/Engine/Source/Runtime/Landscape/Private/Landscape.cpp
+++ b/Engine/Source/Runtime/Landscape/Private/Landscape.cpp
@@ -665,7 +665,7 @@ void ULandscapeComponent::Serialize(FArchive& Ar)
 		}
 	}
 	else
-#endif
+#endif // WITH_EDITOR
 	{
 		Super::Serialize(Ar);
 	}
@@ -675,13 +675,17 @@ void ULandscapeComponent::Serialize(FArchive& Ar)
 		FMeshMapBuildData* LegacyMapBuildData = new FMeshMapBuildData();
 		Ar << LegacyMapBuildData->LightMap;
 		Ar << LegacyMapBuildData->ShadowMap;
+
+#if WITH_EDITORONLY_DATA
 		LegacyMapBuildData->IrrelevantLights = IrrelevantLights_DEPRECATED;
+#endif // WITH_EDITORONLY_DATA
 
 		FMeshMapBuildLegacyData LegacyComponentData;
 		LegacyComponentData.Data.Emplace(MapBuildDataId, LegacyMapBuildData);
 		GComponentsWithLegacyLightmaps.AddAnnotation(this, MoveTemp(LegacyComponentData));
 	}
 
+#if WITH_EDITORONLY_DATA
 	if (Ar.IsLoading() && Ar.CustomVer(FFortniteMainBranchObjectVersion::GUID) < FFortniteMainBranchObjectVersion::NewLandscapeMaterialPerLOD)
 	{
 		if (MobileMaterialInterface_DEPRECATED != nullptr)
@@ -689,13 +693,12 @@ void ULandscapeComponent::Serialize(FArchive& Ar)
 			MobileMaterialInterfaces.AddUnique(MobileMaterialInterface_DEPRECATED);
 		}
 
-#if WITH_EDITORONLY_DATA
 		if (MobileCombinationMaterialInstance_DEPRECATED != nullptr)
 		{
 			MobileCombinationMaterialInstances.AddUnique(MobileCombinationMaterialInstance_DEPRECATED);
 		}
-#endif
 	}
+#endif // WITH_EDITORONLY_DATA
 
 	if (Ar.UEVer() >= VER_UE4_SERIALIZE_LANDSCAPE_GRASS_DATA)
 	{
@@ -1219,8 +1222,21 @@ void ULandscapeComponent::PostLoad()
 
 	GrassData->ConditionalDiscardDataOnLoad();
 
-#endif // WITH_EDITOR
+#if WITH_EDITORONLY_DATA
+	// If the Collision Component is not set yet and we're transferring the property from the lazy object pointer it was previously stored as to the soft object ptr it is now stored as :
+	if (!CollisionComponentRef && CollisionComponent_DEPRECATED.IsValid())
+	{
+		CollisionComponentRef = CollisionComponent_DEPRECATED.Get();
+		CollisionComponent_DEPRECATED = nullptr;
+
+		if (ALandscapeProxy* Proxy = GetLandscapeProxy())
+		{
+			Proxy->RequestPackageDeprecation();
+		}
+	}
+#endif // !WITH_EDITORONLY_DATA
 
+#endif // WITH_EDITOR
 }
 
 #if WITH_EDITORONLY_DATA
@@ -1390,12 +1406,17 @@ ALandscape* ALandscape::GetLandscapeActor()
 
 const ALandscape* ALandscapeStreamingProxy::GetLandscapeActor() const
 {
-	return LandscapeActor.Get();
+	return LandscapeActorRef.Get();
 }
 
 ALandscape* ALandscapeStreamingProxy::GetLandscapeActor()
 {
-	return LandscapeActor.Get();
+	return LandscapeActorRef.Get();
+}
+
+void ALandscapeStreamingProxy::SetLandscapeActor(ALandscape* InLandscape)
+{
+	LandscapeActorRef = InLandscape;
 }
 
 ULandscapeInfo* ALandscapeProxy::CreateLandscapeInfo(bool bMapCheck, bool bUpdateAllAddCollisions)
@@ -2282,15 +2303,15 @@ void ALandscapeProxy::PostRegisterAllComponents()
 				UpdateCachedHasLayersContent(true);
 
 				// Cache the value at this point as CreateLandscapeInfo (-> RegisterActor) might create/destroy layers content if there was a mismatch between landscape & proxy
-				// Check the actual flag here not HasLayersContent() which could return true if the LandscapeActor is valid.
+				// Check the actual flag here not HasLayersContent() which could return true if the LandscapeActorRef is valid.
 				bool bHasLayersContentBefore = bHasLayersContent;
 
 				check(WeightmapFixupVersion == CurrentVersion);
 
 				const bool bNeedOldDataMigration = !bHasLayersContentBefore && CanHaveLayersContent();
-				if (bNeedOldDataMigration && LandscapeInfo->LandscapeActor.IsValid() && LandscapeInfo->LandscapeActor.Get()->HasLayersContent())
+				if (bNeedOldDataMigration && LandscapeInfo->LandscapeActor.IsValid() && LandscapeInfo->LandscapeActor->HasLayersContent())
 				{
-					LandscapeInfo->LandscapeActor.Get()->CopyOldDataToDefaultLayer(this);
+					LandscapeInfo->LandscapeActor->CopyOldDataToDefaultLayer(this);
 				}
 			}
 #endif // WITH_EDITOR
@@ -2312,6 +2333,12 @@ void ALandscapeProxy::PostRegisterAllComponents()
 			LandscapeInfo->FixupProxiesTransform();
 		}
 	}
+
+	if (bPackageDeprecationNeeded && (LandscapeInfo != nullptr))
+	{
+		LandscapeInfo->RequestPackageDeprecation(this);
+		bPackageDeprecationNeeded = false;
+	}
 #endif // WITH_EDITOR
 }
 
@@ -2458,7 +2485,7 @@ void ALandscape::PostLoad()
 				break;
 			}
 		}
-#endif
+#endif // WITH_EDITOR
 	}
 
 #if WITH_EDITOR
@@ -2471,7 +2498,7 @@ void ALandscape::PostLoad()
 			Brush.SetOwner(this);
 		}
 	}
-#endif
+#endif // WITH_EDITOR
 
 	Super::PostLoad();
 }
@@ -2734,7 +2761,6 @@ bool ALandscape::IsValidRenderTargetFormatHeightmap(EPixelFormat InRenderTargetF
 	return false;
 }
 
-// TODO [jonathan.bard] This does not support being called at runtime yet because ULandscapeInfo::LandscapeActor is not valid at runtime or in non-editor worlds (e.g. PIE), for no good reason :
 void ALandscape::RenderHeightmap(const FTransform& InRenderAreaWorldTransform, const FBox2D& InRenderAreaExtents, UTextureRenderTarget2D* OutRenderTarget)
 {
 	// TODO: We may want a version of this function that returns a lambda that can be passed to the render thread and run
@@ -3329,7 +3355,7 @@ void ALandscapeProxy::PostLoad()
 
 	if ((GetLinker() && (GetLinker()->UEVer() < VER_UE4_LANDSCAPE_COMPONENT_LAZY_REFERENCES)) ||
 		LandscapeComponents.Num() != CollisionComponents.Num() ||
-		LandscapeComponents.ContainsByPredicate([](ULandscapeComponent* Comp) { return ((Comp != nullptr) && !Comp->CollisionComponent.IsValid()); }))
+		LandscapeComponents.ContainsByPredicate([](ULandscapeComponent* Comp) { return ((Comp != nullptr) && (Comp->GetCollisionComponent() == nullptr)); }))
 	{
 		CreateLandscapeInfo();
 	}
@@ -3346,13 +3372,32 @@ void ALandscapeProxy::PostLoad()
 	}
 	PRAGMA_ENABLE_DEPRECATION_WARNINGS;
 
-	if (GIsEditor && GetWorld() && !GetWorld()->IsGameWorld())
+	if (GIsEditor)
 	{
+		// We may not have run PostLoad on LandscapeComponents yet
+		for (TObjectPtr<ULandscapeComponent>& LandscapeComponent : LandscapeComponents)
+		{
+			if (LandscapeComponent)
+			{
+				LandscapeComponent->ConditionalPostLoad();
+			}
+		}
+		
+		// We may not have run PostLoad on CollisionComponent yet
+		for (TObjectPtr<ULandscapeHeightfieldCollisionComponent>& CollisionComponent : CollisionComponents)
+		{
+			if (CollisionComponent)
+			{
+				CollisionComponent->ConditionalPostLoad();
+			}
+		}
+
 		if ((GetLinker() && (GetLinker()->UEVer() < VER_UE4_LANDSCAPE_COMPONENT_LAZY_REFERENCES)) ||
 			LandscapeComponents.Num() != CollisionComponents.Num() ||
-			LandscapeComponents.ContainsByPredicate([](ULandscapeComponent* Comp) { return ((Comp != nullptr) && !Comp->CollisionComponent.IsValid()); }))
+			LandscapeComponents.ContainsByPredicate([](ULandscapeComponent* Comp) { return ((Comp != nullptr) && (Comp->GetCollisionComponent() == nullptr)); }) ||
+			CollisionComponents.ContainsByPredicate([](ULandscapeHeightfieldCollisionComponent* Comp) { return ((Comp != nullptr) && (Comp->GetRenderComponent() == nullptr)); }))
 		{
-			// Need to clean up invalid collision components
+			// Need to clean up invalid collision and render components
 			RecreateCollisionComponents();
 		}
 	}
@@ -3720,9 +3765,9 @@ UMaterialInterface* ALandscapeStreamingProxy::GetLandscapeMaterial(int8 InLODInd
 		return LandscapeMaterial;
 	}
 
-	if (LandscapeActor != nullptr)
+	if (const ALandscape* Landscape = GetLandscapeActor())
 	{
-		return LandscapeActor->GetLandscapeMaterial(InLODIndex);
+		return Landscape->GetLandscapeMaterial(InLODIndex);
 	}
 
 	return UMaterial::GetDefaultMaterial(MD_Surface);
@@ -3734,7 +3779,7 @@ UMaterialInterface* ALandscapeStreamingProxy::GetLandscapeHoleMaterial() const
 	{
 		return LandscapeHoleMaterial;
 	}
-	else if (ALandscape* Landscape = LandscapeActor.Get())
+	else if (const ALandscape* Landscape = GetLandscapeActor())
 	{
 		return Landscape->GetLandscapeHoleMaterial();
 	}
@@ -3900,6 +3945,48 @@ void ULandscapeInfo::ModifyObject(UObject* InObject, bool bAlwaysMarkDirty)
 	}
 }
 
+void ULandscapeInfo::RequestPackageDeprecation(UObject* InObject)
+{
+	check(InObject != nullptr);
+	// This is either the first time or the user already deprecated a set of object previously, we need to display the MapCheck error.
+	// We do not display the message during cook as the data will be fixedup in the produced asset, and user cannot click on the action.
+	if (PackagesToDeprecate.IsEmpty() && !GIsCookerLoadingPackage)
+	{
+		UObject* Outermost = InObject->GetOutermostObject();
+		ShowForceUpdateMapCheckError(Outermost);
+	}
+
+	PackagesToDeprecate.Add(InObject);
+}
+
+void ULandscapeInfo::ShowForceUpdateMapCheckError(UObject* InOutermost)
+{
+	check(InOutermost != nullptr)
+
+	FMessageLog("MapCheck").Warning()
+		->AddToken(FUObjectToken::Create(InOutermost))
+		->AddToken(FTextToken::Create(LOCTEXT("ULandscapeInfo_ShowForceUpdateMapCheckError_Message_ForceUpdateNeeded", "Landscape packages are out of date and need updating for landscape to be loaded properly.")))
+		->AddToken(FMapErrorToken::Create(TEXT("ULandscapeInfo_ShowForceUpdateMapCheckError_MapError_ForceUpdateNeeded")))
+		->AddToken(FActionToken::Create(LOCTEXT("ULandscapeInfo_ShowForceUpdateMapCheckError_ActionName_ForceUpdateNeeded", "Update Landscape Packages"), FText(),
+			FOnActionTokenExecuted::CreateUObject(this, &ULandscapeInfo::ForceUpdateDeprecatedObjects), true));
+
+	// Show MapCheck window
+	FMessageLog("MapCheck").Open(EMessageSeverity::Warning);
+}
+
+void ULandscapeInfo::ForceUpdateDeprecatedObjects()
+{
+	for (const TWeakObjectPtr<UObject>& ObjectToDeprecate : PackagesToDeprecate)
+	{
+		if (ObjectToDeprecate.IsValid())
+		{
+			ObjectToDeprecate->MarkPackageDirty();
+		}
+	}
+
+	PackagesToDeprecate.Empty();
+}
+
 ALandscapeProxy* ULandscapeInfo::GetLandscapeProxyForLevel(ULevel* Level) const
 {
 	ALandscapeProxy* LandscapeProxy = nullptr;
@@ -4247,7 +4334,7 @@ void ULandscapeInfo::RegisterActor(ALandscapeProxy* Proxy, bool bMapCheck, bool
 	// register
 	if (ALandscape* Landscape = Cast<ALandscape>(Proxy))
 	{
-		if (LandscapeActor == nullptr)
+		if (!LandscapeActor.IsValid())
 		{
 			LandscapeActor = Landscape;
 
@@ -4263,7 +4350,7 @@ void ULandscapeInfo::RegisterActor(ALandscapeProxy* Proxy, bool bMapCheck, bool
 			{
 				if (ALandscapeStreamingProxy* StreamingProxy = StreamingProxyPtr.Get())
 				{
-					StreamingProxy->LandscapeActor = LandscapeActor;
+					StreamingProxy->SetLandscapeActor(LandscapeActor.Get());
 #if WITH_EDITOR
 					StreamingProxy->FixupSharedData(Landscape);
 #endif // WITH_EDITOR
@@ -4310,7 +4397,7 @@ void ULandscapeInfo::RegisterActor(ALandscapeProxy* Proxy, bool bMapCheck, bool
 			StreamingProxies.Insert(StreamingProxyPtr, InsertIndex);
 		}
 
-		StreamingProxy->LandscapeActor = LandscapeActor;
+		StreamingProxy->SetLandscapeActor(LandscapeActor.Get());
 #if WITH_EDITOR
 		StreamingProxy->FixupSharedData(LandscapeActor.Get());
 #endif // WITH_EDITOR
@@ -4362,7 +4449,7 @@ void ULandscapeInfo::UnregisterActor(ALandscapeProxy* Proxy)
 		{
 			if (ALandscapeStreamingProxy* StreamingProxy = StreamingProxyPtr.Get())
 			{
-				StreamingProxy->LandscapeActor = LandscapeActor;
+				StreamingProxy->SetLandscapeActor(LandscapeActor.Get());
 			}
 		}
 	}
@@ -4371,7 +4458,7 @@ void ULandscapeInfo::UnregisterActor(ALandscapeProxy* Proxy)
 		ALandscapeStreamingProxy* StreamingProxy = CastChecked<ALandscapeStreamingProxy>(Proxy);
 		TWeakObjectPtr<ALandscapeStreamingProxy> StreamingProxyPtr = StreamingProxy;
 		StreamingProxies.Remove(StreamingProxyPtr);
-		StreamingProxy->LandscapeActor = nullptr;
+		StreamingProxy->SetLandscapeActor(nullptr);
 	}
 
 #if WITH_EDITOR
@@ -4480,7 +4567,7 @@ void ULandscapeInfo::UnregisterSplineActor(TScriptInterface<ILandscapeSplineInte
 
 void ULandscapeInfo::RequestSplineLayerUpdate()
 {
-	if (LandscapeActor)
+	if (LandscapeActor.IsValid())
 	{
 		LandscapeActor->RequestSplineLayerUpdate();
 	}
@@ -4488,7 +4575,7 @@ void ULandscapeInfo::RequestSplineLayerUpdate()
 
 void ULandscapeInfo::ForceLayersFullUpdate()
 {
-	if (LandscapeActor)
+	if (LandscapeActor.IsValid())
 	{
 		LandscapeActor->ForceLayersFullUpdate();
 	}
@@ -4528,9 +4615,6 @@ void ULandscapeInfo::UnregisterCollisionComponent(ULandscapeHeightfieldCollision
 	}
 }
 
-// TODO [jonathan.bard] This does not support being called at runtime yet because ULandscapeInfo::LandscapeActor 
-// is not valid at runtime or in non-editor worlds (e.g. PIE), for no good reason, and also XYtoComponentMap is
-// described as being valid only in editor.
 bool ULandscapeInfo::GetOverlappedComponents(const FTransform& InAreaWorldTransform, const FBox2D& InAreaExtents, 
 	TMap<FIntPoint, ULandscapeComponent*>& OutOverlappedComponents, FIntRect& OutComponentIndicesBoundingRect)
 {
diff --git a/Engine/Source/Runtime/Landscape/Private/LandscapeCollision.cpp b/Engine/Source/Runtime/Landscape/Private/LandscapeCollision.cpp
--- a/Engine/Source/Runtime/Landscape/Private/LandscapeCollision.cpp
+++ b/Engine/Source/Runtime/Landscape/Private/LandscapeCollision.cpp
@@ -1727,7 +1727,7 @@ FBoxSphereBounds ULandscapeHeightfieldCollisionComponent::CalcBounds(const FTran
 
 void ULandscapeHeightfieldCollisionComponent::BeginDestroy()
 {
-	HeightfieldRef = NULL;
+	HeightfieldRef = nullptr;
 	HeightfieldGuid = FGuid();
 	Super::BeginDestroy();
 }
@@ -1736,7 +1736,7 @@ void ULandscapeMeshCollisionComponent::BeginDestroy()
 {
 	if (!HasAnyFlags(RF_ClassDefaultObject))
 	{
-		MeshRef = NULL;
+		MeshRef = nullptr;
 		MeshGuid = FGuid();
 	}
 
@@ -2159,6 +2159,21 @@ void ULandscapeHeightfieldCollisionComponent::PostLoad()
 			SpeculativelyLoadAsyncDDCCollsionData();
 		}
 	}
+
+#if WITH_EDITORONLY_DATA
+	// If the RenderComponent is not set yet and we're transferring the property from the lazy object pointer it was previously stored as to the object ptr it is now stored as :
+	if (!RenderComponentRef && RenderComponent_DEPRECATED.IsValid())
+	{
+		RenderComponentRef = RenderComponent_DEPRECATED.Get();
+		RenderComponent_DEPRECATED = nullptr;
+		
+		if (ALandscapeProxy* Proxy = GetLandscapeProxy())
+		{
+			Proxy->RequestPackageDeprecation();
+		}
+	}
+#endif // !WITH_EDITORONLY_DATA
+
 #endif // WITH_EDITOR
 }
 
@@ -2177,6 +2192,8 @@ void ULandscapeHeightfieldCollisionComponent::PreSave(FObjectPreSaveContext Obje
 	{
 #if WITH_EDITOR
 		ALandscapeProxy* Proxy = GetLandscapeProxy();
+		ULandscapeComponent* RenderComponent = GetRenderComponent();
+
 		if (Proxy && Proxy->bBakeMaterialPositionOffsetIntoCollision)
 		{
 			if (!RenderComponent->GrassData->HasData() || RenderComponent->IsGrassMapOutdated())
@@ -2225,7 +2242,7 @@ void ULandscapeInfo::UpdateAllAddCollisions()
 					ULandscapeComponent* NeighborComponent = XYtoComponentMap.FindRef(NeighborsKeys[i]);
 
 					// UpdateAddCollision() treats a null CollisionComponent as an empty hole
-					if (!NeighborComponent || !NeighborComponent->CollisionComponent.IsValid())
+					if (!NeighborComponent || (NeighborComponent->GetCollisionComponent() == nullptr))
 					{
 						UpdateAddCollision(NeighborsKeys[i]);
 					}
@@ -2264,7 +2281,7 @@ void ULandscapeInfo::UpdateAddCollision(FIntPoint LandscapeKey)
 		ULandscapeComponent* Comp = XYtoComponentMap.FindRef(NeighborsKeys[i]);
 		if (Comp)
 		{
-			ULandscapeHeightfieldCollisionComponent* NeighborCollision = Comp->CollisionComponent.Get();
+			ULandscapeHeightfieldCollisionComponent* NeighborCollision = Comp->GetCollisionComponent();
 			// Skip cooked because CollisionHeightData not saved during cook
 			if (NeighborCollision && !NeighborCollision->GetOutermost()->bIsCookedForEditor)
 			{
@@ -2272,12 +2289,12 @@ void ULandscapeInfo::UpdateAddCollision(FIntPoint LandscapeKey)
 			}
 			else
 			{
-				NeighborCollisions[i] = NULL;
+				NeighborCollisions[i] = nullptr;
 			}
 		}
 		else
 		{
-			NeighborCollisions[i] = NULL;
+			NeighborCollisions[i] = nullptr;
 		}
 	}
 
@@ -2649,7 +2666,7 @@ ULandscapeHeightfieldCollisionComponent::~ULandscapeHeightfieldCollisionComponen
 
 ULandscapeComponent* ULandscapeHeightfieldCollisionComponent::GetRenderComponent() const
 {
-	return RenderComponent.Get();
+	return RenderComponentRef.Get();
 }
 
 TOptional<float> ULandscapeHeightfieldCollisionComponent::GetHeight(float X, float Y, EHeightfieldSource HeightFieldSource)
diff --git a/Engine/Source/Runtime/Landscape/Private/LandscapeConfigHelper.cpp b/Engine/Source/Runtime/Landscape/Private/LandscapeConfigHelper.cpp
--- a/Engine/Source/Runtime/Landscape/Private/LandscapeConfigHelper.cpp
+++ b/Engine/Source/Runtime/Landscape/Private/LandscapeConfigHelper.cpp
@@ -370,7 +370,7 @@ void FLandscapeConfigHelper::MoveFoliageToLandscape(ULandscapeInfo* InLandscapeI
 	// Move instances
 	for (const TPair<FIntPoint, ULandscapeComponent*>& OldEntry : InLandscapeInfo->XYtoComponentMap)
 	{
-		ULandscapeHeightfieldCollisionComponent* OldCollisionComponent = OldEntry.Value->CollisionComponent.Get();
+		ULandscapeHeightfieldCollisionComponent* OldCollisionComponent = OldEntry.Value->GetCollisionComponent();
 
 		if (OldCollisionComponent)
 		{
@@ -378,7 +378,7 @@ void FLandscapeConfigHelper::MoveFoliageToLandscape(ULandscapeInfo* InLandscapeI
 
 			for (const TPair<FIntPoint, ULandscapeComponent*>& NewEntry : InNewLandscapeInfo->XYtoComponentMap)
 			{
-				ULandscapeHeightfieldCollisionComponent* NewCollisionComponent = NewEntry.Value->CollisionComponent.Get();
+				ULandscapeHeightfieldCollisionComponent* NewCollisionComponent = NewEntry.Value->GetCollisionComponent();
 
 				if (NewCollisionComponent && FBoxSphereBounds::BoxesIntersect(NewCollisionComponent->Bounds, OldCollisionComponent->Bounds))
 				{
@@ -395,7 +395,7 @@ void FLandscapeConfigHelper::MoveFoliageToLandscape(ULandscapeInfo* InLandscapeI
 	// Snap them to the bounds
 	for (const TPair<FIntPoint, ULandscapeComponent*>& NewEntry : InNewLandscapeInfo->XYtoComponentMap)
 	{
-		ULandscapeHeightfieldCollisionComponent* NewCollisionComponent = NewEntry.Value->CollisionComponent.Get();
+		ULandscapeHeightfieldCollisionComponent* NewCollisionComponent = NewEntry.Value->GetCollisionComponent();
 
 		if (NewCollisionComponent)
 		{
diff --git a/Engine/Source/Runtime/Landscape/Private/LandscapeEdit.cpp b/Engine/Source/Runtime/Landscape/Private/LandscapeEdit.cpp
--- a/Engine/Source/Runtime/Landscape/Private/LandscapeEdit.cpp
+++ b/Engine/Source/Runtime/Landscape/Private/LandscapeEdit.cpp
@@ -59,6 +59,7 @@ LandscapeEdit.cpp: Landscape editing
 #include "NaniteBuilder.h"
 #include "Rendering/NaniteResources.h"
 #include "Misc/ScopedSlowTask.h"
+#include "Misc/UObjectToken.h"
 #include "EngineModule.h"
 #include "EngineUtils.h"
 #include "Framework/Notifications/NotificationManager.h"
@@ -173,7 +174,7 @@ void ULandscapeComponent::UpdateCachedBounds(bool bInApproximateBounds)
 	}
 
 	// Update collision component bounds
-	ULandscapeHeightfieldCollisionComponent* HFCollisionComponent = CollisionComponent.Get();
+	ULandscapeHeightfieldCollisionComponent* HFCollisionComponent = GetCollisionComponent();
 	if (HFCollisionComponent)
 	{
         // In Landscape Layers the Collision Component gets regenerated after the heightmap changes are undone and doesn't need to be transacted
@@ -190,6 +191,8 @@ void ULandscapeComponent::UpdateCachedBounds(bool bInApproximateBounds)
 void ULandscapeComponent::UpdateNavigationRelevance()
 {
 	ALandscapeProxy* Proxy = GetLandscapeProxy();
+	ULandscapeHeightfieldCollisionComponent* CollisionComponent = GetCollisionComponent();
+
 	if (CollisionComponent && Proxy)
 	{
 		CollisionComponent->SetCanEverAffectNavigation(Proxy->bUsedForNavigation);
@@ -199,6 +202,8 @@ void ULandscapeComponent::UpdateNavigationRelevance()
 void ULandscapeComponent::UpdateRejectNavmeshUnderneath()
 {
 	ALandscapeProxy* Proxy = GetLandscapeProxy();
+	ULandscapeHeightfieldCollisionComponent* CollisionComponent = GetCollisionComponent();
+
 	if (CollisionComponent && Proxy)
 	{
 		CollisionComponent->bFillCollisionUnderneathForNavmesh = Proxy->bFillCollisionUnderLandscapeForNavmesh;
@@ -1115,7 +1120,7 @@ void ULandscapeComponent::UpdateCollisionHeightData(const FColor* const Heightma
 {
 	ALandscapeProxy* Proxy = GetLandscapeProxy();
 	FIntPoint ComponentKey = GetSectionBase() / ComponentSizeQuads;
-	ULandscapeHeightfieldCollisionComponent* CollisionComp = CollisionComponent.Get();
+	ULandscapeHeightfieldCollisionComponent* CollisionComp = GetCollisionComponent();
 	ULandscapeMeshCollisionComponent* MeshCollisionComponent = Cast<ULandscapeMeshCollisionComponent>(CollisionComp);
 	ULandscapeHeightfieldCollisionComponent* OldCollisionComponent = CollisionComp;
 
@@ -1183,7 +1188,7 @@ void ULandscapeComponent::UpdateCollisionHeightData(const FColor* const Heightma
 		ComponentY2 = ComponentSizeQuads;
 
 		RecreateCollisionComponent(bUsingSimpleCollision);
-		CollisionComp = CollisionComponent.Get();
+		CollisionComp = GetCollisionComponent();
         MeshCollisionComponent = Cast<ULandscapeMeshCollisionComponent>(CollisionComp);
 	}
 
@@ -1288,11 +1293,12 @@ void ULandscapeComponent::UpdateCollisionHeightData(const FColor* const Heightma
 
 void ULandscapeComponent::DestroyCollisionData()
 {
-	ULandscapeHeightfieldCollisionComponent* CollisionComp = CollisionComponent.Get();
+	ULandscapeHeightfieldCollisionComponent* CollisionComp = GetCollisionComponent();
+	
 	if (CollisionComp)
 	{
 		CollisionComp->DestroyComponent();
-		CollisionComponent = CollisionComp = nullptr;
+		SetCollisionComponent(nullptr);
 	}
 }
 
@@ -1321,7 +1327,7 @@ void ULandscapeComponent::UpdateCollisionData(bool bInUpdateHeightfieldRegion)
 
 void ULandscapeComponent::RecreateCollisionComponent(bool bUseSimpleCollision)
 {
-	ULandscapeHeightfieldCollisionComponent* CollisionComp = CollisionComponent.Get();
+	ULandscapeHeightfieldCollisionComponent* CollisionComp = GetCollisionComponent();
 	ULandscapeMeshCollisionComponent* MeshCollisionComponent = nullptr;
 	TArray<uint8> DominantLayerData;
 	TArray<ULandscapeLayerInfoObject*> LayerInfos;
@@ -1372,7 +1378,7 @@ void ULandscapeComponent::RecreateCollisionComponent(bool bUseSimpleCollision)
 	CollisionComp->SetupAttachment(Proxy->GetRootComponent(), NAME_None);
 	Proxy->CollisionComponents.Add(CollisionComp);
 
-	CollisionComp->RenderComponent = this;
+	CollisionComp->SetRenderComponent(this);
 	CollisionComp->SetSectionBase(GetSectionBase());
 	CollisionComp->CollisionSizeQuads = CollisionSize.SubsectionSizeQuads * NumSubsections;
 	CollisionComp->CollisionScale = (float)(ComponentSizeQuads) / (float)(CollisionComp->CollisionSizeQuads);
@@ -1407,7 +1413,7 @@ void ULandscapeComponent::RecreateCollisionComponent(bool bUseSimpleCollision)
 	{
 		CollisionComp->ComponentLayerInfos = MoveTemp(LayerInfos);
 	}
-	CollisionComponent = CollisionComp;
+	SetCollisionComponent(CollisionComp);
 }
 
 void ULandscapeComponent::UpdateCollisionHeightBuffer(	int32 InComponentX1, int32 InComponentY1, int32 InComponentX2, int32 InComponentY2, int32 InCollisionMipLevel, int32 InHeightmapSizeU, int32 InHeightmapSizeV,
@@ -1585,7 +1591,7 @@ void ULandscapeComponent::UpdateCollisionLayerData(const FColor* const* const We
 	ALandscapeProxy* Proxy = GetLandscapeProxy();
 	FIntPoint ComponentKey = GetSectionBase() / ComponentSizeQuads;
 
-	ULandscapeHeightfieldCollisionComponent* CollisionComp = CollisionComponent.Get();
+	ULandscapeHeightfieldCollisionComponent* CollisionComp = GetCollisionComponent();
 
 	if (CollisionComp)
 	{
@@ -1791,7 +1797,7 @@ bool ULandscapeComponent::CanUpdatePhysicalMaterial()
 		return false;
 	}
 
-	return CollisionComponent != nullptr;
+	return GetCollisionComponent() != nullptr;
 }
 
 
@@ -1811,6 +1817,9 @@ void ULandscapeComponent::UpdatePhysicalMaterialTasks()
 
 	// Check if we need to launch a new task to update the physical material.
 	uint32 Hash = CalculatePhysicalMaterialTaskHash();
+	ULandscapeHeightfieldCollisionComponent* CollisionComponent = GetCollisionComponent();
+	check(CollisionComponent != nullptr);
+
 	if (PhysicalMaterialHash != Hash)
 	{
 		TArray<UPhysicalMaterial*> PhysicalMaterials;
@@ -1855,6 +1864,8 @@ void ULandscapeComponent::UpdatePhysicalMaterialTasks()
 void ULandscapeComponent::UpdateCollisionPhysicalMaterialData(TArray<UPhysicalMaterial*> const& InPhysicalMaterials, TArray<uint8> const& InMaterialIds)
 {
 	TRACE_CPUPROFILER_EVENT_SCOPE(ULandscapeComponent::UpdateCollisionPhysicalMaterialData);
+	ULandscapeHeightfieldCollisionComponent* CollisionComponent = GetCollisionComponent();
+	check(CollisionComponent != nullptr);
 
 	// Copy the physical material array
 	CollisionComponent->PhysicalMaterialRenderObjects = InPhysicalMaterials;
@@ -2561,7 +2572,7 @@ bool ULandscapeInfo::HasUnloadedComponentsInRegion(int32 X1, int32 Y1, int32 X2,
 {
 	bool bResult = false;
 
-	if (LandscapeActor)
+	if (LandscapeActor.IsValid())
 	{
 		UWorld* World = LandscapeActor->GetWorld();
 
@@ -4086,7 +4097,7 @@ FVector ULandscapeInfo::GetLandscapeCenterPos(float& LengthZ, int32 MinX /*= MAX
 				ULandscapeComponent* Comp = XYtoComponentMap.FindRef(FIntPoint(IndexX, IndexY));
 				if (Comp)
 				{
-					ULandscapeHeightfieldCollisionComponent* CollisionComp = Comp->CollisionComponent.Get();
+					ULandscapeHeightfieldCollisionComponent* CollisionComp = Comp->GetCollisionComponent();
 					if (CollisionComp)
 					{
 						uint16* Heights = (uint16*)CollisionComp->CollisionHeightData.Lock(LOCK_READ_ONLY);
@@ -4898,7 +4909,7 @@ ALandscapeProxy* ULandscapeInfo::MoveComponentsToProxy(const TArray<ULandscapeCo
 			TargetSelectedComponents.Add(Component);
 		}
 
-		ULandscapeHeightfieldCollisionComponent* CollisionComp = Component->CollisionComponent.Get();
+		ULandscapeHeightfieldCollisionComponent* CollisionComp = Component->GetCollisionComponent();
 		SelectProxies.Add(CollisionComp->GetLandscapeProxy());
 		if (CollisionComp->GetLandscapeProxy() != LandscapeProxy && (!TargetLevel || CollisionComp->GetLandscapeProxy()->GetOuter() != TargetLevel))
 		{
@@ -5491,7 +5502,7 @@ bool ALandscapeStreamingProxy::CanEditChange(const FProperty* InProperty) const
 		return false;
 	}
 
-	if (InProperty && InProperty->GetFName() == GET_MEMBER_NAME_CHECKED(ALandscapeStreamingProxy, LandscapeActor))
+	if (InProperty && InProperty->GetFName() == GET_MEMBER_NAME_CHECKED(ALandscapeStreamingProxy, LandscapeActorRef))
 	{
 		return !GetWorld()->GetSubsystem<ULandscapeSubsystem>()->IsGridBased();
 	}
@@ -5505,20 +5516,21 @@ void ALandscapeStreamingProxy::PostEditChangeProperty(FPropertyChangedEvent& Pro
 
 	if (PropertyName == FName(TEXT("LandscapeActor")))
 	{
-		if (LandscapeActor && IsValidLandscapeActor(LandscapeActor.Get()))
+		if (LandscapeActorRef && IsValidLandscapeActor(LandscapeActorRef.Get()))
 		{
-			LandscapeGuid = LandscapeActor->GetLandscapeGuid();
+			LandscapeGuid = LandscapeActorRef->GetLandscapeGuid();
 			if (GIsEditor && GetWorld() && !GetWorld()->IsPlayInEditor())
 			{
 				// TODO - only need to refresh the old and new landscape info
 				ULandscapeInfo::RecreateLandscapeInfo(GetWorld(), false);
 				FixupWeightmaps();
 				InitializeProxyLayersWeightmapUsage();
+				// TODO [jonathan.bard] : Call FixupSharedData?
 			}
 		}
 		else
 		{
-			LandscapeActor = nullptr;
+			LandscapeActorRef = nullptr;
 		}
 	}
 	else if (PropertyName == FName(TEXT("LandscapeMaterial")) || PropertyName == FName(TEXT("LandscapeHoleMaterial")) || PropertyName == FName(TEXT("PerLODOverrideMaterials")))
@@ -5581,6 +5593,34 @@ void ALandscapeStreamingProxy::PostRegisterAllComponents()
 {
 	ALandscapeProxy::PostRegisterAllComponents();
 
+#if WITH_EDITORONLY_DATA
+	// If the landscape actor is not set yet and we're transferring the property from the lazy object pointer it was previously stored as to the soft object ptr it is now stored as :
+	if (!LandscapeActorRef)
+	{
+		// Because of how lazy object pointers were made, the only way we can deprecate them is if the object they're pointing to is currently loaded :
+		if (LandscapeActor_DEPRECATED.IsValid())
+		{
+			LandscapeActorRef = LandscapeActor_DEPRECATED.Get();
+			LandscapeActor_DEPRECATED = nullptr;
+			RequestPackageDeprecation();
+		}
+		else if (LandscapeActor_DEPRECATED.IsPending())
+		{
+			FFormatNamedArguments Arguments;
+			Arguments.Add(TEXT("LevelName"), FText::FromString(GetLevel()->GetOutermost()->GetName()));
+			Arguments.Add(TEXT("ProxyName"), FText::FromString(GetName()));
+			FMessageLog("MapCheck").Error()
+				->AddToken(FUObjectToken::Create(this))
+				->AddToken(FTextToken::Create(FText::Format(LOCTEXT("MapCheck_Message_LandscapeLazyObjectPtrDeprecation_Warning", "Landscape proxy {ProxyName} of {LevelName} points to a LandscapeActor that is not currently loaded. This will lose the property upon save. \
+				Please make sure to load the level containing the parent landscape actor prior to {LevelName} so that data deprecation can be performed adequately."), Arguments)))
+				->AddToken(FMapErrorToken::Create(FMapErrors::LandscapeLazyObjectPtrDeprecation_Warning));
+
+			// Show MapCheck window
+			FMessageLog("MapCheck").Open(EMessageSeverity::Warning);
+		}
+	}
+#endif // WITH_EDITORONLY_DATA
+
 	if (LandscapeGuid.IsValid())
 	{
 		ULandscapeInfo* LandscapeInfo = GetLandscapeInfo();
@@ -5635,7 +5675,7 @@ bool ALandscapeStreamingProxy::GetReferencedContentObjects(TArray<UObject*>& Obj
 	Super::GetReferencedContentObjects(Objects);
 
 	// Also return the objects referenced by our parent landscape : 
-	if (LandscapeActor != nullptr)
+	if (const ALandscape* LandscapeActor = GetLandscapeActor())
 	{
 		LandscapeActor->GetReferencedContentObjects(Objects);
 	}
@@ -6058,7 +6098,7 @@ void ALandscapeProxy::ChangedPhysMaterial()
 	{
 		if (LandscapeComponent && LandscapeComponent->IsRegistered())
 		{
-			ULandscapeHeightfieldCollisionComponent* CollisionComponent = LandscapeComponent->CollisionComponent.Get();
+			ULandscapeHeightfieldCollisionComponent* CollisionComponent = LandscapeComponent->GetCollisionComponent();
 			if (CollisionComponent)
 			{
 				LandscapeComponent->UpdateCollisionLayerData();
@@ -6832,17 +6872,23 @@ bool ALandscapeStreamingProxy::IsValidLandscapeActor(ALandscape* Landscape)
 	{
 		if (!Landscape->HasAnyFlags(RF_BeginDestroyed))
 		{
-			if (LandscapeActor.IsNull() && !LandscapeGuid.IsValid())
+			if (!LandscapeActorRef && !LandscapeGuid.IsValid())
 			{
 				return true; // always valid for newly created Proxy
 			}
-			if (((LandscapeActor && LandscapeActor == Landscape)
-				|| (LandscapeActor.IsNull() && LandscapeGuid.IsValid() && LandscapeGuid == Landscape->GetLandscapeGuid()))
-				&& ComponentSizeQuads == Landscape->ComponentSizeQuads
-				&& NumSubsections == Landscape->NumSubsections
-				&& SubsectionSizeQuads == Landscape->SubsectionSizeQuads)
+			if ((LandscapeActorRef && (LandscapeActorRef == Landscape))
+				|| (!LandscapeActorRef && LandscapeGuid.IsValid() && (LandscapeGuid == Landscape->GetLandscapeGuid())))
 			{
-				return true;
+				const bool bCompatibleSize = (ComponentSizeQuads == Landscape->ComponentSizeQuads)
+					&& (NumSubsections == Landscape->NumSubsections)
+					&& (SubsectionSizeQuads == Landscape->SubsectionSizeQuads);
+
+				if (!bCompatibleSize)
+				{
+					UE_LOG(LogLandscape, Warning, TEXT("Landscape streaming proxy %s's setup (Num Quads = %i, Num Subsections = %i, Num Quads per Subsection = %i) is not compatible with landscape actor %s (Num Quads = %i, Num Subsections = %i, Num Quads per Subsection = %i)"), 
+						*GetName(), ComponentSizeQuads, NumSubsections, SubsectionSizeQuads, *Landscape->GetName(), Landscape->ComponentSizeQuads, Landscape->NumSubsections, Landscape->SubsectionSizeQuads);
+				}
+				return bCompatibleSize;
 			}
 		}
 	}
@@ -7123,11 +7169,13 @@ void ALandscapeProxy::RemoveOverlappingComponent(ULandscapeComponent* Component)
 {
 	Modify();
 	Component->Modify();
-	if (Component->CollisionComponent.IsValid() && (Component->CollisionComponent->RenderComponent.Get() == Component || Component->CollisionComponent->RenderComponent.IsNull()))
+
+	ULandscapeHeightfieldCollisionComponent* CollisionComponent = Component->GetCollisionComponent();
+	if ((CollisionComponent != nullptr) && (CollisionComponent->GetRenderComponent() == Component || (CollisionComponent->GetRenderComponent() == nullptr)))
 	{
-		Component->CollisionComponent->Modify();
-		CollisionComponents.Remove(Component->CollisionComponent.Get());
-		Component->CollisionComponent.Get()->DestroyComponent();
+		CollisionComponent->Modify();
+		CollisionComponents.Remove(CollisionComponent);
+		CollisionComponent->DestroyComponent();
 	}
 	LandscapeComponents.Remove(Component);
 	Component->DestroyComponent();
diff --git a/Engine/Source/Runtime/Landscape/Private/LandscapeEditInterface.cpp b/Engine/Source/Runtime/Landscape/Private/LandscapeEditInterface.cpp
--- a/Engine/Source/Runtime/Landscape/Private/LandscapeEditInterface.cpp
+++ b/Engine/Source/Runtime/Landscape/Private/LandscapeEditInterface.cpp
@@ -1743,7 +1743,7 @@ void ULandscapeComponent::DeleteLayerInternal(ULandscapeLayerInfoObject* LayerIn
 			CollisionWeightmapMipData.GetData(),
 			Component->SimpleCollisionMipLevel > Component->CollisionMipLevel ? SimpleCollisionWeightmapMipData.GetData() : nullptr);
 
-		ULandscapeHeightfieldCollisionComponent* LocalCollisionComponent = Component->CollisionComponent.Get();
+		ULandscapeHeightfieldCollisionComponent* LocalCollisionComponent = Component->GetCollisionComponent();
 		if (LocalCollisionComponent)
 		{
 			LocalCollisionComponent->RecreateCollision();
@@ -1764,7 +1764,7 @@ void FLandscapeEditDataInterface::DeleteLayer(ULandscapeLayerInfoObject* LayerIn
 		Component->DeleteLayer(LayerInfo, *this);
 	}
 
-	if (LandscapeInfo->LandscapeActor && LandscapeInfo->LandscapeActor->HasLayersContent())
+	if (LandscapeInfo->LandscapeActor.IsValid() && LandscapeInfo->LandscapeActor->HasLayersContent())
 	{
 		LandscapeInfo->LandscapeActor->RequestLayersContentUpdate(ELandscapeLayerUpdateMode::Update_All);
 	}
@@ -1954,7 +1954,7 @@ void ULandscapeComponent::FillLayer(ULandscapeLayerInfoObject* LayerInfo, FLands
 			CollisionWeightmapMipData.GetData(),
 			Component->SimpleCollisionMipLevel > Component->CollisionMipLevel ? SimpleCollisionWeightmapMipData.GetData() : nullptr);
 
-		ULandscapeHeightfieldCollisionComponent* LocalCollisionComponent = Component->CollisionComponent.Get();
+		ULandscapeHeightfieldCollisionComponent* LocalCollisionComponent = Component->GetCollisionComponent();
 		if (LocalCollisionComponent)
 		{
 			LocalCollisionComponent->RecreateCollision();
@@ -2010,7 +2010,7 @@ void FLandscapeEditDataInterface::FillLayer(ULandscapeLayerInfoObject* LayerInfo
 		}
 	}
 
-	if (LandscapeInfo->LandscapeActor && LandscapeInfo->LandscapeActor->HasLayersContent())
+	if (LandscapeInfo->LandscapeActor.IsValid() && LandscapeInfo->LandscapeActor->HasLayersContent())
 	{
 		LandscapeInfo->LandscapeActor->RequestLayersContentUpdateForceAll();
 	}
@@ -2236,7 +2236,7 @@ void ULandscapeComponent::ReplaceLayer(ULandscapeLayerInfoObject* FromLayerInfo,
 			SimpleCollisionMipLevel > CollisionMipLevel ? SimpleCollisionWeightmapMipData.GetData() : nullptr);
 
 
-		if(ULandscapeHeightfieldCollisionComponent* CollisionComp = CollisionComponent.Get())
+		if(ULandscapeHeightfieldCollisionComponent* CollisionComp = GetCollisionComponent())
 		{
 			CollisionComp->RecreateCollision();
 		}
diff --git a/Engine/Source/Runtime/Landscape/Private/LandscapeEditLayers.cpp b/Engine/Source/Runtime/Landscape/Private/LandscapeEditLayers.cpp
--- a/Engine/Source/Runtime/Landscape/Private/LandscapeEditLayers.cpp
+++ b/Engine/Source/Runtime/Landscape/Private/LandscapeEditLayers.cpp
@@ -8117,7 +8117,7 @@ uint32 ALandscape::UpdateCollisionAndClients(const TArrayView<FLandscapeEditLaye
 		bool bDoUpdateClient = true;
 		if (IsUpdateFlagEnabledForModes(ELandscapeComponentUpdateFlag::Component_Update_Recreate_Collision, ComponentReadbackResult.UpdateModes))
 		{
-			if (ULandscapeHeightfieldCollisionComponent* CollisionComp = LandscapeComponent->CollisionComponent.Get())
+			if (ULandscapeHeightfieldCollisionComponent* CollisionComp = LandscapeComponent->GetCollisionComponent())
 			{
 				FEnableCollisionHashOptimScope Scope(CollisionComp);
 				bDoUpdateClient = CollisionComp->RecreateCollision();
@@ -8128,7 +8128,7 @@ uint32 ALandscape::UpdateCollisionAndClients(const TArrayView<FLandscapeEditLaye
 		{
 			if (!GUndo)
 			{
-				if (ULandscapeHeightfieldCollisionComponent* CollisionComp = LandscapeComponent->CollisionComponent.Get())
+				if (ULandscapeHeightfieldCollisionComponent* CollisionComp = LandscapeComponent->GetCollisionComponent())
 				{
 					FNavigationSystem::UpdateComponentData(*CollisionComp);
 					CollisionComp->SnapFoliageInstances();
diff --git a/Engine/Source/Runtime/Landscape/Private/LandscapeRender.cpp b/Engine/Source/Runtime/Landscape/Private/LandscapeRender.cpp
--- a/Engine/Source/Runtime/Landscape/Private/LandscapeRender.cpp
+++ b/Engine/Source/Runtime/Landscape/Private/LandscapeRender.cpp
@@ -2365,7 +2365,7 @@ void FLandscapeComponentSceneProxy::GetDynamicMeshElements(const TArray<const FS
 	{
 		for (const FSceneView* View : Views)
 		{
-			const FString& LandscapeName = LandscapeComponent->GetLandscapeInfo()->LandscapeActor
+			const FString& LandscapeName = LandscapeComponent->GetLandscapeInfo()->LandscapeActor.IsValid()
 											 ? LandscapeComponent->GetLandscapeInfo()->LandscapeActor->GetName()
 											 : LexToString(LandscapeComponent->GetLandscapeInfo()->LandscapeGuid);
 			const FString& ComponentName = LandscapeComponent->GetName();
diff --git a/Engine/Source/Runtime/Landscape/Private/LandscapeSplineRaster.cpp b/Engine/Source/Runtime/Landscape/Private/LandscapeSplineRaster.cpp
--- a/Engine/Source/Runtime/Landscape/Private/LandscapeSplineRaster.cpp
+++ b/Engine/Source/Runtime/Landscape/Private/LandscapeSplineRaster.cpp
@@ -794,7 +794,7 @@ bool ULandscapeInfo::ApplySplinesInternal(bool bOnlySelected, TScriptInterface<I
 		else
 		{
 			// Recreate collision for modified components and update the navmesh
-			ULandscapeHeightfieldCollisionComponent* CollisionComponent = Component->CollisionComponent.Get();
+			ULandscapeHeightfieldCollisionComponent* CollisionComponent = Component->GetCollisionComponent();
 			if (CollisionComponent)
 			{
 				CollisionComponent->RecreateCollision();
@@ -899,7 +899,7 @@ namespace LandscapeSplineRaster
 			for (ULandscapeComponent* Component : ModifiedComponents)
 			{
 				// Recreate collision for modified components and update the navmesh
-				ULandscapeHeightfieldCollisionComponent* CollisionComponent = Component->CollisionComponent.Get();
+				ULandscapeHeightfieldCollisionComponent* CollisionComponent = Component->GetCollisionComponent();
 				if (CollisionComponent)
 				{
 					CollisionComponent->RecreateCollision();
diff --git a/Engine/Source/Runtime/Landscape/Private/LandscapeSplines.cpp b/Engine/Source/Runtime/Landscape/Private/LandscapeSplines.cpp
--- a/Engine/Source/Runtime/Landscape/Private/LandscapeSplines.cpp
+++ b/Engine/Source/Runtime/Landscape/Private/LandscapeSplines.cpp
@@ -103,7 +103,7 @@ struct FLandscapeFixSplines
 							}
 						});
 
-						if (Pair.Value->LandscapeActor && Pair.Value->LandscapeActor->HasLayersContent())
+						if (Pair.Value->LandscapeActor.IsValid() && Pair.Value->LandscapeActor->HasLayersContent())
 						{
 							Pair.Value->LandscapeActor->RequestSplineLayerUpdate();
 						}
diff --git a/Engine/Source/Runtime/Landscape/Public/LandscapeEdit.h b/Engine/Source/Runtime/Landscape/Public/LandscapeEdit.h
--- a/Engine/Source/Runtime/Landscape/Public/LandscapeEdit.h
+++ b/Engine/Source/Runtime/Landscape/Public/LandscapeEdit.h
@@ -430,7 +430,7 @@ struct FHeightmapAccessor
 				bUpdateNormals = true;
 				for (ULandscapeComponent* Component : Components)
 				{
-					ULandscapeHeightfieldCollisionComponent* CollisionComponent = Component->CollisionComponent.Get();
+					ULandscapeHeightfieldCollisionComponent* CollisionComponent = Component->GetCollisionComponent();
 					if (CollisionComponent && AInstancedFoliageActor::HasFoliageAttached(CollisionComponent))
 					{
 						bUpdateFoliage = true;
@@ -449,7 +449,7 @@ struct FHeightmapAccessor
 
 				for (ULandscapeComponent* Component : Components)
 				{
-					ULandscapeHeightfieldCollisionComponent* CollisionComponent = Component->CollisionComponent.Get();
+					ULandscapeHeightfieldCollisionComponent* CollisionComponent = Component->GetCollisionComponent();
 					if (CollisionComponent)
 					{
 						CollisionComponents.Add(CollisionComponent);
@@ -495,7 +495,7 @@ struct FHeightmapAccessor
 				(*It)->UpdateComponentToWorld();
 
 				// Recreate collision for modified components to update the physical materials
-				ULandscapeHeightfieldCollisionComponent* CollisionComponent = (*It)->CollisionComponent.Get();
+				ULandscapeHeightfieldCollisionComponent* CollisionComponent = (*It)->GetCollisionComponent();
 				if (CollisionComponent)
 				{
 					CollisionComponent->RecreateCollision();
@@ -549,7 +549,7 @@ struct FAlphamapAccessor
 			// Recreate collision for modified components to update the physical materials
 			for (ULandscapeComponent* Component : ModifiedComponents)
 			{
-				ULandscapeHeightfieldCollisionComponent* CollisionComponent = Component->CollisionComponent.Get();
+				ULandscapeHeightfieldCollisionComponent* CollisionComponent = Component->GetCollisionComponent();
 				if (CollisionComponent)
 				{
 					CollisionComponent->RecreateCollision();
