Fixed static analysis

#rnx

diff --git a/Engine/Source/Runtime/Landscape/Private/Landscape.cpp b/Engine/Source/Runtime/Landscape/Private/Landscape.cpp
--- a/Engine/Source/Runtime/Landscape/Private/Landscape.cpp
+++ b/Engine/Source/Runtime/Landscape/Private/Landscape.cpp
@@ -2325,18 +2325,21 @@ void ALandscapeProxy::PostRegisterAllComponents()
 		}
 	}
 #if WITH_EDITOR
-	if (!IsPendingKillPending())
+	if (LandscapeInfo != nullptr)
 	{
-		if (LandscapeGuid.IsValid())
+		if (!IsPendingKillPending())
 		{
-			LandscapeInfo->FixupProxiesTransform();
+			if (LandscapeGuid.IsValid())
+			{
+				LandscapeInfo->FixupProxiesTransform();
+			}
 		}
-	}
 
-	if (bPackageDeprecationNeeded && (LandscapeInfo != nullptr))
-	{
-		LandscapeInfo->RequestPackageDeprecation(this);
-		bPackageDeprecationNeeded = false;
+		if (bPackageDeprecationNeeded)
+		{
+			LandscapeInfo->RequestPackageDeprecation(this);
+			bPackageDeprecationNeeded = false;
+		}
 	}
 #endif // WITH_EDITOR
 }
