Temporary workaround for PLAY-2043. It partially reverts CL 23469595 changes.

diff --git a/Engine/Source/Runtime/Landscape/Private/LandscapeRender.cpp b/Engine/Source/Runtime/Landscape/Private/LandscapeRender.cpp
--- a/Engine/Source/Runtime/Landscape/Private/LandscapeRender.cpp
+++ b/Engine/Source/Runtime/Landscape/Private/LandscapeRender.cpp
@@ -1766,9 +1766,9 @@ void FLandscapeComponentSceneProxy::OnTransformChanged()
 		const float SizeX = FMath::Max(HeightmapResource->GetSizeX(), 1u);
 		const float SizeY = FMath::Max(HeightmapResource->GetSizeY(), 1u);
 		LandscapeParams.HeightmapTextureSize = FVector4f(SizeX, SizeY, 1.f / SizeX, 1.f / SizeY);
-		LandscapeParams.HeightmapTexture = HeightmapResource->TextureRHI.GetReference();
+		LandscapeParams.HeightmapTexture = HeightmapTexture->TextureReference.TextureReferenceRHI;
 		LandscapeParams.HeightmapTextureSampler = TStaticSamplerState<SF_Point>::GetRHI();
-		LandscapeParams.NormalmapTexture = HeightmapResource->TextureRHI.GetReference();
+		LandscapeParams.NormalmapTexture = HeightmapTexture->TextureReference.TextureReferenceRHI;
 		LandscapeParams.NormalmapTextureSampler = TStaticSamplerState<SF_Bilinear>::GetRHI();
 	}
 	else
@@ -1783,7 +1783,7 @@ void FLandscapeComponentSceneProxy::OnTransformChanged()
 	FTextureResource* XYOffsetmapResource = XYOffsetmapTexture ? XYOffsetmapTexture->GetResource() : nullptr;
 	if (XYOffsetmapResource)
 	{
-		LandscapeParams.XYOffsetmapTexture = XYOffsetmapResource->TextureRHI.GetReference();
+		LandscapeParams.XYOffsetmapTexture = XYOffsetmapTexture->TextureReference.TextureReferenceRHI;
 		LandscapeParams.XYOffsetmapTextureSampler = TStaticSamplerState<SF_Point>::GetRHI();
 	}
 	else
