Landscape: Fix truncation warnings, part 3

diff --git a/Engine/Source/Runtime/Landscape/Private/LandscapeEditLayers.cpp b/Engine/Source/Runtime/Landscape/Private/LandscapeEditLayers.cpp
--- a/Engine/Source/Runtime/Landscape/Private/LandscapeEditLayers.cpp
+++ b/Engine/Source/Runtime/Landscape/Private/LandscapeEditLayers.cpp
@@ -374,7 +374,7 @@ public:
 	virtual void InitRHI()
 	{
 		FVertexDeclarationElementList Elements;
-		uint32 Stride = sizeof(FLandscapeLayersVertex);
+		constexpr uint16 Stride = sizeof(FLandscapeLayersVertex);
 		Elements.Add(FVertexElement(0, STRUCT_OFFSET(FLandscapeLayersVertex, Position), VET_Float2, 0, Stride));
 		Elements.Add(FVertexElement(0, STRUCT_OFFSET(FLandscapeLayersVertex, UV), VET_Float2, 1, Stride));
 		VertexDeclarationRHI = PipelineStateCache::GetOrCreateVertexDeclaration(Elements);
@@ -524,7 +524,7 @@ public:
 
 		FVector4f LayerInfo(InParams.LayerAlpha, InParams.LayerVisible ? 1.0f : 0.0f, InParams.LayerBlendMode == LSBM_AlphaBlend ? 1.0f : 0.f, 0.f);
 		FVector4f OutputConfig(InParams.ApplyLayerModifiers ? 1.0f : 0.0f, 0.0f /*unused*/, InParams.ReadHeightmap2 ? 1.0f : 0.0f, InParams.GenerateNormals ? 1.0f : 0.0f);
-		FVector2f TextureSize(InParams.HeightmapSize.X, InParams.HeightmapSize.Y);
+		FVector2f TextureSize(static_cast<float>(InParams.HeightmapSize.X), static_cast<float>(InParams.HeightmapSize.Y));
 
 		SetShaderValue(RHICmdList, RHICmdList.GetBoundPixelShader(), LayerInfoParam, LayerInfo);
 		SetShaderValue(RHICmdList, RHICmdList.GetBoundPixelShader(), OutputConfigParam, OutputConfig);
@@ -578,8 +578,8 @@ public:
 	{
 		SetTextureParameter(RHICmdList, RHICmdList.GetBoundPixelShader(), ReadTexture1Param, ReadTexture1SamplerParam, TStaticSamplerState<SF_Point, AM_Clamp, AM_Clamp, AM_Clamp>::GetRHI(), InParams.ReadHeightmap1->GetResource()->TextureRHI);
 
-		SetShaderValue(RHICmdList, RHICmdList.GetBoundPixelShader(), CurrentMipSizeParam, FVector2f(InParams.CurrentMipSize.X, InParams.CurrentMipSize.Y));
-		SetShaderValue(RHICmdList, RHICmdList.GetBoundPixelShader(), ParentMipSizeParam, FVector2f(InParams.ParentMipSize.X, InParams.ParentMipSize.Y));
+		SetShaderValue(RHICmdList, RHICmdList.GetBoundPixelShader(), CurrentMipSizeParam, FVector2f(static_cast<float>(InParams.CurrentMipSize.X), static_cast<float>(InParams.CurrentMipSize.Y)));
+		SetShaderValue(RHICmdList, RHICmdList.GetBoundPixelShader(), ParentMipSizeParam, FVector2f(static_cast<float>(InParams.ParentMipSize.X), static_cast<float>(InParams.ParentMipSize.Y)));
 		SetShaderValue(RHICmdList, RHICmdList.GetBoundPixelShader(), CurrentMipComponentVertexCountParam, (float)InParams.CurrentMipComponentVertexCount);
 	}
 
@@ -707,8 +707,8 @@ public:
 	{
 		SetTextureParameter(RHICmdList, RHICmdList.GetBoundPixelShader(), ReadTexture1Param, ReadTexture1SamplerParam, TStaticSamplerState<SF_Point, AM_Clamp, AM_Clamp, AM_Clamp>::GetRHI(), InParams.ReadWeightmap1->GetResource()->TextureRHI);
 
-		SetShaderValue(RHICmdList, RHICmdList.GetBoundPixelShader(), CurrentMipSizeParam, FVector2f(InParams.CurrentMipSize.X, InParams.CurrentMipSize.Y));
-		SetShaderValue(RHICmdList, RHICmdList.GetBoundPixelShader(), ParentMipSizeParam, FVector2f(InParams.ParentMipSize.X, InParams.ParentMipSize.Y));
+		SetShaderValue(RHICmdList, RHICmdList.GetBoundPixelShader(), CurrentMipSizeParam, FVector2f(static_cast<float>(InParams.CurrentMipSize.X), static_cast<float>(InParams.CurrentMipSize.Y)));
+		SetShaderValue(RHICmdList, RHICmdList.GetBoundPixelShader(), ParentMipSizeParam, FVector2f(static_cast<float>(InParams.ParentMipSize.X), static_cast<float>(InParams.ParentMipSize.Y)));
 		SetShaderValue(RHICmdList, RHICmdList.GetBoundPixelShader(), CurrentMipComponentVertexCountParam, (float)InParams.CurrentMipComponentVertexCount);
 	}
 
@@ -1634,7 +1634,8 @@ public:
 			GraphicsPSOInit.BoundShaderState.VertexShaderRHI = VertexShader.GetVertexShader();
 			GraphicsPSOInit.BoundShaderState.PixelShaderRHI = PixelShader.GetPixelShader();
 
-			InRHICmdList.SetViewport(View->UnscaledViewRect.Min.X, View->UnscaledViewRect.Min.Y, 0.0f, View->UnscaledViewRect.Max.X, View->UnscaledViewRect.Max.Y, 1.0f);
+			InRHICmdList.SetViewport(static_cast<float>(View->UnscaledViewRect.Min.X), static_cast<float>(View->UnscaledViewRect.Min.Y), 0.0f,
+			                         static_cast<float>(View->UnscaledViewRect.Max.X), static_cast<float>(View->UnscaledViewRect.Max.Y), 1.0f);
 
 			InRHICmdList.ApplyCachedRenderTargets(GraphicsPSOInit);
 			SetGraphicsPipelineState(InRHICmdList, GraphicsPSOInit, 0);
@@ -1652,7 +1653,7 @@ public:
 			GraphicsPSOInit.BoundShaderState.VertexShaderRHI = VertexShader.GetVertexShader();
 			GraphicsPSOInit.BoundShaderState.PixelShaderRHI = PixelShader.GetPixelShader();
 
-			InRHICmdList.SetViewport(0.0f, 0.0f, 0.0f, WriteRenderTargetSize.X, WriteRenderTargetSize.Y, 1.0f);
+			InRHICmdList.SetViewport(0.0f, 0.0f, 0.0f, static_cast<float>(WriteRenderTargetSize.X), static_cast<float>(WriteRenderTargetSize.Y), 1.0f);
 
 			InRHICmdList.ApplyCachedRenderTargets(GraphicsPSOInit);
 			SetGraphicsPipelineState(InRHICmdList, GraphicsPSOInit, 0);
@@ -2377,7 +2378,6 @@ void ALandscape::ValidateProxyLayersWeightmapUsage() const
 		return;
 	}
 
-	bool bIsReady = true;
 	Info->ForAllLandscapeProxies([=](const ALandscapeProxy* Proxy)
 	{
 		Proxy->ValidateProxyLayersWeightmapUsage();
@@ -2650,27 +2650,27 @@ void ALandscape::DrawWeightmapComponentsToRenderTarget(const FString& InDebugNam
 		{
 		case ERTDrawingType::RTAtlas:
 		{
-			GenerateLayersRenderQuadsAtlas(InSectionBaseList[i], FVector2D(WeightmapScaleBias), SubsectionSizeQuads, WeightmapReadTextureSize, WeightmapWriteTextureSize, TriangleList);
+			GenerateLayersRenderQuadsAtlas(InSectionBaseList[i], FVector2D(WeightmapScaleBias), static_cast<float>(SubsectionSizeQuads), WeightmapReadTextureSize, WeightmapWriteTextureSize, TriangleList);
 		} break;
 
 		case ERTDrawingType::RTAtlasToNonAtlas:
 		{
-			GenerateLayersRenderQuadsAtlasToNonAtlas(InSectionBaseList[i], FVector2D(WeightmapScaleBias), SubsectionSizeQuads, WeightmapReadTextureSize, WeightmapWriteTextureSize, TriangleList);
+			GenerateLayersRenderQuadsAtlasToNonAtlas(InSectionBaseList[i], FVector2D(WeightmapScaleBias), static_cast<float>(SubsectionSizeQuads), WeightmapReadTextureSize, WeightmapWriteTextureSize, TriangleList);
 		} break;
 
 		case ERTDrawingType::RTNonAtlas:
 		{
-			GenerateLayersRenderQuadsNonAtlas(InSectionBaseList[i], FVector2D(WeightmapScaleBias), SubsectionSizeQuads, WeightmapReadTextureSize, WeightmapWriteTextureSize, TriangleList);
+			GenerateLayersRenderQuadsNonAtlas(InSectionBaseList[i], FVector2D(WeightmapScaleBias), static_cast<float>(SubsectionSizeQuads), WeightmapReadTextureSize, WeightmapWriteTextureSize, TriangleList);
 		} break;
 
 		case ERTDrawingType::RTNonAtlasToAtlas:
 		{
-			GenerateLayersRenderQuadsNonAtlasToAtlas(InSectionBaseList[i], FVector2D(WeightmapScaleBias), SubsectionSizeQuads, WeightmapReadTextureSize, WeightmapWriteTextureSize, TriangleList);
+			GenerateLayersRenderQuadsNonAtlasToAtlas(InSectionBaseList[i], FVector2D(WeightmapScaleBias), static_cast<float>(SubsectionSizeQuads), WeightmapReadTextureSize, WeightmapWriteTextureSize, TriangleList);
 		} break;
 
 		case ERTDrawingType::RTMips:
 		{
-			GenerateLayersRenderQuadsMip(InSectionBaseList[i], FVector2D(WeightmapScaleBias), SubsectionSizeQuads, WeightmapReadTextureSize, WeightmapWriteTextureSize, InMipRender, TriangleList);
+			GenerateLayersRenderQuadsMip(InSectionBaseList[i], FVector2D(WeightmapScaleBias), static_cast<float>(SubsectionSizeQuads), WeightmapReadTextureSize, WeightmapWriteTextureSize, InMipRender, TriangleList);
 		} break;
 
 		default:
@@ -2718,7 +2718,7 @@ void ALandscape::DrawWeightmapComponentsToRenderTarget(const FString& InDebugNam
 
 	for (ULandscapeComponent* Component : InComponentsToDraw)
 	{
-		FVector2f WeightmapScaleBias(Component->WeightmapScaleBias.Z, Component->WeightmapScaleBias.W);
+		FVector2f WeightmapScaleBias(static_cast<float>(Component->WeightmapScaleBias.Z), static_cast<float>(Component->WeightmapScaleBias.W));
 		WeightmapScaleBiasList.Add(WeightmapScaleBias);
 
 		FIntPoint ComponentSectionBase = Component->GetSectionBase() - InLandscapeBase;
@@ -2744,8 +2744,8 @@ void ALandscape::DrawWeightmapComponentToRenderTargetMips(const TArray<FVector2f
 
 	for (const FVector2f& TexturePosition : InTexturePositionsToDraw)
 	{
-		FVector2f PositionOffset(FMath::RoundToInt(TexturePosition.X / LocalComponentSizeVerts), FMath::RoundToInt(TexturePosition.Y / LocalComponentSizeVerts));
-		SectionBaseToDraw.Add(FIntPoint(PositionOffset.X * LocalComponentSizeQuad, PositionOffset.Y * LocalComponentSizeQuad));
+		FVector2f PositionOffset(static_cast<float>(FMath::RoundToInt(TexturePosition.X / LocalComponentSizeVerts)), static_cast<float>(FMath::RoundToInt(TexturePosition.Y / LocalComponentSizeVerts)));
+		SectionBaseToDraw.Add(FIntPoint(static_cast<int32>(PositionOffset.X * LocalComponentSizeQuad), static_cast<int32>(PositionOffset.Y * LocalComponentSizeQuad)));
 	}
 
 	FVector2f WeightmapScaleBias(0.0f, 0.0f); // we dont need a scale bias for mip drawing
@@ -2757,7 +2757,7 @@ void ALandscape::DrawWeightmapComponentToRenderTargetMips(const TArray<FVector2f
 		if (WriteMipRT != nullptr)
 		{
 			DrawWeightmapComponentsToRenderTarget(FString::Printf(TEXT("LS Weight: %s = -> %s Mips %d"), *ReadMipRT->GetName(), *WriteMipRT->GetName(), CurrentMip),
-				SectionBaseToDraw, WeightmapScaleBias, nullptr, ReadMipRT, nullptr, WriteMipRT, ERTDrawingType::RTMips, InClearRTWrite, InShaderParams, CurrentMip);
+				SectionBaseToDraw, WeightmapScaleBias, nullptr, ReadMipRT, nullptr, WriteMipRT, ERTDrawingType::RTMips, InClearRTWrite, InShaderParams, static_cast<uint8>(CurrentMip));
 			++CurrentMip;
 		}
 
@@ -2789,7 +2789,7 @@ void ALandscape::DrawHeightmapComponentsToRenderTargetMips(const TArray<ULandsca
 		if (WriteMipRT != nullptr)
 		{
 			DrawHeightmapComponentsToRenderTarget(FString::Printf(TEXT("LS Height: %s = -> %s CombinedAtlasWithMips %d"), *ReadMipRT->GetName(), *WriteMipRT->GetName(), CurrentMip),
-				InComponentsToDraw, InLandscapeBase, ReadMipRT, nullptr, WriteMipRT, ERTDrawingType::RTMips, InClearRTWrite, InShaderParams, CurrentMip);
+				InComponentsToDraw, InLandscapeBase, ReadMipRT, nullptr, WriteMipRT, ERTDrawingType::RTMips, InClearRTWrite, InShaderParams, static_cast<uint8>(CurrentMip));
 			++CurrentMip;
 		}
 
@@ -2819,34 +2819,34 @@ void ALandscape::DrawHeightmapComponentsToRenderTarget(const FString& InDebugNam
 
 	for (ULandscapeComponent* Component : InComponentsToDraw)
 	{
-		FVector2f HeightmapScaleBias(Component->HeightmapScaleBias.Z, Component->HeightmapScaleBias.W);
+		FVector2f HeightmapScaleBias(static_cast<float>(Component->HeightmapScaleBias.Z), static_cast<float>(Component->HeightmapScaleBias.W));
 		FIntPoint ComponentSectionBase = Component->GetSectionBase() - InLandscapeBase;
 
 		switch (InDrawType)
 		{
 			case ERTDrawingType::RTAtlas:
 			{
-				GenerateLayersRenderQuadsAtlas(ComponentSectionBase, FVector2D(HeightmapScaleBias), SubsectionSizeQuads, HeightmapReadTextureSize, HeightmapWriteTextureSize, TriangleList);
+				GenerateLayersRenderQuadsAtlas(ComponentSectionBase, FVector2D(HeightmapScaleBias), static_cast<float>(SubsectionSizeQuads), HeightmapReadTextureSize, HeightmapWriteTextureSize, TriangleList);
 			} break;
 
 			case ERTDrawingType::RTAtlasToNonAtlas:
 			{
-				GenerateLayersRenderQuadsAtlasToNonAtlas(ComponentSectionBase, FVector2D(HeightmapScaleBias), SubsectionSizeQuads, HeightmapReadTextureSize, HeightmapWriteTextureSize, TriangleList);
+				GenerateLayersRenderQuadsAtlasToNonAtlas(ComponentSectionBase, FVector2D(HeightmapScaleBias), static_cast<float>(SubsectionSizeQuads), HeightmapReadTextureSize, HeightmapWriteTextureSize, TriangleList);
 			} break;
 
 			case ERTDrawingType::RTNonAtlas:
 			{
-				GenerateLayersRenderQuadsNonAtlas(ComponentSectionBase, FVector2D(HeightmapScaleBias), SubsectionSizeQuads, HeightmapReadTextureSize, HeightmapWriteTextureSize, TriangleList);
+				GenerateLayersRenderQuadsNonAtlas(ComponentSectionBase, FVector2D(HeightmapScaleBias), static_cast<float>(SubsectionSizeQuads), HeightmapReadTextureSize, HeightmapWriteTextureSize, TriangleList);
 			} break;
 
 			case ERTDrawingType::RTNonAtlasToAtlas:
 			{
-				GenerateLayersRenderQuadsNonAtlasToAtlas(ComponentSectionBase, FVector2D(HeightmapScaleBias), SubsectionSizeQuads, HeightmapReadTextureSize, HeightmapWriteTextureSize, TriangleList);
+				GenerateLayersRenderQuadsNonAtlasToAtlas(ComponentSectionBase, FVector2D(HeightmapScaleBias), static_cast<float>(SubsectionSizeQuads), HeightmapReadTextureSize, HeightmapWriteTextureSize, TriangleList);
 			} break;
 
 			case ERTDrawingType::RTMips:
 			{
-				GenerateLayersRenderQuadsMip(ComponentSectionBase, FVector2D(HeightmapScaleBias), SubsectionSizeQuads, HeightmapReadTextureSize, HeightmapWriteTextureSize, InMipRender, TriangleList);
+				GenerateLayersRenderQuadsMip(ComponentSectionBase, FVector2D(HeightmapScaleBias), static_cast<float>(SubsectionSizeQuads), HeightmapReadTextureSize, HeightmapWriteTextureSize, InMipRender, TriangleList);
 			} break;
 
 		default:
@@ -2886,27 +2886,33 @@ void ALandscape::DrawHeightmapComponentsToRenderTarget(const FString& InDebugNam
 
 void ALandscape::GenerateLayersRenderQuad(const FIntPoint& InVertexPosition, float InVertexSize, const FVector2D& InUVStart, const FVector2D& InUVSize, TArray<FLandscapeLayersTriangle>& OutTriangles) const
 {
-	FLandscapeLayersTriangle Tri1;
-
-	Tri1.V0.Position = FVector2f(InVertexPosition.X, InVertexPosition.Y);
-	Tri1.V1.Position = FVector2f(InVertexPosition.X + InVertexSize, InVertexPosition.Y);
-	Tri1.V2.Position = FVector2f(InVertexPosition.X + InVertexSize, InVertexPosition.Y + InVertexSize);
+	// Set min/max values for rectangle in XY and UV.
+	const float X[] = { static_cast<float>(InVertexPosition.X), static_cast<float>(InVertexPosition.X) + InVertexSize };
+	const float Y[] = { static_cast<float>(InVertexPosition.Y), static_cast<float>(InVertexPosition.Y) + InVertexSize };
+	const float U[] = { static_cast<float>(InUVStart.X), static_cast<float>(InUVStart.X + InUVSize.X) };
+	const float V[] = { static_cast<float>(InUVStart.Y), static_cast<float>(InUVStart.Y + InUVSize.Y) };
 
-	Tri1.V0.UV = FVector2f(InUVStart.X, InUVStart.Y);
-	Tri1.V1.UV = FVector2f(InUVStart.X + InUVSize.X, InUVStart.Y);
-	Tri1.V2.UV = FVector2f(InUVStart.X + InUVSize.X, InUVStart.Y + InUVSize.Y);
-	OutTriangles.Add(Tri1);
+	// Helper function for creating a vertex from given min/max indices.
+	auto SetVertex = [&X, &Y, &U, &V](FLandscapeLayersVertex& Vertex, const int32 Index1, const int32 Index2)
+	{
+		Vertex.Position.X = X[Index1];
+		Vertex.Position.Y = Y[Index2];
+		Vertex.UV.X = U[Index1];
+		Vertex.UV.Y = V[Index2];
+	};
 
-	FLandscapeLayersTriangle Tri2;
-	Tri2.V0.Position = FVector2f(InVertexPosition.X + InVertexSize, InVertexPosition.Y + InVertexSize);
-	Tri2.V1.Position = FVector2f(InVertexPosition.X, InVertexPosition.Y + InVertexSize);
-	Tri2.V2.Position = FVector2f(InVertexPosition.X, InVertexPosition.Y);
+	FLandscapeLayersTriangle Tri;
 
-	Tri2.V0.UV = FVector2f(InUVStart.X + InUVSize.X, InUVStart.Y + InUVSize.Y);
-	Tri2.V1.UV = FVector2f(InUVStart.X, InUVStart.Y + InUVSize.Y);
-	Tri2.V2.UV = FVector2f(InUVStart.X, InUVStart.Y);
+	// Create first triangle.
+	SetVertex(Tri.V0, 0, 0);
+	SetVertex(Tri.V1, 1, 0);
+	SetVertex(Tri.V2, 1, 1);
+	OutTriangles.Add(Tri);
 
-	OutTriangles.Add(Tri2);
+	// Create second triangle; V0 is identical to the previous triangle.
+	SetVertex(Tri.V1, 1, 1);
+	SetVertex(Tri.V2, 0, 1);
+	OutTriangles.Add(Tri);
 }
 
 void ALandscape::GenerateLayersRenderQuadsAtlas(const FIntPoint& InSectionBase, const FVector2D& InScaleBias, float InSubSectionSizeQuad, const FIntPoint& InReadSize, const FIntPoint& InWriteSize, TArray<FLandscapeLayersTriangle>& OutTriangles) const
@@ -2914,9 +2920,9 @@ void ALandscape::GenerateLayersRenderQuadsAtlas(const FIntPoint& InSectionBase,
 	FIntPoint ComponentSectionBase = InSectionBase;
 	FIntPoint UVComponentSectionBase = InSectionBase;
 
-	int32 SubsectionSizeVerts = InSubSectionSizeQuad + 1;
-	int32 LocalComponentSizeQuad = InSubSectionSizeQuad * NumSubsections;
-	int32 LocalComponentSizeVerts = SubsectionSizeVerts * NumSubsections;
+	const int32 SubsectionSizeVerts = static_cast<int32>(InSubSectionSizeQuad) + 1;
+	const int32 LocalComponentSizeQuad = static_cast<int32>(InSubSectionSizeQuad * NumSubsections);
+	const int32 LocalComponentSizeVerts = SubsectionSizeVerts * NumSubsections;
 
 	FVector2D PositionOffset(FMath::RoundToInt((float)ComponentSectionBase.X / (float)LocalComponentSizeQuad), FMath::RoundToInt((float)ComponentSectionBase.Y / (float)LocalComponentSizeQuad));
 	FVector2D ComponentsPerTexture(FMath::RoundToInt((float)InWriteSize.X / (float)LocalComponentSizeQuad), FMath::RoundToInt((float)InWriteSize.Y / (float)LocalComponentSizeQuad));
@@ -2927,15 +2933,19 @@ void ALandscape::GenerateLayersRenderQuadsAtlas(const FIntPoint& InSectionBase,
 		{
 			if (ComponentsPerTexture.X > 1.0f)
 			{
-				UVComponentSectionBase.X = PositionOffset.X * LocalComponentSizeVerts;
+				UVComponentSectionBase.X = static_cast<int32>(PositionOffset.X * LocalComponentSizeVerts);
 			}
 			else
 			{
-				UVComponentSectionBase.X -= (UVComponentSectionBase.X + LocalComponentSizeQuad > InWriteSize.X) ? FMath::FloorToInt((PositionOffset.X / ComponentsPerTexture.X)) * ComponentsPerTexture.X * LocalComponentSizeQuad : 0;
+				UVComponentSectionBase.X -= (UVComponentSectionBase.X + LocalComponentSizeQuad > InWriteSize.X)
+					                            ? static_cast<int32>(FMath::FloorToInt(PositionOffset.X / ComponentsPerTexture.X) * ComponentsPerTexture.X * LocalComponentSizeQuad)
+					                            : 0;
 			}
 		}
 
-		ComponentSectionBase.X -= (ComponentSectionBase.X + LocalComponentSizeQuad > InWriteSize.X) ? FMath::FloorToInt((PositionOffset.X / ComponentsPerTexture.X)) * ComponentsPerTexture.X * LocalComponentSizeQuad : 0;
+		ComponentSectionBase.X -= (ComponentSectionBase.X + LocalComponentSizeQuad > InWriteSize.X)
+			                          ? static_cast<int32>(FMath::FloorToInt(PositionOffset.X / ComponentsPerTexture.X) * ComponentsPerTexture.X * LocalComponentSizeQuad)
+			                          : 0;
 		PositionOffset.X = ComponentSectionBase.X / LocalComponentSizeQuad;
 	}
 
@@ -2945,20 +2955,24 @@ void ALandscape::GenerateLayersRenderQuadsAtlas(const FIntPoint& InSectionBase,
 		{
 			if (ComponentsPerTexture.Y > 1.0f)
 			{
-				UVComponentSectionBase.Y = PositionOffset.Y * LocalComponentSizeVerts;
+				UVComponentSectionBase.Y = static_cast<int32>(PositionOffset.Y * LocalComponentSizeVerts);
 			}
 			else
 			{
-				UVComponentSectionBase.Y -= (UVComponentSectionBase.Y + LocalComponentSizeQuad > InWriteSize.Y) ? FMath::FloorToInt((PositionOffset.Y / ComponentsPerTexture.Y)) * ComponentsPerTexture.Y * LocalComponentSizeQuad : 0;
+				UVComponentSectionBase.Y -= (UVComponentSectionBase.Y + LocalComponentSizeQuad > InWriteSize.Y)
+					                            ? static_cast<int32>(FMath::FloorToInt((PositionOffset.Y / ComponentsPerTexture.Y)) * ComponentsPerTexture.Y * LocalComponentSizeQuad)
+					                            : 0;
 			}
 		}
 
-		ComponentSectionBase.Y -= (ComponentSectionBase.Y + LocalComponentSizeQuad > InWriteSize.Y) ? FMath::FloorToInt((PositionOffset.Y / ComponentsPerTexture.Y)) * ComponentsPerTexture.Y * LocalComponentSizeQuad : 0;
+		ComponentSectionBase.Y -= (ComponentSectionBase.Y + LocalComponentSizeQuad > InWriteSize.Y)
+			                          ? static_cast<int32>(FMath::FloorToInt((PositionOffset.Y / ComponentsPerTexture.Y)) * ComponentsPerTexture.Y * LocalComponentSizeQuad)
+			                          : 0;
 		PositionOffset.Y = ComponentSectionBase.Y / LocalComponentSizeQuad;
 	}
 
-	ComponentSectionBase.X = PositionOffset.X * LocalComponentSizeVerts;
-	ComponentSectionBase.Y = PositionOffset.Y * LocalComponentSizeVerts;
+	ComponentSectionBase.X = static_cast<int32>(PositionOffset.X * LocalComponentSizeVerts);
+	ComponentSectionBase.Y = static_cast<int32>(PositionOffset.Y * LocalComponentSizeVerts);
 
 	FVector2D UVSize((float)SubsectionSizeVerts / (float)InReadSize.X, (float)SubsectionSizeVerts / (float)InReadSize.Y);
 	FIntPoint SubSectionSectionBase;
@@ -2991,15 +3005,15 @@ void ALandscape::GenerateLayersRenderQuadsAtlas(const FIntPoint& InSectionBase,
 				UVStart.Y = InScaleBias.Y + UVSize.Y * (float)SubY;
 			}
 
-			GenerateLayersRenderQuad(SubSectionSectionBase, SubsectionSizeVerts, UVStart, UVSize, OutTriangles);
+			GenerateLayersRenderQuad(SubSectionSectionBase, static_cast<float>(SubsectionSizeVerts), UVStart, UVSize, OutTriangles);
 		}
 	}
 }
 
 void ALandscape::GenerateLayersRenderQuadsMip(const FIntPoint& InSectionBase, const FVector2D& InScaleBias, float InSubSectionSizeQuad, const FIntPoint& InReadSize, const FIntPoint& InWriteSize, uint8 InCurrentMip, TArray<FLandscapeLayersTriangle>& OutTriangles) const
 {
-	int32 SubsectionSizeVerts = InSubSectionSizeQuad + 1;
-	int32 LocalComponentSizeQuad = InSubSectionSizeQuad * NumSubsections;
+	int32 SubsectionSizeVerts = static_cast<int32>(InSubSectionSizeQuad) + 1;
+	int32 LocalComponentSizeQuad = static_cast<int32>(InSubSectionSizeQuad) * NumSubsections;
 	int32 LocalComponentSizeVerts = SubsectionSizeVerts * NumSubsections;
 	int32 MipSubsectionSizeVerts = SubsectionSizeVerts >> InCurrentMip;
 	int32 MipLocalComponentSizeVerts = MipSubsectionSizeVerts * NumSubsections;
@@ -3007,8 +3021,8 @@ void ALandscape::GenerateLayersRenderQuadsMip(const FIntPoint& InSectionBase, co
 	FVector2D PositionOffset(FMath::RoundToInt((float)InSectionBase.X / (float)LocalComponentSizeQuad), FMath::RoundToInt((float)InSectionBase.Y / (float)LocalComponentSizeQuad));
 	FVector2D ComponentsPerTexture(FMath::RoundToInt((float)InWriteSize.X / (float)LocalComponentSizeQuad), FMath::RoundToInt((float)InWriteSize.Y / (float)LocalComponentSizeQuad));
 
-	FIntPoint ComponentSectionBase(PositionOffset.X * MipLocalComponentSizeVerts, PositionOffset.Y * MipLocalComponentSizeVerts);
-	FIntPoint UVComponentSectionBase(PositionOffset.X * LocalComponentSizeVerts, PositionOffset.Y * LocalComponentSizeVerts);
+	FIntPoint ComponentSectionBase(static_cast<int32>(PositionOffset.X * MipLocalComponentSizeVerts), static_cast<int32>(PositionOffset.Y * MipLocalComponentSizeVerts));
+	FIntPoint UVComponentSectionBase(static_cast<int32>(PositionOffset.X * LocalComponentSizeVerts), static_cast<int32>(PositionOffset.Y * LocalComponentSizeVerts));
 	FVector2D UVSize((float)(SubsectionSizeVerts >> (InCurrentMip - 1)) / (float)InReadSize.X, (float)(SubsectionSizeVerts >> (InCurrentMip - 1)) / (float)InReadSize.Y);
 	FIntPoint SubSectionSectionBase;
 
@@ -3022,23 +3036,23 @@ void ALandscape::GenerateLayersRenderQuadsMip(const FIntPoint& InSectionBase, co
 			// Offset for this component's data in texture
 			FVector2D UVStart(((float)(UVComponentSectionBase.X >> (InCurrentMip - 1)) / (float)InReadSize.X) + UVSize.X * (float)SubX, ((float)(UVComponentSectionBase.Y >> (InCurrentMip - 1)) / (float)InReadSize.Y) + UVSize.Y * (float)SubY);
 
-			GenerateLayersRenderQuad(SubSectionSectionBase, MipSubsectionSizeVerts, UVStart, UVSize, OutTriangles);
+			GenerateLayersRenderQuad(SubSectionSectionBase, static_cast<float>(MipSubsectionSizeVerts), UVStart, UVSize, OutTriangles);
 		}
 	}
 }
 
 void ALandscape::GenerateLayersRenderQuadsAtlasToNonAtlas(const FIntPoint& InSectionBase, const FVector2D& InScaleBias, float InSubSectionSizeQuad, const FIntPoint& InReadSize, const FIntPoint& InWriteSize, TArray<struct FLandscapeLayersTriangle>& OutTriangles) const
 {
-	int32 SubsectionSizeVerts = InSubSectionSizeQuad + 1;
+	int32 SubsectionSizeVerts = static_cast<int32>(InSubSectionSizeQuad) + 1;
 	FVector2D UVSize((float)SubsectionSizeVerts / (float)InReadSize.X, (float)SubsectionSizeVerts / (float)InReadSize.Y);
 
 	for (int8 SubY = 0; SubY < NumSubsections; ++SubY)
 	{
 		for (int8 SubX = 0; SubX < NumSubsections; ++SubX)
 		{
-			FIntPoint SubSectionSectionBase(InSectionBase.X + InSubSectionSizeQuad * SubX, InSectionBase.Y + InSubSectionSizeQuad * SubY);
+			FIntPoint SubSectionSectionBase(static_cast<int32>(InSectionBase.X + InSubSectionSizeQuad * SubX), static_cast<int32>(InSectionBase.Y + InSubSectionSizeQuad * SubY));
 			FVector2D PositionOffset(FMath::RoundToInt(SubSectionSectionBase.X / InSubSectionSizeQuad), FMath::RoundToInt(SubSectionSectionBase.Y / InSubSectionSizeQuad));
-			FIntPoint UVComponentSectionBase(PositionOffset.X * SubsectionSizeVerts, PositionOffset.Y * SubsectionSizeVerts);
+			FIntPoint UVComponentSectionBase(static_cast<int32>(PositionOffset.X * SubsectionSizeVerts), static_cast<int32>(PositionOffset.Y * SubsectionSizeVerts));
 
 			// Offset for this component's data in texture
 			FVector2D UVStart;
@@ -3061,7 +3075,7 @@ void ALandscape::GenerateLayersRenderQuadsAtlasToNonAtlas(const FIntPoint& InSec
 				UVStart.Y = InScaleBias.Y + UVSize.Y * (float)SubY;
 			}
 
-			GenerateLayersRenderQuad(SubSectionSectionBase, SubsectionSizeVerts, UVStart, UVSize, OutTriangles);
+			GenerateLayersRenderQuad(SubSectionSectionBase, static_cast<float>(SubsectionSizeVerts), UVStart, UVSize, OutTriangles);
 		}
 	}
 }
@@ -3071,7 +3085,7 @@ void ALandscape::GenerateLayersRenderQuadsNonAtlas(const FIntPoint& InSectionBas
 	// We currently only support drawing in non atlas mode with the same texture size
 	check(InReadSize.X == InWriteSize.X && InReadSize.Y == InWriteSize.Y);
 
-	int32 SubsectionSizeVerts = InSubSectionSizeQuad + 1;
+	int32 SubsectionSizeVerts = static_cast<int32>(InSubSectionSizeQuad) + 1;
 
 	FVector2D UVSize((float)SubsectionSizeVerts / (float)InReadSize.X, (float)SubsectionSizeVerts / (float)InReadSize.Y);
 
@@ -3081,23 +3095,23 @@ void ALandscape::GenerateLayersRenderQuadsNonAtlas(const FIntPoint& InSectionBas
 		{
 			FIntPoint SubSectionSectionBase(InSectionBase.X + SubsectionSizeQuads * SubX, InSectionBase.Y + SubsectionSizeQuads * SubY);
 			FVector2D PositionOffset(FMath::RoundToInt(SubSectionSectionBase.X / InSubSectionSizeQuad), FMath::RoundToInt(SubSectionSectionBase.Y / InSubSectionSizeQuad));
-			FIntPoint UVComponentSectionBase(PositionOffset.X * InSubSectionSizeQuad, PositionOffset.Y * InSubSectionSizeQuad);
+			FIntPoint UVComponentSectionBase(static_cast<int32>(PositionOffset.X * InSubSectionSizeQuad), static_cast<int32>(PositionOffset.Y * InSubSectionSizeQuad));
 
 			// Offset for this component's data in texture
 			FVector2D UVStart(((float)UVComponentSectionBase.X / (float)InReadSize.X), ((float)UVComponentSectionBase.Y / (float)InReadSize.Y));
-			GenerateLayersRenderQuad(SubSectionSectionBase, SubsectionSizeVerts, UVStart, UVSize, OutTriangles);
+			GenerateLayersRenderQuad(SubSectionSectionBase, static_cast<float>(SubsectionSizeVerts), UVStart, UVSize, OutTriangles);
 		}
 	}
 }
 
 void ALandscape::GenerateLayersRenderQuadsNonAtlasToAtlas(const FIntPoint& InSectionBase, const FVector2D& InScaleBias, float InSubSectionSizeQuad, const FIntPoint& InReadSize, const FIntPoint& InWriteSize, TArray<struct FLandscapeLayersTriangle>& OutTriangles) const
 {
-	int32 SubsectionSizeVerts = InSubSectionSizeQuad + 1;
-	int32 LocalComponentSizeQuad = InSubSectionSizeQuad * NumSubsections;
+	int32 SubsectionSizeVerts = static_cast<int32>(InSubSectionSizeQuad) + 1;
+	int32 LocalComponentSizeQuad = static_cast<int32>(InSubSectionSizeQuad) * NumSubsections;
 	int32 LocalComponentSizeVerts = SubsectionSizeVerts * NumSubsections;
 
 	FVector2D PositionOffset(FMath::RoundToInt((float)InSectionBase.X / (float)LocalComponentSizeQuad), FMath::RoundToInt((float)InSectionBase.Y / (float)LocalComponentSizeQuad));
-	FIntPoint ComponentSectionBase(PositionOffset.X * LocalComponentSizeVerts, PositionOffset.Y * LocalComponentSizeVerts);
+	FIntPoint ComponentSectionBase(static_cast<int32>(PositionOffset.X * LocalComponentSizeVerts), static_cast<int32>(PositionOffset.Y * LocalComponentSizeVerts));
 	FVector2D UVSize((float)SubsectionSizeVerts / (float)InReadSize.X, (float)SubsectionSizeVerts / (float)InReadSize.Y);
 
 	FIntPoint SubSectionSectionBase;
@@ -3114,7 +3128,7 @@ void ALandscape::GenerateLayersRenderQuadsNonAtlasToAtlas(const FIntPoint& InSec
 			float ScaleBiasW = (float)InSectionBase.Y / (float)InReadSize.Y;
 			FVector2D UVStart(ScaleBiasZ + ((float)InSubSectionSizeQuad / (float)InReadSize.X) * (float)SubX, ScaleBiasW + ((float)InSubSectionSizeQuad / (float)InReadSize.Y) * (float)SubY);
 
-			GenerateLayersRenderQuad(SubSectionSectionBase, SubsectionSizeVerts, UVStart, UVSize, OutTriangles);
+			GenerateLayersRenderQuad(SubSectionSectionBase, static_cast<float>(SubsectionSizeVerts), UVStart, UVSize, OutTriangles);
 		}
 	}
 }
@@ -4083,7 +4097,7 @@ namespace EditLayersHeightmapLocalMerge_RenderThread
 			check(SizeZ > 0);
 			// TODO [jonathan.bard] : change to PF_R8G8 once edit layers heightmaps are stored as such :
 			FRDGTextureDesc Desc = FRDGTextureDesc::Create2DArray(InLocalMergeInfo.ComponentSizeVerts, PF_B8G8R8A8, FClearValueBinding::None, TexCreate_RenderTargetable | TexCreate_ShaderResource, 
-				SizeZ, /*InNumMips = */1, /*InNumSamples = */1);
+				static_cast<uint16>(SizeZ), /*InNumMips = */1, /*InNumSamples = */1);
 			OutResources.EditLayersHeightmapsTextureArray = GraphBuilder.CreateTexture(Desc, TEXT("LandscapeEditLayersHeightmapsTextureArray"));
 			OutResources.EditLayersHeightmapsTextureArraySRV = GraphBuilder.CreateSRV(FRDGTextureSRVDesc::Create(OutResources.EditLayersHeightmapsTextureArray));
 		}
@@ -4092,7 +4106,7 @@ namespace EditLayersHeightmapLocalMerge_RenderThread
 			int32 SizeZ = InLocalMergeInfo.MaxNumComponentsToRenderPerResolveComponentBatch;
 			// We only need 2 channels since this only stores the (packed) height :
 			FRDGTextureDesc Desc = FRDGTextureDesc::Create2DArray(InLocalMergeInfo.ComponentSizeVerts, PF_R8G8, FClearValueBinding::None, TexCreate_RenderTargetable | TexCreate_TargetArraySlicesIndependently | TexCreate_ShaderResource, 
-				SizeZ, /*InNumMips = */1, /*InNumSamples = */1);
+				static_cast<uint16>(SizeZ), /*InNumMips = */1, /*InNumSamples = */1);
 			// Create 2 texture arrays "Merged" and "Stitched" (ScratchMergedHeightmapTextureArray will be copied/merged into ScratchStitchedHeightmapTextureArray, slice by slice)
 			OutResources.ScratchMergedHeightmapTextureArray = GraphBuilder.CreateTexture(Desc, TEXT("LandscapeEditLayersMergedHeightmapTextureArray"));
 			OutResources.ScratchStitchedHeightmapTextureArray = GraphBuilder.CreateTexture(Desc, TEXT("LandscapeEditLayersStitchedHeightmapTextureArray"));
@@ -4180,7 +4194,7 @@ namespace EditLayersHeightmapLocalMerge_RenderThread
 				check(IndexInBatch < RDGResources.ScratchMergedHeightmapTextureArray->Desc.ArraySize);
 
 				FLandscapeLayersHeightmapsMergeEditLayersPS::FParameters* MergeEditLayersPSParams = GraphBuilder.AllocParameters<FLandscapeLayersHeightmapsMergeEditLayersPS::FParameters>();
-				MergeEditLayersPSParams->RenderTargets[0] = FRenderTargetBinding(RDGResources.ScratchMergedHeightmapTextureArray, ERenderTargetLoadAction::ENoAction, /*InMipIndex = */0, /*InArraySlice = */IndexInBatch);
+				MergeEditLayersPSParams->RenderTargets[0] = FRenderTargetBinding(RDGResources.ScratchMergedHeightmapTextureArray, ERenderTargetLoadAction::ENoAction, /*InMipIndex = */0, /*InArraySlice = */static_cast<int16>(IndexInBatch));
 				MergeEditLayersPSParams->InNumEditLayers = NumLayers;
 				MergeEditLayersPSParams->InEditLayersTextures = RDGResources.EditLayersHeightmapsTextureArraySRV;
 				MergeEditLayersPSParams->InEditLayersMergeInfos = RDGResources.EditLayersMergeInfosBufferSRV;
@@ -4207,7 +4221,7 @@ namespace EditLayersHeightmapLocalMerge_RenderThread
 			// Now, stitch the heightmap using the StitchHeightmapPS and output the result to a single slice in ScratchStitchedHeightmapTextureArray: 
 			{
 				FLandscapeLayersHeightmapsStitchHeightmapPS::FParameters* StitchHeightmapPSParams = GraphBuilder.AllocParameters<FLandscapeLayersHeightmapsStitchHeightmapPS::FParameters>();
-				StitchHeightmapPSParams->RenderTargets[0] = FRenderTargetBinding(RDGResources.ScratchStitchedHeightmapTextureArray, ERenderTargetLoadAction::ENoAction, /*InMipIndex = */0, /*InArraySlice = */IndexInBatch);
+				StitchHeightmapPSParams->RenderTargets[0] = FRenderTargetBinding(RDGResources.ScratchStitchedHeightmapTextureArray, ERenderTargetLoadAction::ENoAction, /*InMipIndex = */0, /*InArraySlice = */static_cast<int16>(IndexInBatch));
 				StitchHeightmapPSParams->InSourceTextureSize = FUintVector2(InLocalMergeInfo.ComponentSizeVerts.X, InLocalMergeInfo.ComponentSizeVerts.Y);
 				StitchHeightmapPSParams->InNumSubsections = InLocalMergeInfo.NumSubsections;
 				StitchHeightmapPSParams->InSourceHeightmaps = RDGResources.ScratchMergedHeightmapTextureArraySRV;
@@ -4304,7 +4318,7 @@ namespace EditLayersHeightmapLocalMerge_RenderThread
 				CurrentMipSubregionSize.Y >>= 1;
 
 				FLandscapeLayersHeightmapsGenerateMipsPS::FParameters* GenerateMipsPSParams = GraphBuilder.AllocParameters<FLandscapeLayersHeightmapsGenerateMipsPS::FParameters>();
-				GenerateMipsPSParams->RenderTargets[0] = FRenderTargetBinding(TrackedTexture->ScratchTextureRef, ERenderTargetLoadAction::ENoAction, MipLevel);
+				GenerateMipsPSParams->RenderTargets[0] = FRenderTargetBinding(TrackedTexture->ScratchTextureRef, ERenderTargetLoadAction::ENoAction, static_cast<uint8>(MipLevel));
 				GenerateMipsPSParams->InCurrentMipSubregionSize = FUintVector2(CurrentMipSubregionSize.X, CurrentMipSubregionSize.Y);
 				GenerateMipsPSParams->InNumSubsections = InLocalMergeInfo.NumSubsections;
 				GenerateMipsPSParams->InSourceHeightmap = TrackedTexture->ScratchTextureMipsSRVRefs[MipLevel - 1];
@@ -4340,7 +4354,7 @@ namespace EditLayersHeightmapLocalMerge_RenderThread
 
 		// Even if we have heightmaps of different sizes to handle, we only need one empty heightmap to copy from (whose size is MaxHeightmapSize) :
 		// TODO [jonathan.bard] : change to PF_R8G8 once edit layers heightmaps are stored as such :
-		FRDGTextureDesc Desc = FRDGTextureDesc::Create2D(InLocalMergeInfo.MaxHeightmapSize, PF_B8G8R8A8, FClearValueBinding(ClearHeightColor), TexCreate_ShaderResource | TexCreate_RenderTargetable, InLocalMergeInfo.MaxHeightmapNumMips, /*InNumSamples = */1);
+		FRDGTextureDesc Desc = FRDGTextureDesc::Create2D(InLocalMergeInfo.MaxHeightmapSize, PF_B8G8R8A8, FClearValueBinding(ClearHeightColor), TexCreate_ShaderResource | TexCreate_RenderTargetable, static_cast<uint8>(InLocalMergeInfo.MaxHeightmapNumMips), /*InNumSamples = */1);
 		FRDGTextureRef EmptyTexture = GraphBuilder.CreateTexture(Desc, TEXT("LandscapeEditLayersEmptyHeightmap"));
 
 		FRDGTextureClearInfo ClearInfo;
@@ -4431,7 +4445,7 @@ void ALandscape::PrepareLayersHeightmapsLocalMergeRenderThreadData(const FUpdate
 			OutRenderThreadData.MaxHeightmapSize = TextureSize.ComponentMax(OutRenderThreadData.MaxHeightmapSize);
 			OutRenderThreadData.MaxHeightmapNumMips = FMath::Max((int32)FMath::CeilLogTwo(TextureSize.GetMin()) + 1, OutRenderThreadData.MaxHeightmapNumMips);
 
-			FIntPoint HeightmapOffset(Component->HeightmapScaleBias.Z * TextureSize.X, Component->HeightmapScaleBias.W * TextureSize.Y);
+			FIntPoint HeightmapOffset(static_cast<int32>(Component->HeightmapScaleBias.Z * TextureSize.X), static_cast<int32>(Component->HeightmapScaleBias.W * TextureSize.Y));
 			// Effective area of the texture affecting this component (because of texture sharing) :
 			const FIntRect ComponentTextureSubregion(HeightmapOffset, HeightmapOffset + OutRenderThreadData.ComponentSizeVerts);
 
@@ -4480,8 +4494,8 @@ void ALandscape::PrepareLayersHeightmapsLocalMergeRenderThreadData(const FUpdate
 		for (ULandscapeComponent* Component : InUpdateLayersContentContext.LandscapeComponentsHeightmapsToResolve)
 		{
 			UTexture2D* ComponentHeightmap = Component->GetHeightmap();
-			const int32 HeightmapOffsetX = Component->HeightmapScaleBias.Z * (float)ComponentHeightmap->Source.GetSizeX();
-			const int32 HeightmapOffsetY = Component->HeightmapScaleBias.W * (float)ComponentHeightmap->Source.GetSizeY();
+			const int32 HeightmapOffsetX = static_cast<int32>(Component->HeightmapScaleBias.Z * ComponentHeightmap->Source.GetSizeX());
+			const int32 HeightmapOffsetY = static_cast<int32>(Component->HeightmapScaleBias.W * ComponentHeightmap->Source.GetSizeY());
 			// Effective area of the texture affecting this component (because of texture sharing) :
 			const FIntRect ComponentTextureSubregion(FIntPoint(HeightmapOffsetX, HeightmapOffsetY), FIntPoint(HeightmapOffsetX, HeightmapOffsetY) + OutRenderThreadData.ComponentSizeVerts);
 
@@ -4709,7 +4723,7 @@ int32 ALandscape::PerformLayersHeightmapsGlobalMerge(const FUpdateLayersContentC
 
 			FIntPoint ComponentSectionBase = Component->GetSectionBase() - LandscapeBaseQuads;
 			FVector2D SourcePositionOffset(FMath::RoundToInt((float)ComponentSectionBase.X / ComponentSizeQuad), FMath::RoundToInt((float)ComponentSectionBase.Y / ComponentSizeQuad));
-			FIntPoint ComponentVertexPosition = FIntPoint(SourcePositionOffset.X * ComponentSizeVerts, SourcePositionOffset.Y * ComponentSizeVerts);
+			FIntPoint ComponentVertexPosition = FIntPoint(static_cast<int32>(SourcePositionOffset.X * ComponentSizeVerts), static_cast<int32>(SourcePositionOffset.Y * ComponentSizeVerts));
 			ALandscapeProxy* Proxy = Component->GetLandscapeProxy();
 
 			if (Index == INDEX_NONE)
@@ -5143,8 +5157,8 @@ void ALandscape::UpdateHeightDirtyData(ULandscapeComponent* InLandscapeComponent
 	TUniquePtr<uint8[]> DirtyData = MakeUnique<uint8[]>(DirtyDataSize);
 	const int32 SizeU = InHeightmap->Source.GetSizeX();
 	const int32 SizeV = InHeightmap->Source.GetSizeY();
-	const int32 HeightmapOffsetX = InLandscapeComponent->HeightmapScaleBias.Z * SizeU;
-	const int32 HeightmapOffsetY = InLandscapeComponent->HeightmapScaleBias.W * SizeV;
+	const int32 HeightmapOffsetX = static_cast<int32>(InLandscapeComponent->HeightmapScaleBias.Z * SizeU);
+	const int32 HeightmapOffsetY = static_cast<int32>(InLandscapeComponent->HeightmapScaleBias.W * SizeV);
 	const uint8 DirtyHeight = 1 << 0;
 	LandscapeEdit.GetDirtyData(X1, Y1, X2, Y2, DirtyData.Get(), 0);
 
@@ -5193,8 +5207,8 @@ void ALandscape::OnDirtyHeightmap(FTextureToComponentHelper const& MapHelper, UT
 
 				const int32 SizeU = InHeightmap->Source.GetSizeX();
 				const int32 SizeV = InHeightmap->Source.GetSizeY();
-				const int32 HeightmapOffsetX = Component->HeightmapScaleBias.Z * (float)SizeU;
-				const int32 HeightmapOffsetY = Component->HeightmapScaleBias.W * (float)SizeV;
+				const int32 HeightmapOffsetX = static_cast<int32>(Component->HeightmapScaleBias.Z * SizeU);
+				const int32 HeightmapOffsetY = static_cast<int32>(Component->HeightmapScaleBias.W * SizeV);
 				const int32 ComponentWidth = (SubsectionSizeQuads + 1) * NumSubsections;
 				FIntRect SubRegion(HeightmapOffsetX, HeightmapOffsetY, HeightmapOffsetX + ComponentWidth, HeightmapOffsetY + ComponentWidth);
 
@@ -5418,7 +5432,7 @@ void ALandscape::PrepareComponentDataToExtractMaterialLayersCS(const TArray<ULan
 
 				FIntPoint ComponentSectionBase = Component->GetSectionBase() - InLandscapeBase;
 				FVector2D SourcePositionOffset(FMath::RoundToInt((float)ComponentSectionBase.X / LocalComponentSizeQuad), FMath::RoundToInt((float)ComponentSectionBase.Y / LocalComponentSizeQuad));
-				FIntPoint SourceComponentVertexPosition = FIntPoint(SourcePositionOffset.X * LocalComponentSizeVerts, SourcePositionOffset.Y * LocalComponentSizeVerts);
+				FIntPoint SourceComponentVertexPosition = FIntPoint(static_cast<int32>(SourcePositionOffset.X * LocalComponentSizeVerts), static_cast<int32>(SourcePositionOffset.Y * LocalComponentSizeVerts));
 
 				FLandscapeLayersCopyTextureParams& CopyTextureParams = DeferredCopyTextures.Add_GetRef(FLandscapeLayersCopyTextureParams(*LayerWeightmap->GetName(), LayerWeightmap->GetResource(), FString::Printf(TEXT("%s WeightmapScratchTexture"), *InLayer.Name.ToString()), InOutTextureData));
 				// Only copy the size that's actually needed : 
@@ -5444,7 +5458,7 @@ void ALandscape::PrepareComponentDataToExtractMaterialLayersCS(const TArray<ULan
 							SourceComponentVertexPosition,
 							0,
 							WeightmapLayerAllocation.WeightmapTextureChannel,
-							FIntPoint(DestPositionOffset.X * LocalComponentSizeVerts, DestPositionOffset.Y * LocalComponentSizeVerts)
+							FIntPoint(static_cast<int32>(DestPositionOffset.X * LocalComponentSizeVerts), static_cast<int32>(DestPositionOffset.Y * LocalComponentSizeVerts))
 						};
 
 						if (WeightmapLayerAllocation.LayerInfo == ALandscapeProxy::VisibilityLayer)
@@ -5554,8 +5568,8 @@ void ALandscape::PrepareComponentDataToPackMaterialLayersCS(int32 InCurrentWeigh
 						int32 LocalComponentSizeVerts = (ChannelComponent->SubsectionSizeQuads + 1) * NumSubsections;
 						FVector2D PositionOffset(FMath::RoundToInt((float)ComponentSectionBase.X / LocalComponentSizeQuad), FMath::RoundToInt((float)ComponentSectionBase.Y / LocalComponentSizeQuad));
 
-						Data.ComponentVertexPositionX[WeightmapChannelIndex] = PositionOffset.X * LocalComponentSizeVerts;
-						Data.ComponentVertexPositionY[WeightmapChannelIndex] = PositionOffset.Y * LocalComponentSizeVerts;
+						Data.ComponentVertexPositionX[WeightmapChannelIndex] = static_cast<int32>(PositionOffset.X * LocalComponentSizeVerts);
+						Data.ComponentVertexPositionY[WeightmapChannelIndex] = static_cast<int32>(PositionOffset.Y * LocalComponentSizeVerts);
 
 						Data.WeightmapChannelToProcess[WeightmapChannelIndex] = ChannelLayerAllocation.WeightmapTextureChannel;
 
@@ -6110,7 +6124,7 @@ namespace EditLayersWeightmapLocalMerge_RenderThread
 			check(SizeZ > 0);
 
 			// TODO [chris.tchou] this texture does not have to be a render target, but RDG does not support transient shader-resource-only/copy-populated textures yet -- see issue https://jira.it.epicgames.com/browse/UE-162198
-			FRDGTextureDesc Desc = FRDGTextureDesc::Create2DArray(InLocalMergeInfo.ComponentSizeVerts, PF_B8G8R8A8, FClearValueBinding::None, TexCreate_ShaderResource | TexCreate_RenderTargetable, SizeZ, /*InNumMips = */1, /*InNumSamples = */1);
+			FRDGTextureDesc Desc = FRDGTextureDesc::Create2DArray(InLocalMergeInfo.ComponentSizeVerts, PF_B8G8R8A8, FClearValueBinding::None, TexCreate_ShaderResource | TexCreate_RenderTargetable, static_cast<uint16>(SizeZ), /*InNumMips = */1, /*InNumSamples = */1);
 			OutResources.EditLayersWeightmapsTextureArray = GraphBuilder.CreateTexture(Desc, TEXT("LandscapeEditLayersWeightmapsTextureArray"));
 			OutResources.EditLayersWeightmapsTextureArraySRV = GraphBuilder.CreateSRV(FRDGTextureSRVDesc::Create(OutResources.EditLayersWeightmapsTextureArray));
 		}
@@ -6124,7 +6138,7 @@ namespace EditLayersWeightmapLocalMerge_RenderThread
 			check(SizeZ > 0);
 			for (int32 Index = 0; Index < InLocalMergeInfo.MaxNumWeightmapArraysPerResolveTextureBatch; ++Index)
 			{
-				FRDGTextureDesc Desc = FRDGTextureDesc::Create2DArray(InLocalMergeInfo.ComponentSizeVerts, PF_G8, FClearValueBinding::None, TexCreate_ShaderResource | TexCreate_RenderTargetable | TexCreate_TargetArraySlicesIndependently, SizeZ, /*InNumMips = */1, /*InNumSamples = */1);
+				FRDGTextureDesc Desc = FRDGTextureDesc::Create2DArray(InLocalMergeInfo.ComponentSizeVerts, PF_G8, FClearValueBinding::None, TexCreate_ShaderResource | TexCreate_RenderTargetable | TexCreate_TargetArraySlicesIndependently, static_cast<uint16>(SizeZ), /*InNumMips = */1, /*InNumSamples = */1);
 				FRDGTextureRef& TextureRef = OutResources.ScratchPaintLayerWeightmapTextureArrays.Add_GetRef(GraphBuilder.CreateTexture(Desc, TEXT("LandscapeEditLayersScratchPaintLayerWeightmapTextureArray")));
 				OutResources.ScratchPaintLayerWeightmapTextureArraysSRV.Add(GraphBuilder.CreateSRV(FRDGTextureSRVDesc::Create(TextureRef)));
 			}
@@ -6133,7 +6147,7 @@ namespace EditLayersWeightmapLocalMerge_RenderThread
 		{
 			// Allocate a single scratch texture will all of its mip for each individual texture we want to resolve :
 			check(InLocalMergeInfo.NumMips > 0);
-			FRDGTextureDesc Desc = FRDGTextureDesc::Create2D(InLocalMergeInfo.ComponentSizeVerts, PF_B8G8R8A8, FClearValueBinding::None, TexCreate_ShaderResource | TexCreate_RenderTargetable, InLocalMergeInfo.NumMips, /*InNumSamples = */1);
+			FRDGTextureDesc Desc = FRDGTextureDesc::Create2D(InLocalMergeInfo.ComponentSizeVerts, PF_B8G8R8A8, FClearValueBinding::None, TexCreate_ShaderResource | TexCreate_RenderTargetable, static_cast<uint8>(InLocalMergeInfo.NumMips), /*InNumSamples = */1);
 			OutResources.ScratchFinalWeightmapTexture = GraphBuilder.CreateTexture(Desc, TEXT("LandscapeEditLayersScratchFinalWeightmapTexture"));
 			for (int32 MipLevel = 0; MipLevel < InLocalMergeInfo.NumMips; ++MipLevel)
 			{
@@ -6239,7 +6253,7 @@ namespace EditLayersWeightmapLocalMerge_RenderThread
 
 				FLandscapeLayersWeightmapsMergeEditLayersPS::FParameters* MergeEditLayersPSParams = GraphBuilder.AllocParameters<FLandscapeLayersWeightmapsMergeEditLayersPS::FParameters>();
 				// We'll write to a single slice of the texture array for this component, since we're acting paint layer by paint layer here :
-				MergeEditLayersPSParams->RenderTargets[0] = FRenderTargetBinding(ScratchTextureArrayRef, ERenderTargetLoadAction::ENoAction, /*InMipIndex = */0, /*InArraySlice = */ComponentPaintLayerIndex);
+				MergeEditLayersPSParams->RenderTargets[0] = FRenderTargetBinding(ScratchTextureArrayRef, ERenderTargetLoadAction::ENoAction, /*InMipIndex = */0, /*InArraySlice = */static_cast<int16>(ComponentPaintLayerIndex));
 				MergeEditLayersPSParams->InNumEditLayers = ComponentPaintLayerRenderInfo.VisibleEditLayerInfos.Num();
 				MergeEditLayersPSParams->InPackedWeightmaps = RDGResources.EditLayersWeightmapsTextureArraySRV;
 				MergeEditLayersPSParams->InEditLayersMergeInfos = RDGResources.EditLayerMergeInfosBufferSRV;
@@ -6320,7 +6334,7 @@ namespace EditLayersWeightmapLocalMerge_RenderThread
 
 			// Read from scratch weightmap texture (mip N - 1) -> write to scratch weightmap texture (mip N) :
 			FLandscapeLayersWeightmapsGenerateMipsPS::FParameters* GenerateMipsPSParams = GraphBuilder.AllocParameters<FLandscapeLayersWeightmapsGenerateMipsPS::FParameters>();
-			GenerateMipsPSParams->RenderTargets[0] = FRenderTargetBinding(RDGResources.ScratchFinalWeightmapTexture, ERenderTargetLoadAction::ENoAction, MipLevel);
+			GenerateMipsPSParams->RenderTargets[0] = FRenderTargetBinding(RDGResources.ScratchFinalWeightmapTexture, ERenderTargetLoadAction::ENoAction, static_cast<uint8>(MipLevel));
 			GenerateMipsPSParams->InCurrentMipSize = FUintVector2(CurrentMipSize.X, CurrentMipSize.Y);
 			GenerateMipsPSParams->InNumSubsections = InLocalMergeInfo.NumSubsections;
 			GenerateMipsPSParams->InSourceWeightmap = RDGResources.ScratchFinalWeightmapTextureMipsSRV[MipLevel - 1];
@@ -7104,7 +7118,9 @@ int32 ALandscape::PerformLayersWeightmapsGlobalMerge(FUpdateLayersContentContext
 							continue;
 						}
 
-						FIntPoint TextureTopLeftPositionInAtlas(WeightmapTextureOutputOffset[NextTextureIndexToProcess - StartTextureIndex].X, WeightmapTextureOutputOffset[NextTextureIndexToProcess - StartTextureIndex].Y);
+						FIntPoint TextureTopLeftPositionInAtlas(
+							static_cast<int32>(WeightmapTextureOutputOffset[NextTextureIndexToProcess - StartTextureIndex].X),
+							static_cast<int32>(WeightmapTextureOutputOffset[NextTextureIndexToProcess - StartTextureIndex].Y));
 
 						int32 CurrentMip = 0;
 						const int32 TextureSizeX = WeightmapTexture->GetResource()->GetSizeX();
@@ -7125,7 +7141,7 @@ int32 ALandscape::PerformLayersWeightmapsGlobalMerge(FUpdateLayersContentContext
 								// Copy from the composited texture's position to the top-left corner of the heightmap
 								CopyTextureParams.SourcePosition.X = TextureTopLeftPositionInAtlas.X >> CurrentMip;
 								CopyTextureParams.SourcePosition.Y = TextureTopLeftPositionInAtlas.Y >> CurrentMip;
-								CopyTextureParams.DestMip = CurrentMip;
+								CopyTextureParams.DestMip = static_cast<uint8>(CurrentMip);
 								++CurrentMip;
 							}
 						}
@@ -7326,7 +7342,7 @@ void ALandscape::UpdateLayersMaterialInstances(const TArray<ULandscapeComponent*
 
 		for (int32 LODIndex = 0; LODIndex <= MaxLOD; ++LODIndex)
 		{
-			UMaterialInterface* CurrentMaterial = Component->GetLandscapeMaterial(LODIndex);
+			UMaterialInterface* CurrentMaterial = Component->GetLandscapeMaterial(static_cast<int8>(LODIndex));
 
 			// if we have a LOD0 override, do not let the base material override it, it should override everything!
 			if (CurrentMaterial == BaseMaterial && BaseMaterial != LOD0Material)
@@ -7343,9 +7359,9 @@ void ALandscape::UpdateLayersMaterialInstances(const TArray<ULandscapeComponent*
 			else
 			{
 				int32 AddedIndex = NewMaterialPerLOD.Num();
-				NewMaterialPerLOD.Add(CurrentMaterial, LODIndex);
-				Component->LODIndexToMaterialIndex[LODIndex] = AddedIndex;
-				LastLODIndex = AddedIndex;
+				NewMaterialPerLOD.Add(CurrentMaterial, static_cast<int8>(LODIndex));
+				Component->LODIndexToMaterialIndex[LODIndex] = static_cast<int8>(AddedIndex);
+				LastLODIndex = static_cast<int8>(AddedIndex);
 			}
 		}
 
@@ -8665,7 +8681,7 @@ void ALandscape::SetLayerLocked(int32 InLayerIndex, bool bLocked)
 
 uint8 ALandscape::GetLayerCount() const
 {
-	return LandscapeLayers.Num();
+	return static_cast<uint8>(LandscapeLayers.Num());
 }
 
 FLandscapeLayer* ALandscape::GetLayer(int32 InLayerIndex)
@@ -8752,7 +8768,7 @@ void ALandscape::DeleteLayer(int32 InLayerIndex)
 
 void ALandscape::CollapseLayer(int32 InLayerIndex)
 {
-	FScopedSlowTask SlowTask(GetLandscapeInfo()->XYtoComponentMap.Num(), LOCTEXT("Landscape_CollapseLayer_SlowWork", "Collapsing Layer..."));
+	FScopedSlowTask SlowTask(static_cast<float>(GetLandscapeInfo()->XYtoComponentMap.Num()), LOCTEXT("Landscape_CollapseLayer_SlowWork", "Collapsing Layer..."));
 	SlowTask.MakeDialog();
 	TArray<bool> BackupVisibility;
 	TArray<bool> BackupBrushVisibility;
diff --git a/Engine/Source/Runtime/Landscape/Private/Materials/MaterialExpressionLandscapeLayerCoords.cpp b/Engine/Source/Runtime/Landscape/Private/Materials/MaterialExpressionLandscapeLayerCoords.cpp
--- a/Engine/Source/Runtime/Landscape/Private/Materials/MaterialExpressionLandscapeLayerCoords.cpp
+++ b/Engine/Source/Runtime/Landscape/Private/Materials/MaterialExpressionLandscapeLayerCoords.cpp
@@ -67,7 +67,6 @@ int32 UMaterialExpressionLandscapeLayerCoords::Compile(class FMaterialCompiler*
 	const float Cos = FMath::Cos(MappingRotation * PI / 180.0f);
 	const float Sin = FMath::Sin(MappingRotation * PI / 180.0f);
 
-	int32 ResultUV = INDEX_NONE;
 	int32 TransformedUV = Compiler->Add(
 		Compiler->Mul(RealScale,
 		Compiler->AppendVector(
@@ -76,8 +75,6 @@ int32 UMaterialExpressionLandscapeLayerCoords::Compile(class FMaterialCompiler*
 		),
 		Compiler->Constant2(MappingPanU, MappingPanV)
 		);
-
-	ResultUV = TransformedUV;
 	
 	return TransformedUV;
 }
