Updates LazyObjectPtr deprecation to silently fixup the data when possible.
Fixes LandscapeActor property assignement for LandscapeStreamingProxy.

diff --git a/Engine/Source/Runtime/Landscape/Classes/LandscapeInfo.h b/Engine/Source/Runtime/Landscape/Classes/LandscapeInfo.h
--- a/Engine/Source/Runtime/Landscape/Classes/LandscapeInfo.h
+++ b/Engine/Source/Runtime/Landscape/Classes/LandscapeInfo.h
@@ -262,10 +262,6 @@ public:
 
 	LANDSCAPE_API void ModifyObject(UObject* InObject, bool bAlwaysMarkDirty = true);
 	LANDSCAPE_API void MarkObjectDirty(UObject* InObject);
-
-	LANDSCAPE_API void RequestPackageDeprecation(UObject* InObject);
-	LANDSCAPE_API void ShowForceUpdateMapCheckError(UObject* InOutermost);
-	LANDSCAPE_API void ForceUpdateDeprecatedObjects();
 #endif //WITH_EDITOR
 
 	/**
diff --git a/Engine/Source/Runtime/Landscape/Classes/LandscapeProxy.h b/Engine/Source/Runtime/Landscape/Classes/LandscapeProxy.h
--- a/Engine/Source/Runtime/Landscape/Classes/LandscapeProxy.h
+++ b/Engine/Source/Runtime/Landscape/Classes/LandscapeProxy.h
@@ -773,7 +773,6 @@ public:
 
 #if WITH_EDITOR
 	FOnLandscapeProxyComponentDataChanged OnComponentDataChanged;
-	bool bPackageDeprecationNeeded = false;
 #endif
 
 	// Blueprint functions
@@ -1309,8 +1308,6 @@ public:
 	LANDSCAPE_API bool CanHaveLayersContent() const;
 
 	LANDSCAPE_API virtual void UpdateCachedHasLayersContent(bool InCheckComponentDataIntegrity = false);
-	
-	LANDSCAPE_API void RequestPackageDeprecation() { bPackageDeprecationNeeded = true; };
 
 protected:
 	friend class ALandscape;
diff --git a/Engine/Source/Runtime/Landscape/Private/Landscape.cpp b/Engine/Source/Runtime/Landscape/Private/Landscape.cpp
--- a/Engine/Source/Runtime/Landscape/Private/Landscape.cpp
+++ b/Engine/Source/Runtime/Landscape/Private/Landscape.cpp
@@ -1228,11 +1228,6 @@ void ULandscapeComponent::PostLoad()
 	{
 		CollisionComponentRef = CollisionComponent_DEPRECATED.Get();
 		CollisionComponent_DEPRECATED = nullptr;
-
-		if (ALandscapeProxy* Proxy = GetLandscapeProxy())
-		{
-			Proxy->RequestPackageDeprecation();
-		}
 	}
 #endif // !WITH_EDITORONLY_DATA
 
@@ -2334,13 +2329,6 @@ void ALandscapeProxy::PostRegisterAllComponents()
 				LandscapeInfo->FixupProxiesTransform();
 			}
 		}
-
-		if (bPackageDeprecationNeeded)
-		{
-			LandscapeInfo->RequestPackageDeprecation(this);
-			bPackageDeprecationNeeded = false;
-		}
-	}
 #endif // WITH_EDITOR
 }
 
@@ -3947,48 +3935,6 @@ void ULandscapeInfo::ModifyObject(UObject* InObject, bool bAlwaysMarkDirty)
 	}
 }
 
-void ULandscapeInfo::RequestPackageDeprecation(UObject* InObject)
-{
-	check(InObject != nullptr);
-	// This is either the first time or the user already deprecated a set of object previously, we need to display the MapCheck error.
-	// We do not display the message during cook as the data will be fixedup in the produced asset, and user cannot click on the action.
-	if (PackagesToDeprecate.IsEmpty() && !GIsCookerLoadingPackage)
-	{
-		UObject* Outermost = InObject->GetOutermostObject();
-		ShowForceUpdateMapCheckError(Outermost);
-	}
-
-	PackagesToDeprecate.Add(InObject);
-}
-
-void ULandscapeInfo::ShowForceUpdateMapCheckError(UObject* InOutermost)
-{
-	check(InOutermost != nullptr)
-
-	FMessageLog("MapCheck").Warning()
-		->AddToken(FUObjectToken::Create(InOutermost))
-		->AddToken(FTextToken::Create(LOCTEXT("ULandscapeInfo_ShowForceUpdateMapCheckError_Message_ForceUpdateNeeded", "Landscape packages are out of date and need updating for landscape to be loaded properly.")))
-		->AddToken(FMapErrorToken::Create(TEXT("ULandscapeInfo_ShowForceUpdateMapCheckError_MapError_ForceUpdateNeeded")))
-		->AddToken(FActionToken::Create(LOCTEXT("ULandscapeInfo_ShowForceUpdateMapCheckError_ActionName_ForceUpdateNeeded", "Update Landscape Packages"), FText(),
-			FOnActionTokenExecuted::CreateUObject(this, &ULandscapeInfo::ForceUpdateDeprecatedObjects), true));
-
-	// Show MapCheck window
-	FMessageLog("MapCheck").Open(EMessageSeverity::Warning);
-}
-
-void ULandscapeInfo::ForceUpdateDeprecatedObjects()
-{
-	for (const TWeakObjectPtr<UObject>& ObjectToDeprecate : PackagesToDeprecate)
-	{
-		if (ObjectToDeprecate.IsValid())
-		{
-			ObjectToDeprecate->MarkPackageDirty();
-		}
-	}
-
-	PackagesToDeprecate.Empty();
-}
-
 ALandscapeProxy* ULandscapeInfo::GetLandscapeProxyForLevel(ULevel* Level) const
 {
 	ALandscapeProxy* LandscapeProxy = nullptr;
diff --git a/Engine/Source/Runtime/Landscape/Private/LandscapeCollision.cpp b/Engine/Source/Runtime/Landscape/Private/LandscapeCollision.cpp
--- a/Engine/Source/Runtime/Landscape/Private/LandscapeCollision.cpp
+++ b/Engine/Source/Runtime/Landscape/Private/LandscapeCollision.cpp
@@ -2166,11 +2166,6 @@ void ULandscapeHeightfieldCollisionComponent::PostLoad()
 	{
 		RenderComponentRef = RenderComponent_DEPRECATED.Get();
 		RenderComponent_DEPRECATED = nullptr;
-		
-		if (ALandscapeProxy* Proxy = GetLandscapeProxy())
-		{
-			Proxy->RequestPackageDeprecation();
-		}
 	}
 #endif // !WITH_EDITORONLY_DATA
 
diff --git a/Engine/Source/Runtime/Landscape/Private/LandscapeEdit.cpp b/Engine/Source/Runtime/Landscape/Private/LandscapeEdit.cpp
--- a/Engine/Source/Runtime/Landscape/Private/LandscapeEdit.cpp
+++ b/Engine/Source/Runtime/Landscape/Private/LandscapeEdit.cpp
@@ -5513,7 +5513,7 @@ void ALandscapeStreamingProxy::PostEditChangeProperty(FPropertyChangedEvent& Pro
 {
 	const FName PropertyName = PropertyChangedEvent.MemberProperty ? PropertyChangedEvent.MemberProperty->GetFName() : NAME_None;
 
-	if (PropertyName == FName(TEXT("LandscapeActor")))
+	if (PropertyName == FName(TEXT("LandscapeActorRef")))
 	{
 		if (LandscapeActorRef && IsValidLandscapeActor(LandscapeActorRef.Get()))
 		{
@@ -5601,7 +5601,6 @@ void ALandscapeStreamingProxy::PostRegisterAllComponents()
 		{
 			LandscapeActorRef = LandscapeActor_DEPRECATED.Get();
 			LandscapeActor_DEPRECATED = nullptr;
-			RequestPackageDeprecation();
 		}
 		else if (LandscapeActor_DEPRECATED.IsPending())
 		{
@@ -5611,7 +5610,7 @@ void ALandscapeStreamingProxy::PostRegisterAllComponents()
 			FMessageLog("MapCheck").Error()
 				->AddToken(FUObjectToken::Create(this))
 				->AddToken(FTextToken::Create(FText::Format(LOCTEXT("MapCheck_Message_LandscapeLazyObjectPtrDeprecation_Warning", "Landscape proxy {ProxyName} of {LevelName} points to a LandscapeActor that is not currently loaded. This will lose the property upon save. \
-				Please make sure to load the level containing the parent landscape actor prior to {LevelName} so that data deprecation can be performed adequately."), Arguments)))
+				Please make sure to load the level containing the parent landscape actor prior to {LevelName} so that data deprecation can be performed adequately. It is advised to reassign the \"Landscape Actor\" property of LandscapeStreamingProxies and resave these actors."), Arguments)))
 				->AddToken(FMapErrorToken::Create(FMapErrors::LandscapeLazyObjectPtrDeprecation_Warning));
 
 			// Show MapCheck window
