Cleans chained condition after code removal.

diff --git a/Engine/Source/Runtime/Landscape/Private/Landscape.cpp b/Engine/Source/Runtime/Landscape/Private/Landscape.cpp
--- a/Engine/Source/Runtime/Landscape/Private/Landscape.cpp
+++ b/Engine/Source/Runtime/Landscape/Private/Landscape.cpp
@@ -2320,15 +2320,9 @@ void ALandscapeProxy::PostRegisterAllComponents()
 		}
 	}
 #if WITH_EDITOR
-	if (LandscapeInfo != nullptr)
+	if ((LandscapeInfo != nullptr) && !IsPendingKillPending() && LandscapeGuid.IsValid())
 	{
-		if (!IsPendingKillPending())
-		{
-			if (LandscapeGuid.IsValid())
-			{
-				LandscapeInfo->FixupProxiesTransform();
-			}
-		}
+		LandscapeInfo->FixupProxiesTransform();
 	}
 #endif // WITH_EDITOR
 }
