Improve text of notifications when waiting for landscape resources.

diff --git a/Engine/Source/Runtime/Landscape/Private/LandscapeEditLayers.cpp b/Engine/Source/Runtime/Landscape/Private/LandscapeEditLayers.cpp
--- a/Engine/Source/Runtime/Landscape/Private/LandscapeEditLayers.cpp
+++ b/Engine/Source/Runtime/Landscape/Private/LandscapeEditLayers.cpp
@@ -7963,7 +7963,7 @@ void ALandscape::UpdateLayersContent(bool bInWaitForStreaming, bool bInSkipMonit
 		if (!InvalidShadingModelNotification.IsValid())
 		{
 			InvalidShadingModelNotification = MakeShared<FLandscapeNotification>(this, FLandscapeNotification::EType::ShadingModelInvalid);
-			static const FText NotificationText(LOCTEXT("InvalidShadingModel", "Landscape editor cannot update landscape with a feature level less than SM5"));
+			static const FText NotificationText(LOCTEXT("InvalidShadingModel", "Cannot update landscape with a feature level less than SM5"));
 			InvalidShadingModelNotification->NotificationText = NotificationText;
 		}
 		LandscapeNotificationManager->RegisterNotification(InvalidShadingModelNotification);
@@ -7985,7 +7985,7 @@ void ALandscape::UpdateLayersContent(bool bInWaitForStreaming, bool bInSkipMonit
 		if (!WaitingForTexturesNotification.IsValid())
 		{
 			WaitingForTexturesNotification = MakeShared<FLandscapeNotification>(this, FLandscapeNotification::EType::LandscapeTextureResourcesNotReady);
-			static const FText NotificationText(LOCTEXT("WaitForLandscapeTextureResources", "Landscape editor waiting for landscape texture resources to be ready"));
+			static const FText NotificationText(LOCTEXT("WaitForLandscapeTextureResources", "Waiting for texture resources to be ready"));
 			WaitingForTexturesNotification->NotificationText = NotificationText;
 			WaitingForTexturesNotification->NotificationStartTime = WaitingForLandscapeTextureResourcesStartTime + TimeBeforeDisplayingWaitingForResourcesNotification;
 		}
@@ -8005,7 +8005,7 @@ void ALandscape::UpdateLayersContent(bool bInWaitForStreaming, bool bInSkipMonit
 		if (!WaitingForBrushesNotification.IsValid())
 		{
 			WaitingForBrushesNotification = MakeShared<FLandscapeNotification>(this, FLandscapeNotification::EType::LandscapeBrushResourcesNotReady);
-			static const FText NotificationText(LOCTEXT("WaitForLandscapeBrushResources", "Landscape editor waiting for landscape brush resources to be ready"));
+			static const FText NotificationText(LOCTEXT("WaitForLandscapeBrushResources", "Waiting for brush resources to be ready"));
 			WaitingForBrushesNotification->NotificationText = NotificationText;
 			WaitingForBrushesNotification->NotificationStartTime = WaitingForLandscapeBrushResourcesStartTime + TimeBeforeDisplayingWaitingForResourcesNotification;
 		}
diff --git a/Engine/Source/Runtime/Landscape/Private/LandscapeNotification.cpp b/Engine/Source/Runtime/Landscape/Private/LandscapeNotification.cpp
--- a/Engine/Source/Runtime/Landscape/Private/LandscapeNotification.cpp
+++ b/Engine/Source/Runtime/Landscape/Private/LandscapeNotification.cpp
@@ -91,7 +91,7 @@ void FLandscapeNotificationManager::Tick()
 
 		FText Text = FText::Format(LOCTEXT("NotificationFooter", "{0} for {1}"),
 			LocalActiveNotification->NotificationText,
-			(NumIdenticalNotifications > 1) ? LOCTEXT("MultipleLandscapes", "multiple landscapes") : FText::FromString(LocalActiveNotification->GetLandscape()->GetActorLabel()));
+			(NumIdenticalNotifications > 1) ? LOCTEXT("MultipleLandscapes", "multiple landscapes") : FText::Format(LOCTEXT("NotificationLandscapeName", "landscape named: {0}"), FText::FromString(LocalActiveNotification->GetLandscape()->GetActorLabel())));
 
 		ShowNotificationItem(Text);
 	}
