Fixed possible access to Slate in commandlets
* This fixes a crash in the HLOD commandlet when a landscape is in the scene

diff --git a/Engine/Source/Runtime/Landscape/Private/LandscapeEditLayers.cpp b/Engine/Source/Runtime/Landscape/Private/LandscapeEditLayers.cpp
--- a/Engine/Source/Runtime/Landscape/Private/LandscapeEditLayers.cpp
+++ b/Engine/Source/Runtime/Landscape/Private/LandscapeEditLayers.cpp
@@ -7875,7 +7875,6 @@ void ALandscape::UpdateLayersContent(bool bInWaitForStreaming, bool bInSkipMonit
 	ULandscapeSubsystem* LandscapeSubsystem = World->GetSubsystem<ULandscapeSubsystem>();
 	check(LandscapeSubsystem != nullptr);
 	FLandscapeNotificationManager* LandscapeNotificationManager = LandscapeSubsystem->GetNotificationManager();
-	check(LandscapeNotificationManager != nullptr);
 
 	// Make sure Update doesn't dirty Landscape packages when not in Landscape Ed Mode
 	FLandscapeDirtyOnlyInModeScope DirtyOnlyInMode(LandscapeInfo);
@@ -7960,14 +7959,17 @@ void ALandscape::UpdateLayersContent(bool bInWaitForStreaming, bool bInSkipMonit
 	// The Edit layers shaders only work on SM5 : cancel any update that might happen when SM5+ shading model is not active :
 	if (World->FeatureLevel < ERHIFeatureLevel::SM5)
 	{
-		if (!InvalidShadingModelNotification.IsValid())
+		if (LandscapeNotificationManager)
 		{
-			InvalidShadingModelNotification = MakeShared<FLandscapeNotification>(this, FLandscapeNotification::EType::ShadingModelInvalid);
-			static const FText NotificationText(LOCTEXT("InvalidShadingModel", "Cannot update landscape with a feature level less than SM5"));
-			InvalidShadingModelNotification->NotificationText = NotificationText;
+			if (!InvalidShadingModelNotification.IsValid())
+			{
+				InvalidShadingModelNotification = MakeShared<FLandscapeNotification>(this, FLandscapeNotification::EType::ShadingModelInvalid);
+				static const FText NotificationText(LOCTEXT("InvalidShadingModel", "Cannot update landscape with a feature level less than SM5"));
+				InvalidShadingModelNotification->NotificationText = NotificationText;
+			}
+			LandscapeNotificationManager->RegisterNotification(InvalidShadingModelNotification);
+			bHideNotifications = false;
 		}
-		LandscapeNotificationManager->RegisterNotification(InvalidShadingModelNotification);
-		bHideNotifications = false;
 		return;
 	}
 	else
@@ -7979,7 +7981,7 @@ void ALandscape::UpdateLayersContent(bool bInWaitForStreaming, bool bInSkipMonit
 	static constexpr double TimeBeforeDisplayingWaitingForResourcesNotification = 3.0;
 
 	bResourcesReady &= PrepareLayersTextureResources(bInWaitForStreaming);
-	if (!bResourcesReady)
+	if (!bResourcesReady && LandscapeNotificationManager)
 	{
 		WaitingForLandscapeTextureResourcesStartTime = GetCurrentTime(); 
 		if (!WaitingForTexturesNotification.IsValid())
@@ -7999,7 +8001,7 @@ void ALandscape::UpdateLayersContent(bool bInWaitForStreaming, bool bInSkipMonit
 	}
 
 	bResourcesReady &= PrepareLayersBrushResources(World->FeatureLevel, bInWaitForStreaming);
-	if (!bResourcesReady)
+	if (!bResourcesReady && LandscapeNotificationManager)
 	{
 		WaitingForLandscapeBrushResourcesStartTime = GetCurrentTime();
 		if (!WaitingForBrushesNotification.IsValid())
diff --git a/Engine/Source/Runtime/Landscape/Private/LandscapeNotification.cpp b/Engine/Source/Runtime/Landscape/Private/LandscapeNotification.cpp
--- a/Engine/Source/Runtime/Landscape/Private/LandscapeNotification.cpp
+++ b/Engine/Source/Runtime/Landscape/Private/LandscapeNotification.cpp
@@ -19,7 +19,9 @@ FLandscapeNotification::FLandscapeNotification(const TWeakObjectPtr<ALandscape>&
 	, NotificationType(InNotificationType)
 	, ConditionCallback(InConditionCallback)
 	, UpdateTextCallback(InUpdateTextCallback)
-{}
+{
+	check(!IsRunningCommandlet());
+}
 
 bool FLandscapeNotification::operator == (const FLandscapeNotification& Other) const
 {
diff --git a/Engine/Source/Runtime/Landscape/Private/LandscapeSubsystem.cpp b/Engine/Source/Runtime/Landscape/Private/LandscapeSubsystem.cpp
--- a/Engine/Source/Runtime/Landscape/Private/LandscapeSubsystem.cpp
+++ b/Engine/Source/Runtime/Landscape/Private/LandscapeSubsystem.cpp
@@ -47,6 +47,12 @@ DECLARE_CYCLE_STAT(TEXT("LandscapeSubsystem Tick"), STAT_LandscapeSubsystemTick,
 
 ULandscapeSubsystem::ULandscapeSubsystem()
 	: bIsGrassCreationPrioritized(false)
+#if WITH_EDITOR
+	, GrassMapsBuilder(nullptr)
+	, GIBakedTextureBuilder(nullptr)
+	, PhysicalMaterialBuilder(nullptr)
+	, NotificationManager(nullptr)
+#endif
 {
 }
 
@@ -72,7 +78,11 @@ void ULandscapeSubsystem::Initialize(FSubsystemCollectionBase& Collection)
 	GrassMapsBuilder = new FLandscapeGrassMapsBuilder(GetWorld());
 	GIBakedTextureBuilder = new FLandscapeGIBakedTextureBuilder(GetWorld());
 	PhysicalMaterialBuilder = new FLandscapePhysicalMaterialBuilder(GetWorld());
-	NotificationManager = new FLandscapeNotificationManager();
+
+	if (!IsRunningCommandlet())
+	{
+		NotificationManager = new FLandscapeNotificationManager();
+	}
 #endif
 }
 
@@ -258,7 +268,11 @@ void ULandscapeSubsystem::Tick(float DeltaTime)
 	if (GIsEditor && !World->IsPlayInEditor())
 	{
 		LandscapePhysicalMaterial::GarbageCollectTasks();
-		NotificationManager->Tick();
+
+		if (NotificationManager)
+		{
+			NotificationManager->Tick();
+		}
 	}
 #endif
 }
