Landscape: Fix has for LandscapeHeightfieldCollisionComponent not taking size properties into account

diff --git a/Engine/Source/Runtime/Landscape/Private/LandscapeCollision.cpp b/Engine/Source/Runtime/Landscape/Private/LandscapeCollision.cpp
--- a/Engine/Source/Runtime/Landscape/Private/LandscapeCollision.cpp
+++ b/Engine/Source/Runtime/Landscape/Private/LandscapeCollision.cpp
@@ -1620,7 +1620,7 @@ uint32 ULandscapeHeightfieldCollisionComponent::ComputeCollisionHash() const
 	Hash = HashCombine(GetTypeHash(CollisionScale), Hash);
 
 	FTransform ComponentTransform = GetComponentToWorld();
-	Hash = FCrc::MemCrc32(&ComponentTransform, sizeof(ComponentTransform));
+	Hash = FCrc::MemCrc32(&ComponentTransform, sizeof(ComponentTransform), Hash);
 
 	const void* HeightBuffer = CollisionHeightData.LockReadOnly();
 	Hash = FCrc::MemCrc32(HeightBuffer, CollisionHeightData.GetBulkDataSize(), Hash);
