Setting to enable distance field lighting for LandscapeGrassTypes added.

diff --git a/Engine/Source/Runtime/Landscape/Classes/LandscapeGrassType.h b/Engine/Source/Runtime/Landscape/Classes/LandscapeGrassType.h
--- a/Engine/Source/Runtime/Landscape/Classes/LandscapeGrassType.h
+++ b/Engine/Source/Runtime/Landscape/Classes/LandscapeGrassType.h
@@ -113,6 +113,10 @@ struct FGrassVariety
 	UPROPERTY(EditAnywhere, Category = Grass)
 	bool bReceivesDecals;
 
+	/** Controls whether the primitive should affect dynamic distance field lighting methods. */
+	UPROPERTY(EditAnywhere, Category = Grass)
+	bool bAffectDistanceFieldLighting;
+
 	/** Whether the grass should cast shadows when using non-precomputed shadowing. **/
 	UPROPERTY(EditAnywhere, Category = Grass)
 	bool bCastDynamicShadow;
diff --git a/Engine/Source/Runtime/Landscape/Private/LandscapeGrass.cpp b/Engine/Source/Runtime/Landscape/Private/LandscapeGrass.cpp
--- a/Engine/Source/Runtime/Landscape/Private/LandscapeGrass.cpp
+++ b/Engine/Source/Runtime/Landscape/Private/LandscapeGrass.cpp
@@ -1358,6 +1358,7 @@ FGrassVariety::FGrassVariety()
 	, AlignToSurface(true)
 	, bUseLandscapeLightmap(false)
 	, bReceivesDecals(true)
+	, bAffectDistanceFieldLighting(false)
 	, bCastDynamicShadow(true)
 	, bCastContactShadow(true)
 	, bKeepInstanceBufferCPUCopy(false)
@@ -3162,6 +3163,7 @@ void ALandscapeProxy::UpdateGrass(const TArray<FVector>& Cameras, int32& InOutNu
 										GrassInstancedStaticMeshComponent->LightingChannels = GrassVariety.LightingChannels;
 										GrassInstancedStaticMeshComponent->bCastStaticShadow = false;
 										GrassInstancedStaticMeshComponent->CastShadow = (GrassVariety.bCastDynamicShadow || GrassVariety.bCastContactShadow) && !bDisableDynamicShadows;
+										GrassInstancedStaticMeshComponent->bAffectDistanceFieldLighting = GrassVariety.bAffectDistanceFieldLighting;
 										GrassInstancedStaticMeshComponent->bCastDynamicShadow = GrassVariety.bCastDynamicShadow && !bDisableDynamicShadows;
 										GrassInstancedStaticMeshComponent->bCastContactShadow = GrassVariety.bCastContactShadow && !bDisableDynamicShadows;
 										GrassInstancedStaticMeshComponent->OverrideMaterials = GrassVariety.OverrideMaterials;
@@ -3204,9 +3206,6 @@ void ALandscapeProxy::UpdateGrass(const TArray<FVector>& Cameras, int32& InOutNu
 											GrassInstancedStaticMeshComponent->InstanceEndCullDistance = GrassVariety.GetEndCullDistance() * CullDistanceScale;
 										}
 
-										//@todo - take the settings from a UFoliageType object.  For now, disable distance field lighting on grass so we don't hitch.
-										GrassInstancedStaticMeshComponent->bAffectDistanceFieldLighting = false;
-
 										{
 											QUICK_SCOPE_CYCLE_COUNTER(STAT_GrassAttachComp);
 
