Fixed ULandscapeInfo being created in PostLoad and added assert to make sure it doesn't creep up again. The reason is that in a streaming situation, when load/post-load occurs, the OwningWorld of the ULevel being loaded is the level package's world. Adding a ULandscapeInfo at this point would register it to that world. However, once the level is done streaming in (e.g. PostRegisterAllComponents is being called on its actor), the ULevel's OwningWorld has been changed (SetLoadedLevel) to the streaming level's final world and will therefore fail to find the landscape info that has been recorded previously.
It is therefore technically incorrect to access UWorld, its modifying its PerModuleDataObjects member or even accessing a UWorldSubsystem at this time, since the load-time UWorld is not the final one.

#rnx

diff --git a/Engine/Source/Runtime/Landscape/Private/Landscape.cpp b/Engine/Source/Runtime/Landscape/Private/Landscape.cpp
--- a/Engine/Source/Runtime/Landscape/Private/Landscape.cpp
+++ b/Engine/Source/Runtime/Landscape/Private/Landscape.cpp
@@ -3338,12 +3338,6 @@ void ALandscapeProxy::PostLoad()
 		BodyInstance.FixupData(this);
 	}
 
-	if ((GetLinker() && (GetLinker()->UEVer() < VER_UE4_LANDSCAPE_COMPONENT_LAZY_REFERENCES)) ||
-		LandscapeComponents.Num() != CollisionComponents.Num() ||
-		LandscapeComponents.ContainsByPredicate([](ULandscapeComponent* Comp) { return ((Comp != nullptr) && (Comp->GetCollisionComponent() == nullptr)); }))
-	{
-		CreateLandscapeInfo();
-	}
 #if WITH_EDITOR
 	PRAGMA_DISABLE_DEPRECATION_WARNINGS;
 	if (!LandscapeMaterialsOverride_DEPRECATED.IsEmpty())
diff --git a/Engine/Source/Runtime/Landscape/Private/LandscapeEdit.cpp b/Engine/Source/Runtime/Landscape/Private/LandscapeEdit.cpp
--- a/Engine/Source/Runtime/Landscape/Private/LandscapeEdit.cpp
+++ b/Engine/Source/Runtime/Landscape/Private/LandscapeEdit.cpp
@@ -1332,7 +1332,6 @@ void ULandscapeComponent::RecreateCollisionComponent(bool bUseSimpleCollision)
 	TArray<uint8> DominantLayerData;
 	TArray<ULandscapeLayerInfoObject*> LayerInfos;
 	ALandscapeProxy* Proxy = GetLandscapeProxy();
-	ULandscapeInfo* Info = GetLandscapeInfo();
 	const FCollisionSize CollisionSize = FCollisionSize::Create(NumSubsections, SubsectionSizeQuads, CollisionMipLevel);
 	const FCollisionSize SimpleCollisionSize = FCollisionSize::CreateSimple(bUseSimpleCollision, NumSubsections, SubsectionSizeQuads, SimpleCollisionMipLevel);
 	const int32 TotalCollisionSize = CollisionSize.SizeVertsSquare + SimpleCollisionSize.SizeVertsSquare;
@@ -1354,10 +1353,6 @@ void ULandscapeComponent::RecreateCollisionComponent(bool bUseSimpleCollision)
 			LayerInfos = CollisionComp->ComponentLayerInfos;
 		}
 
-		if (Info)
-		{
-			Info->Modify();
-		}
 		Proxy->Modify();
 		CollisionComp->DestroyComponent();
 		CollisionComp = nullptr;
diff --git a/Engine/Source/Runtime/Landscape/Private/LandscapeInfoMap.cpp b/Engine/Source/Runtime/Landscape/Private/LandscapeInfoMap.cpp
--- a/Engine/Source/Runtime/Landscape/Private/LandscapeInfoMap.cpp
+++ b/Engine/Source/Runtime/Landscape/Private/LandscapeInfoMap.cpp
@@ -47,6 +47,9 @@ void ULandscapeInfoMap::AddReferencedObjects(UObject* InThis, FReferenceCollecto
 
 ULandscapeInfoMap& ULandscapeInfoMap::GetLandscapeInfoMap(const UWorld* World)
 {
+	checkf(!FUObjectThreadContext::Get().IsRoutingPostLoad, TEXT("This is probably a mistake to call this function in PostLoad, as the world used to register the landscape proxy is not the final world, \
+		since the actor's ULevel's final OwningWorld is only set after streaming is done"));
+
 	ULandscapeInfoMap *FoundObject = nullptr;
 	World->PerModuleDataObjects.FindItemByClass(&FoundObject);
 
