Reverted CL 23834590 due to assert when loading firingrange.

diff --git a/Engine/Source/Runtime/Landscape/Private/Landscape.cpp b/Engine/Source/Runtime/Landscape/Private/Landscape.cpp
--- a/Engine/Source/Runtime/Landscape/Private/Landscape.cpp
+++ b/Engine/Source/Runtime/Landscape/Private/Landscape.cpp
@@ -3338,6 +3338,12 @@ void ALandscapeProxy::PostLoad()
 		BodyInstance.FixupData(this);
 	}
 
+	if ((GetLinker() && (GetLinker()->UEVer() < VER_UE4_LANDSCAPE_COMPONENT_LAZY_REFERENCES)) ||
+		LandscapeComponents.Num() != CollisionComponents.Num() ||
+		LandscapeComponents.ContainsByPredicate([](ULandscapeComponent* Comp) { return ((Comp != nullptr) && (Comp->GetCollisionComponent() == nullptr)); }))
+	{
+		CreateLandscapeInfo();
+	}
 #if WITH_EDITOR
 	PRAGMA_DISABLE_DEPRECATION_WARNINGS;
 	if (!LandscapeMaterialsOverride_DEPRECATED.IsEmpty())
diff --git a/Engine/Source/Runtime/Landscape/Private/LandscapeEdit.cpp b/Engine/Source/Runtime/Landscape/Private/LandscapeEdit.cpp
--- a/Engine/Source/Runtime/Landscape/Private/LandscapeEdit.cpp
+++ b/Engine/Source/Runtime/Landscape/Private/LandscapeEdit.cpp
@@ -1332,6 +1332,7 @@ void ULandscapeComponent::RecreateCollisionComponent(bool bUseSimpleCollision)
 	TArray<uint8> DominantLayerData;
 	TArray<ULandscapeLayerInfoObject*> LayerInfos;
 	ALandscapeProxy* Proxy = GetLandscapeProxy();
+	ULandscapeInfo* Info = GetLandscapeInfo();
 	const FCollisionSize CollisionSize = FCollisionSize::Create(NumSubsections, SubsectionSizeQuads, CollisionMipLevel);
 	const FCollisionSize SimpleCollisionSize = FCollisionSize::CreateSimple(bUseSimpleCollision, NumSubsections, SubsectionSizeQuads, SimpleCollisionMipLevel);
 	const int32 TotalCollisionSize = CollisionSize.SizeVertsSquare + SimpleCollisionSize.SizeVertsSquare;
@@ -1353,6 +1354,10 @@ void ULandscapeComponent::RecreateCollisionComponent(bool bUseSimpleCollision)
 			LayerInfos = CollisionComp->ComponentLayerInfos;
 		}
 
+		if (Info)
+		{
+			Info->Modify();
+		}
 		Proxy->Modify();
 		CollisionComp->DestroyComponent();
 		CollisionComp = nullptr;
diff --git a/Engine/Source/Runtime/Landscape/Private/LandscapeInfoMap.cpp b/Engine/Source/Runtime/Landscape/Private/LandscapeInfoMap.cpp
--- a/Engine/Source/Runtime/Landscape/Private/LandscapeInfoMap.cpp
+++ b/Engine/Source/Runtime/Landscape/Private/LandscapeInfoMap.cpp
@@ -47,9 +47,6 @@ void ULandscapeInfoMap::AddReferencedObjects(UObject* InThis, FReferenceCollecto
 
 ULandscapeInfoMap& ULandscapeInfoMap::GetLandscapeInfoMap(const UWorld* World)
 {
-	checkf(!FUObjectThreadContext::Get().IsRoutingPostLoad, TEXT("This is probably a mistake to call this function in PostLoad, as the world used to register the landscape proxy is not the final world, \
-		since the actor's ULevel's final OwningWorld is only set after streaming is done"));
-
 	ULandscapeInfoMap *FoundObject = nullptr;
 	World->PerModuleDataObjects.FindItemByClass(&FoundObject);
 
