Allow Landscape Spline Actors to change which Landscape they are targeting.

diff --git a/Engine/Source/Runtime/Landscape/Classes/ILandscapeSplineInterface.h b/Engine/Source/Runtime/Landscape/Classes/ILandscapeSplineInterface.h
--- a/Engine/Source/Runtime/Landscape/Classes/ILandscapeSplineInterface.h
+++ b/Engine/Source/Runtime/Landscape/Classes/ILandscapeSplineInterface.h
@@ -24,6 +24,7 @@ public:
 	virtual ULandscapeInfo* GetLandscapeInfo() const PURE_VIRTUAL(ILandscapeSplineInterface::GetLandscapeInfo, return nullptr;);
 	virtual FTransform LandscapeActorToWorld() const PURE_VIRTUAL(ILandscapeSplineInterface::LandscapeActorToWorld, return FTransform::Identity;);
 	virtual ULandscapeSplinesComponent* GetSplinesComponent() const PURE_VIRTUAL(ILandscapeSplineInterface::GetSplineComponent, return nullptr;);
+	virtual void UpdateSharedProperties(ULandscapeInfo* InLandscapeInfo) PURE_VIRTUAL(ILandscapeSplineInterface::UpdateSharedProperties);
 
 #if WITH_EDITOR
 	virtual bool SupportsForeignSplineMesh() const PURE_VIRTUAL(ILandscapeSplineInterface::SupportsForeignSplineMesh, return false;);
diff --git a/Engine/Source/Runtime/Landscape/Classes/LandscapeInfo.h b/Engine/Source/Runtime/Landscape/Classes/LandscapeInfo.h
--- a/Engine/Source/Runtime/Landscape/Classes/LandscapeInfo.h
+++ b/Engine/Source/Runtime/Landscape/Classes/LandscapeInfo.h
@@ -168,6 +168,7 @@ private:
 #if WITH_EDITORONLY_DATA
 	TSet<TWeakObjectPtr<UObject>> PackagesToDeprecate;
 
+	// SORTED list of all actors implementing the spline interface that are registered with this landscape info
 	UPROPERTY()
 	TArray<TScriptInterface<ILandscapeSplineInterface>> SplineActors;
 
@@ -406,6 +407,10 @@ public:
 	LANDSCAPE_API void RegisterSplineActor(TScriptInterface<ILandscapeSplineInterface> SplineActor);
 	LANDSCAPE_API void UnregisterSplineActor(TScriptInterface<ILandscapeSplineInterface> SplineActor);
 
+	/** Updates the spline registrations (with ALL landscape infos) when a spline has changed the landscape it is targeting
+	 */
+	LANDSCAPE_API static void UpdateRegistrationForSplineActor(UWorld* InWorld, TScriptInterface<ILandscapeSplineInterface> InSplineActor);
+
 	LANDSCAPE_API void RequestSplineLayerUpdate();
 	LANDSCAPE_API void ForceLayersFullUpdate();
 
diff --git a/Engine/Source/Runtime/Landscape/Classes/LandscapeProxy.h b/Engine/Source/Runtime/Landscape/Classes/LandscapeProxy.h
--- a/Engine/Source/Runtime/Landscape/Classes/LandscapeProxy.h
+++ b/Engine/Source/Runtime/Landscape/Classes/LandscapeProxy.h
@@ -990,6 +990,7 @@ public:
 
 	// ILandscapeSplineInterface
 	LANDSCAPE_API virtual ULandscapeSplinesComponent* GetSplinesComponent() const override { return SplineComponent; }
+	LANDSCAPE_API virtual void UpdateSharedProperties(ULandscapeInfo* InLandscapeInfo) override;
 
 	// Retrieve the screen size at which each LOD should be rendered
 	LANDSCAPE_API TArray<float> GetLODScreenSizeArray() const;
diff --git a/Engine/Source/Runtime/Landscape/Classes/LandscapeSplineActor.h b/Engine/Source/Runtime/Landscape/Classes/LandscapeSplineActor.h
--- a/Engine/Source/Runtime/Landscape/Classes/LandscapeSplineActor.h
+++ b/Engine/Source/Runtime/Landscape/Classes/LandscapeSplineActor.h
@@ -22,7 +22,7 @@ protected:
 
 #if WITH_EDITORONLY_DATA
 	/** Landscape **/
-	UPROPERTY()
+	UPROPERTY(EditAnywhere, Category = "Target Landscape")
 	TObjectPtr<ALandscape> LandscapeActor;
 #endif
 	
@@ -31,6 +31,7 @@ public:
 	virtual ULandscapeSplinesComponent* GetSplinesComponent() const override;
 	virtual FTransform LandscapeActorToWorld() const override;
 	virtual ULandscapeInfo* GetLandscapeInfo() const override;
+	virtual void UpdateSharedProperties(ULandscapeInfo* InLandscapeInfo) override;
 		
 #if WITH_EDITOR
 	LANDSCAPE_API bool HasGeneratedLandscapeSplineMeshesActors() const;
@@ -43,6 +44,7 @@ public:
 	virtual void PostRegisterAllComponents() override;
 	virtual void UnregisterAllComponents(bool bForReregister = false) override;
 	virtual void PostEditMove(bool bFinished) override;
+	virtual void PostEditChangeProperty(FPropertyChangedEvent& PropertyChangedEvent) override;
 	virtual bool EditorCanAttachTo(const AActor* InParent, FText& OutReason) const override { return false; }
 	virtual AActor* GetSceneOutlinerParent() const;
 	virtual bool SupportsForeignSplineMesh() const override { return false; }
diff --git a/Engine/Source/Runtime/Landscape/Private/Landscape.cpp b/Engine/Source/Runtime/Landscape/Private/Landscape.cpp
--- a/Engine/Source/Runtime/Landscape/Private/Landscape.cpp
+++ b/Engine/Source/Runtime/Landscape/Private/Landscape.cpp
@@ -1433,6 +1433,11 @@ FTransform ALandscapeProxy::LandscapeActorToWorld() const
 	return TM;
 }
 
+void ALandscapeProxy::UpdateSharedProperties(ULandscapeInfo* InLandscapeInfo)
+{
+	check(LandscapeGuid == InLandscapeInfo->LandscapeGuid);
+}
+
 static TArray<float> GetLODScreenSizeArray(const ALandscapeProxy* InLandscapeProxy, const int32 InNumLODLevels)
 {
 	static TConsoleVariableData<float>* CVarSMLODDistanceScale = IConsoleManager::Get().FindTConsoleVariableDataFloat(TEXT("r.StaticMeshLODDistanceScale"));
@@ -4282,6 +4287,12 @@ void ULandscapeInfo::RegisterActor(ALandscapeProxy* Proxy, bool bMapCheck, bool
 			LandscapeActor = Landscape;
 
 #if WITH_EDITOR
+			// Update registered splines so they can pull the actor pointer
+			for (TScriptInterface<ILandscapeSplineInterface> SplineActor : SplineActors)
+			{
+				SplineActor->UpdateSharedProperties(this);
+			}
+
 			// In world composition user is not allowed to move landscape in editor, only through WorldBrowser 
 			bool bIsLockLocation = LandscapeActor->IsLockLocation();
 			bIsLockLocation |= OwningWorld != nullptr ? OwningWorld->WorldComposition != nullptr : false;
@@ -4491,6 +4502,8 @@ void ULandscapeInfo::RegisterSplineActor(TScriptInterface<ILandscapeSplineInterf
 		SplineActors.Insert(SplineActor, LBoundIdx);
 	}
 
+	SplineActor->UpdateSharedProperties(this);
+
 	if (SplineActor->GetSplinesComponent())
 	{
 		RequestSplineLayerUpdate();
@@ -4508,6 +4521,39 @@ void ULandscapeInfo::UnregisterSplineActor(TScriptInterface<ILandscapeSplineInte
 	}
 }
 
+void ULandscapeInfo::UpdateRegistrationForSplineActor(UWorld* InWorld, TScriptInterface<ILandscapeSplineInterface> InSplineActor)
+{
+	if (InWorld == nullptr)
+		return;
+
+	ULandscapeInfoMap& LandscapeInfoMap = ULandscapeInfoMap::GetLandscapeInfoMap(InWorld);
+	FGuid SplineLandscapeGUID = InSplineActor->GetLandscapeGuid();
+
+	// first let's unregister from any landscapes that have it (incorrectly) registered
+	for (const TPair<FGuid, ULandscapeInfo*>& pair : LandscapeInfoMap.Map)
+	{
+		ULandscapeInfo* LandscapeInfo = pair.Value;
+
+		// only unregister if the landscape guids don't match
+		if ((LandscapeInfo->LandscapeGuid != SplineLandscapeGUID) &&
+			LandscapeInfo->SplineActors.Contains(InSplineActor))
+		{
+			LandscapeInfo->UnregisterSplineActor(InSplineActor);
+		}
+	}
+
+	// then let's make sure it is registered with the correct landscape info
+	if (SplineLandscapeGUID.IsValid())
+	{
+		ULandscapeInfo* LandscapeInfo = InSplineActor->GetLandscapeInfo();
+		check(LandscapeInfo);
+		if (!LandscapeInfo->SplineActors.Contains(InSplineActor))
+		{
+			LandscapeInfo->RegisterSplineActor(InSplineActor);
+		}
+	}
+}
+
 void ULandscapeInfo::RequestSplineLayerUpdate()
 {
 	if (LandscapeActor.IsValid())
diff --git a/Engine/Source/Runtime/Landscape/Private/LandscapeSplineActor.cpp b/Engine/Source/Runtime/Landscape/Private/LandscapeSplineActor.cpp
--- a/Engine/Source/Runtime/Landscape/Private/LandscapeSplineActor.cpp
+++ b/Engine/Source/Runtime/Landscape/Private/LandscapeSplineActor.cpp
@@ -40,6 +40,16 @@ ULandscapeInfo* ALandscapeSplineActor::GetLandscapeInfo() const
 	return ULandscapeInfo::Find(GetWorld(), LandscapeGuid);
 }
 
+// this is called when Splines or LandscapeActors are registered with LandscapeInfo, to update the LandscapeActor pointer
+void ALandscapeSplineActor::UpdateSharedProperties(ULandscapeInfo* InLandscapeInfo)
+{
+	check(InLandscapeInfo);
+	check(InLandscapeInfo->LandscapeGuid == LandscapeGuid);
+#if WITH_EDITOR
+	LandscapeActor = InLandscapeInfo->LandscapeActor.Get();
+#endif // WITH_EDITOR
+}
+
 #if WITH_EDITOR
 void ALandscapeSplineActor::GetActorDescProperties(FPropertyPairsMap& PropertyPairsMap) const
 {
@@ -51,6 +61,7 @@ void ALandscapeSplineActor::GetActorDescProperties(FPropertyPairsMap& PropertyPa
 	}
 }
 
+// this is called when a LandscapeSplineActor is created, to copy it's landscape target data from the landscape info
 void ALandscapeSplineActor::GetSharedProperties(ULandscapeInfo* InLandscapeInfo)
 {
 	Modify();
@@ -128,6 +139,27 @@ void ALandscapeSplineActor::UnregisterAllComponents(bool bForReregister)
 	Super::UnregisterAllComponents(bForReregister);
 }
 
+void ALandscapeSplineActor::PostEditChangeProperty(FPropertyChangedEvent& PropertyChangedEvent)
+{
+	const FName PropertyName = PropertyChangedEvent.GetMemberPropertyName();
+	if (PropertyName == FName(TEXT("LandscapeActor")))
+	{
+		if (LandscapeActor && !LandscapeActor->HasAnyFlags(RF_BeginDestroyed))
+		{
+			LandscapeGuid = LandscapeActor->GetLandscapeGuid();
+		}
+		else
+		{
+			LandscapeActor = nullptr;
+			LandscapeGuid.Invalidate();
+		}
+
+		ULandscapeInfo::UpdateRegistrationForSplineActor(GetWorld(), this);
+	}
+
+	Super::PostEditChangeProperty(PropertyChangedEvent);
+}
+
 void ALandscapeSplineActor::PostEditMove(bool bFinished)
 {
 	Super::PostEditMove(bFinished);
