Fix in LandscapeRender to read whether WorldPositionOffset is enabled in the material from the correct source (shader map) instead of the old, approximate source (asking the material if the property is connected).

diff --git a/Engine/Source/Runtime/Landscape/Private/LandscapeRender.cpp b/Engine/Source/Runtime/Landscape/Private/LandscapeRender.cpp
--- a/Engine/Source/Runtime/Landscape/Private/LandscapeRender.cpp
+++ b/Engine/Source/Runtime/Landscape/Private/LandscapeRender.cpp
@@ -2527,7 +2527,7 @@ void FLandscapeComponentSceneProxy::GetDynamicRayTracingInstances(FRayTracingMat
 				const FMaterialRenderProxy* FallbackMaterialRenderProxyPtr = nullptr;
 				const FMaterial& Material = MeshBatch.MaterialRenderProxy->GetMaterialWithFallback(((FSceneInterface*)Context.Scene)->GetFeatureLevel(), FallbackMaterialRenderProxyPtr);
 
-				if (Material.HasVertexPositionOffsetConnected())
+				if (Material.GetRenderingThreadShaderMap()->UsesWorldPositionOffset())
 				{
 					const FMaterialRenderProxy* MaterialRenderProxy = FallbackMaterialRenderProxyPtr ? FallbackMaterialRenderProxyPtr : MeshBatch.MaterialRenderProxy;
 
