Fixed crash when listing textures (listtextures cmd) after having deleted a landscape actor

diff --git a/Engine/Source/Runtime/Landscape/Private/LandscapeRender.cpp b/Engine/Source/Runtime/Landscape/Private/LandscapeRender.cpp
--- a/Engine/Source/Runtime/Landscape/Private/LandscapeRender.cpp
+++ b/Engine/Source/Runtime/Landscape/Private/LandscapeRender.cpp
@@ -3597,8 +3597,13 @@ void ULandscapeComponent::GetStreamingRenderAssetInfo(FStreamingTextureLevelCont
 	float TexelFactor = 0.0f;
 	if (Proxy)
 	{
+		double ScaleFactor = 1.0;
+		if (USceneComponent* ProxyRootComponent = Proxy->GetRootComponent())
+		{
+			ScaleFactor = FMath::Abs(ProxyRootComponent->GetRelativeScale3D().X);
+		}
 		LocalStreamingDistanceMultiplier = FMath::Max(0.0f, Proxy->StreamingDistanceMultiplier);
-		TexelFactor = 0.75f * LocalStreamingDistanceMultiplier * ComponentSizeQuads * FMath::Abs(Proxy->GetRootComponent()->GetRelativeScale3D().X);
+		TexelFactor = 0.75f * LocalStreamingDistanceMultiplier * ComponentSizeQuads * ScaleFactor;
 	}
 
 	ERHIFeatureLevel::Type FeatureLevel = LevelContext.GetFeatureLevel();
