Starting to convert shader classes to use FRHIBatchedShaderParameters to set parameters instead of setting them on the RHICommandList themselves.

diff --git a/Engine/Plugins/Enterprise/GLTFExporter/Source/GLTFExporter/Private/Converters/GLTFNormalMapPreview.cpp b/Engine/Plugins/Enterprise/GLTFExporter/Source/GLTFExporter/Private/Converters/GLTFNormalMapPreview.cpp
--- a/Engine/Plugins/Enterprise/GLTFExporter/Source/GLTFExporter/Private/Converters/GLTFNormalMapPreview.cpp
+++ b/Engine/Plugins/Enterprise/GLTFExporter/Source/GLTFExporter/Private/Converters/GLTFNormalMapPreview.cpp
@@ -11,6 +11,7 @@
 #include "ShaderParameterUtils.h"
 #include "PipelineStateCache.h"
 #include "RHIStaticStates.h"
+#include "DataDrivenShaderPlatformInfo.h"
 
 /*------------------------------------------------------------------------------
 	Batched element shaders for previewing normal maps.
@@ -48,10 +49,9 @@ public:
 	 * Set shader parameters.
 	 * @param NormalMapTexture - The normal map texture to sample.
 	 */
-	void SetParameters(FRHICommandList& RHICmdList, const FTexture* NormalMapTexture)
+	void SetParameters(FRHIBatchedShaderParameters& BatchedParameters, const FTexture* NormalMapTexture)
 	{
-		FRHIPixelShader* PixelShaderRHI = RHICmdList.GetBoundPixelShader();
-		SetTextureParameter(RHICmdList, PixelShaderRHI,Texture,TextureSampler,NormalMapTexture);
+		SetTextureParameter(BatchedParameters, Texture, TextureSampler, NormalMapTexture);
 	}
 
 private:
@@ -83,7 +83,7 @@ void FGLTFNormalMapPreview::BindShaders(
 
 	SetGraphicsPipelineState(RHICmdList, GraphicsPSOInit, 0, EApplyRendertargetOption::CheckApply);
 
-	VertexShader->SetParameters(RHICmdList, InTransform);
-	PixelShader->SetParameters(RHICmdList, Texture);
+	SetAllShaderParametersVS(RHICmdList, VertexShader, InTransform);
+	SetAllShaderParametersPS(RHICmdList, PixelShader, Texture);
 }
 
diff --git a/Engine/Plugins/Enterprise/GLTFExporter/Source/GLTFExporter/Private/Converters/GLTFSimpleTexture2DPreview.cpp b/Engine/Plugins/Enterprise/GLTFExporter/Source/GLTFExporter/Private/Converters/GLTFSimpleTexture2DPreview.cpp
--- a/Engine/Plugins/Enterprise/GLTFExporter/Source/GLTFExporter/Private/Converters/GLTFSimpleTexture2DPreview.cpp
+++ b/Engine/Plugins/Enterprise/GLTFExporter/Source/GLTFExporter/Private/Converters/GLTFSimpleTexture2DPreview.cpp
@@ -11,6 +11,7 @@
 #include "SimpleElementShaders.h"
 #include "ShaderParameterUtils.h"
 #include "PipelineStateCache.h"
+#include "DataDrivenShaderPlatformInfo.h"
 
 class FGLTFSimpleTexture2DPS : public FGlobalShader
 {
@@ -39,10 +40,9 @@ public:
 	 * Set shader parameters.
 	 * @param SampleTexture - The texture to sample.
 	 */
-	void SetParameters(FRHICommandList& RHICmdList, const FTexture* SampleTexture)
+	void SetParameters(FRHIBatchedShaderParameters& BatchedParameters, const FTexture* SampleTexture)
 	{
-		FRHIPixelShader* PixelShaderRHI = RHICmdList.GetBoundPixelShader();
-		SetTextureParameter(RHICmdList, PixelShaderRHI,Texture,TextureSampler,SampleTexture);
+		SetTextureParameter(BatchedParameters, Texture, TextureSampler, SampleTexture);
 	}
 
 private:
@@ -74,6 +74,6 @@ void FGLTFSimpleTexture2DPreview::BindShaders(
 
 	SetGraphicsPipelineState(RHICmdList, GraphicsPSOInit, 0, EApplyRendertargetOption::CheckApply);
 
-	VertexShader->SetParameters(RHICmdList, InTransform);
-	PixelShader->SetParameters(RHICmdList, Texture);
+	SetAllShaderParametersVS(RHICmdList, VertexShader, InTransform);
+	SetAllShaderParametersPS(RHICmdList, PixelShader, Texture);
 }
diff --git a/Engine/Plugins/FX/Niagara/Source/NiagaraShader/Private/NiagaraBatchedElements.cpp b/Engine/Plugins/FX/Niagara/Source/NiagaraShader/Private/NiagaraBatchedElements.cpp
--- a/Engine/Plugins/FX/Niagara/Source/NiagaraShader/Private/NiagaraBatchedElements.cpp
+++ b/Engine/Plugins/FX/Niagara/Source/NiagaraShader/Private/NiagaraBatchedElements.cpp
@@ -8,6 +8,7 @@
 #include "SimpleElementShaders.h"
 #include "ShaderParameterStruct.h"
 #include "ShaderParameterMacros.h"
+#include "ShaderParameterUtils.h"
 
 //////////////////////////////////////////////////////////////////////////
 
@@ -64,7 +65,7 @@ void FBatchedElementNiagara2DArrayAttribute::BindShaders(FRHICommandList& RHICmd
 	RHICmdList.ApplyCachedRenderTargets(GraphicsPSOInit);
 	SetGraphicsPipelineState(RHICmdList, GraphicsPSOInit, 0);
 
-	VertexShader->SetParameters(RHICmdList, InTransform);
+	SetAllShaderParametersVS(RHICmdList, VertexShader, InTransform);
 
 	FNiagaraSimpleElement2DArrayAttribute::FParameters PassParameters;
 	PassParameters.InAttributeSlices = AttributeSlices;
@@ -134,7 +135,7 @@ void FBatchedElementNiagaraVolumeAttribute::BindShaders(FRHICommandList& RHICmdL
 	RHICmdList.ApplyCachedRenderTargets(GraphicsPSOInit);
 	SetGraphicsPipelineState(RHICmdList, GraphicsPSOInit, 0);
 
-	VertexShader->SetParameters(RHICmdList, InTransform);
+	SetAllShaderParametersVS(RHICmdList, VertexShader, InTransform);
 
 	FNiagaraSimpleElementVolumeAttribute::FParameters PassParameters;
 	PassParameters.InNumAttributes	= AttributeUVs.Num();
@@ -197,7 +198,7 @@ void FBatchedElementNiagaraInvertColorChannel::BindShaders(FRHICommandList& RHIC
 	RHICmdList.ApplyCachedRenderTargets(GraphicsPSOInit);
 	SetGraphicsPipelineState(RHICmdList, GraphicsPSOInit, 0);
 
-	VertexShader->SetParameters(RHICmdList, InTransform);
+	SetAllShaderParametersVS(RHICmdList, VertexShader, InTransform);
 
 	FNiagaraSimpleElementPS::FParameters PassParameters;
 	PassParameters.InColorWeights	= FMatrix44f::Identity;
@@ -221,7 +222,7 @@ void FBatchedElementNiagaraSimple::BindShaders(FRHICommandList& RHICmdList, FGra
 	RHICmdList.ApplyCachedRenderTargets(GraphicsPSOInit);
 	SetGraphicsPipelineState(RHICmdList, GraphicsPSOInit, 0);
 
-	VertexShader->SetParameters(RHICmdList, InTransform);
+	SetAllShaderParametersVS(RHICmdList, VertexShader, InTransform);
 
 	FNiagaraSimpleElementPS::FParameters PassParameters;
 	PassParameters.InColorWeights = FMatrix44f(ColorTransform);
diff --git a/Engine/Source/Editor/UnrealEd/Private/NormalMapPreview.cpp b/Engine/Source/Editor/UnrealEd/Private/NormalMapPreview.cpp
--- a/Engine/Source/Editor/UnrealEd/Private/NormalMapPreview.cpp
+++ b/Engine/Source/Editor/UnrealEd/Private/NormalMapPreview.cpp
@@ -11,6 +11,7 @@
 #include "ShaderParameterUtils.h"
 #include "PipelineStateCache.h"
 #include "RHIStaticStates.h"
+#include "DataDrivenShaderPlatformInfo.h"
 
 /*------------------------------------------------------------------------------
 	Batched element shaders for previewing normal maps.
@@ -48,10 +49,9 @@ public:
 	 * Set shader parameters.
 	 * @param NormalMapTexture - The normal map texture to sample.
 	 */
-	void SetParameters(FRHICommandList& RHICmdList, const FTexture* NormalMapTexture)
+	void SetParameters(FRHIBatchedShaderParameters& BatchedParameters, const FTexture* NormalMapTexture)
 	{
-		FRHIPixelShader* PixelShaderRHI = RHICmdList.GetBoundPixelShader();
-		SetTextureParameter(RHICmdList, PixelShaderRHI,Texture,TextureSampler,NormalMapTexture);
+		SetTextureParameter(BatchedParameters, Texture, TextureSampler, NormalMapTexture);
 	}
 
 private:
@@ -84,7 +84,7 @@ void FNormalMapBatchedElementParameters::BindShaders(
 	RHICmdList.ApplyCachedRenderTargets(GraphicsPSOInit);
 	SetGraphicsPipelineState(RHICmdList, GraphicsPSOInit, 0);
 
-	VertexShader->SetParameters(RHICmdList, InTransform);
-	PixelShader->SetParameters(RHICmdList, Texture);
+	SetAllShaderParametersVS(RHICmdList, VertexShader, InTransform);
+	SetAllShaderParametersPS(RHICmdList, PixelShader, Texture);
 }
 
diff --git a/Engine/Source/Editor/UnrealEd/Private/Texture2DPreview.cpp b/Engine/Source/Editor/UnrealEd/Private/Texture2DPreview.cpp
--- a/Engine/Source/Editor/UnrealEd/Private/Texture2DPreview.cpp
+++ b/Engine/Source/Editor/UnrealEd/Private/Texture2DPreview.cpp
@@ -17,6 +17,7 @@
 #include "VirtualTexturing.h"
 #include "Engine/Texture2DArray.h"
 #include "ShaderParameterStruct.h"
+#include "DataDrivenShaderPlatformInfo.h"
 
 /*------------------------------------------------------------------------------
 	Batched element shaders for previewing 2d textures.
@@ -165,7 +166,7 @@ void FBatchedElementTexture2DPreviewParameters::BindShaders(
 	RHICmdList.ApplyCachedRenderTargets(GraphicsPSOInit);
 	SetGraphicsPipelineState(RHICmdList, GraphicsPSOInit, 0);
 
-	VertexShader->SetParameters(RHICmdList, InTransform);
+	SetAllShaderParametersVS(RHICmdList, VertexShader, InTransform);
 
 	SetShaderParameters(RHICmdList, PixelShader, PixelShader.GetPixelShader(), Parameters);
 }
diff --git a/Engine/Source/Editor/UnrealEd/Private/VolumeTexturePreview.cpp b/Engine/Source/Editor/UnrealEd/Private/VolumeTexturePreview.cpp
--- a/Engine/Source/Editor/UnrealEd/Private/VolumeTexturePreview.cpp
+++ b/Engine/Source/Editor/UnrealEd/Private/VolumeTexturePreview.cpp
@@ -16,6 +16,7 @@
 #include "Editor.h"
 #include "TextureResource.h"
 #include "GlobalRenderResources.h"
+#include "DataDrivenShaderPlatformInfo.h"
 
 UNREALED_API void GetBestFitForNumberOfTiles(int32 InSize, int32& OutNumTilesX, int32& OutNumTilesY)
 {
@@ -77,47 +78,45 @@ public:
 		return IsFeatureLevelSupported(Parameters.Platform, ERHIFeatureLevel::SM5) && !IsConsolePlatform(Parameters.Platform);
 	}
 	
-	void SetParameters(FRHICommandList& RHICmdList, const FTexture* TextureValue, int32 SizeZ, const FMatrix44f& ColorWeightsValue, float GammaValue, float MipLevel, float Opacity, const FRotator& TraceOrientation, bool bUsePointSampling)
+	void SetParameters(FRHIBatchedShaderParameters& BatchedParameters, const FTexture* TextureValue, int32 SizeZ, const FMatrix44f& ColorWeightsValue, float GammaValue, float MipLevel, float Opacity, const FRotator& TraceOrientation, bool bUsePointSampling)
 	{
-		FRHIPixelShader* ShaderRHI = RHICmdList.GetBoundPixelShader();
-
 		FRHISamplerState* SamplerState = bUsePointSampling ? TStaticSamplerState<SF_Point>::GetRHI() : TextureValue->SamplerStateRHI.GetReference();
-		SetTextureParameter(RHICmdList, ShaderRHI, InTexture, InTextureSampler, SamplerState, TextureValue->TextureRHI);
+		SetTextureParameter(BatchedParameters, InTexture, InTextureSampler, SamplerState, TextureValue->TextureRHI);
 
 		if (GEditor && GEditor->Bad)
 		{
-			SetTextureParameter(RHICmdList, ShaderRHI, BadTexture, BadTextureSampler, GEditor->Bad->GetResource());
+			SetTextureParameter(BatchedParameters, BadTexture, BadTextureSampler, GEditor->Bad->GetResource());
 		}
 		else
 		{
-			SetTextureParameter(RHICmdList, ShaderRHI, BadTexture, GWhiteTexture->TextureRHI);
+			SetTextureParameter(BatchedParameters, BadTexture, GWhiteTexture->TextureRHI);
 		}
-		SetShaderValue(RHICmdList, ShaderRHI,ColorWeights,ColorWeightsValue);
+		SetShaderValue(BatchedParameters,ColorWeights,ColorWeightsValue);
 
 		const int32 MipSizeZ = MipLevel >= 0 ? FMath::Max<int32>(SizeZ >> FMath::FloorToInt(MipLevel), 1) : SizeZ;
 		FVector4f PackedParametersValue(GammaValue, MipLevel, (float)MipSizeZ, Opacity);
-		SetShaderValue(RHICmdList, ShaderRHI, PackedParameters, PackedParametersValue);
+		SetShaderValue(BatchedParameters, PackedParameters, PackedParametersValue);
 
 		int32 NumTilesX = 0;
 		int32 NumTilesY = 0;
 		GetBestFitForNumberOfTiles(MipSizeZ, NumTilesX, NumTilesY);
-		SetShaderValue(RHICmdList, ShaderRHI, NumTilesPerSideParameter, FVector4f((float)NumTilesX, (float)NumTilesY, 0 ,0));
+		SetShaderValue(BatchedParameters, NumTilesPerSideParameter, FVector4f((float)NumTilesX, (float)NumTilesY, 0 ,0));
 
-		SetShaderValue(RHICmdList, ShaderRHI,TextureComponentReplicate,TextureValue->bGreyScaleFormat ? FLinearColor(1,0,0,0) : FLinearColor(0,0,0,0));
-		SetShaderValue(RHICmdList, ShaderRHI,TextureComponentReplicateAlpha,TextureValue->bGreyScaleFormat ? FLinearColor(1,0,0,0) : FLinearColor(0,0,0,1));
+		SetShaderValue(BatchedParameters,TextureComponentReplicate,TextureValue->bGreyScaleFormat ? FLinearColor(1,0,0,0) : FLinearColor(0,0,0,0));
+		SetShaderValue(BatchedParameters,TextureComponentReplicateAlpha,TextureValue->bGreyScaleFormat ? FLinearColor(1,0,0,0) : FLinearColor(0,0,0,1));
 
 		const FVector3f TextureDimension((float)TextureValue->GetSizeX(), (float)TextureValue->GetSizeY(), (float)SizeZ);
 		const float OneOverMinDimension = 1.f / FMath::Max(TextureDimension.GetMin(), 1.f);
-		SetShaderValue(RHICmdList, ShaderRHI, TraceVolumeScalingParameter, FVector4f(
+		SetShaderValue(BatchedParameters, TraceVolumeScalingParameter, FVector4f(
 				TextureDimension.X * OneOverMinDimension, 
 				TextureDimension.Y * OneOverMinDimension, 
 				TextureDimension.Z * OneOverMinDimension, 
 				TextureDimension.GetMax() * OneOverMinDimension * .5f) // Extent
 			);
 
-		SetShaderValue(RHICmdList, ShaderRHI, TextureDimensionParameter, FVector3f(TextureDimension.X, TextureDimension.Y, TextureDimension.Z));
+		SetShaderValue(BatchedParameters, TextureDimensionParameter, FVector3f(TextureDimension.X, TextureDimension.Y, TextureDimension.Z));
 
-		SetShaderValue(RHICmdList, ShaderRHI, TraceViewMatrixParameter, FMatrix44f(FRotationMatrix::Make(TraceOrientation)));
+		SetShaderValue(BatchedParameters, TraceViewMatrixParameter, FMatrix44f(FRotationMatrix::Make(TraceOrientation)));
 	}
 
 private:
@@ -206,6 +205,6 @@ void FBatchedElementVolumeTexturePreviewParameters::BindShaders(
 	RHICmdList.ApplyCachedRenderTargets(GraphicsPSOInit);
 	SetGraphicsPipelineState(RHICmdList, GraphicsPSOInit, 0);
 
-	VertexShader->SetParameters(RHICmdList, InTransform);
-	PixelShader->SetParameters(RHICmdList, Texture, SizeZ, ColorWeights, InGamma, MipLevel, Opacity, TraceOrientation, bUsePointSampling);
+	SetAllShaderParametersVS(RHICmdList, VertexShader, InTransform);
+	SetAllShaderParametersPS(RHICmdList, PixelShader, Texture, SizeZ, ColorWeights, InGamma, MipLevel, Opacity, TraceOrientation, bUsePointSampling);
 }
diff --git a/Engine/Source/Runtime/Engine/Private/BatchedElements.cpp b/Engine/Source/Runtime/Engine/Private/BatchedElements.cpp
--- a/Engine/Source/Runtime/Engine/Private/BatchedElements.cpp
+++ b/Engine/Source/Runtime/Engine/Private/BatchedElements.cpp
@@ -9,6 +9,7 @@
 #include "Misc/LargeWorldRenderPosition.h"
 #include "PipelineStateCache.h"
 #include "SceneRelativeViewMatrices.h"
+#include "ShaderParameterUtils.h"
 #include "GlobalRenderResources.h"
 #include "HDRHelper.h"
 
@@ -669,7 +670,7 @@ void FBatchedElements::PrepareShaders(
 
 			SetGraphicsPipelineState(RHICmdList, GraphicsPSOInit, StencilRef);
 
-			HitTestingPixelShader->SetParameters(RHICmdList, Texture);
+			SetAllShaderParametersPS(RHICmdList, HitTestingPixelShader, Texture);
 		}
 		else
 		{
@@ -686,7 +687,8 @@ void FBatchedElements::PrepareShaders(
 					SetGraphicsPipelineState(RHICmdList, GraphicsPSOInit, StencilRef);
 
 					MaskedPixelShader->SetEditorCompositingParameters(RHICmdList, View);
-					MaskedPixelShader->SetParameters(RHICmdList, Texture, Gamma, OpacityMaskRefVal, BlendMode);
+
+					SetAllShaderParametersPS(RHICmdList, MaskedPixelShader, Texture, Gamma, OpacityMaskRefVal, BlendMode);
 				}
 				else
 				{
@@ -696,7 +698,7 @@ void FBatchedElements::PrepareShaders(
 					SetGraphicsPipelineState(RHICmdList, GraphicsPSOInit, StencilRef);
 
 					MaskedPixelShader->SetEditorCompositingParameters(RHICmdList, View);
-					MaskedPixelShader->SetParameters(RHICmdList, Texture, Gamma, OpacityMaskRefVal, BlendMode);
+					SetAllShaderParametersPS(RHICmdList, MaskedPixelShader, Texture, Gamma, OpacityMaskRefVal, BlendMode);
 				}
 			}
 			// render distance field elements
@@ -734,9 +736,10 @@ void FBatchedElements::PrepareShaders(
 					BlendMode == SE_BLEND_MaskedDistanceFieldShadowed || 
 					BlendMode == SE_BLEND_TranslucentDistanceFieldShadowed
 					);
-				
-				DistanceFieldPixelShader->SetParameters(
-					RHICmdList, 
+
+				SetAllShaderParametersPS(
+					RHICmdList,
+					DistanceFieldPixelShader,
 					Texture,
 					Gamma,
 					AlphaRefVal,
@@ -748,6 +751,9 @@ void FBatchedElements::PrepareShaders(
 					(GlowInfo != NULL) ? *GlowInfo : FDepthFieldGlowInfo(),
 					BlendMode
 					);
+
+				// This shader does not use editor compositing
+				DistanceFieldPixelShader->SetEditorCompositingParameters(RHICmdList, nullptr);
 			}
 			else if(BlendMode == SE_BLEND_TranslucentAlphaOnly || BlendMode == SE_BLEND_TranslucentAlphaOnlyWriteAlpha)
 			{
@@ -760,7 +766,7 @@ void FBatchedElements::PrepareShaders(
 
 					SetGraphicsPipelineState(RHICmdList, GraphicsPSOInit, StencilRef);
 
-					AlphaOnlyPixelShader->SetParameters(RHICmdList, Texture);
+					SetAllShaderParametersPS(RHICmdList, AlphaOnlyPixelShader, Texture);
 					AlphaOnlyPixelShader->SetEditorCompositingParameters(RHICmdList, View);
 				}
 				else
@@ -770,7 +776,7 @@ void FBatchedElements::PrepareShaders(
 
 					SetGraphicsPipelineState(RHICmdList, GraphicsPSOInit, StencilRef);
 
-					GammaAlphaOnlyPixelShader->SetParameters(RHICmdList, Texture, Gamma, BlendMode);
+					SetAllShaderParametersPS(RHICmdList, GammaAlphaOnlyPixelShader, Texture, Gamma, BlendMode);
 					GammaAlphaOnlyPixelShader->SetEditorCompositingParameters(RHICmdList, View);
 				}
 			}
@@ -781,7 +787,7 @@ void FBatchedElements::PrepareShaders(
 
 				SetGraphicsPipelineState(RHICmdList, GraphicsPSOInit, StencilRef);
 			
-				ColorChannelMaskPixelShader->SetParameters(RHICmdList, Texture, ColorWeights, GammaToUse );
+				SetAllShaderParametersPS(RHICmdList, ColorChannelMaskPixelShader, Texture, ColorWeights, GammaToUse);
 			}
 			else
 			{
@@ -794,7 +800,7 @@ void FBatchedElements::PrepareShaders(
 
 					SetGraphicsPipelineState(RHICmdList, GraphicsPSOInit, StencilRef);
 
-					PixelShader->SetParameters(RHICmdList, Texture);
+					SetAllShaderParametersPS(RHICmdList, PixelShader, Texture);
 					PixelShader->SetEditorCompositingParameters(RHICmdList, View);
 				}
 				else
@@ -814,14 +820,14 @@ void FBatchedElements::PrepareShaders(
 					GraphicsPSOInit.BoundShaderState.PixelShaderRHI = BasePixelShader.GetPixelShader();
 					SetGraphicsPipelineState(RHICmdList, GraphicsPSOInit, StencilRef);
 
-					BasePixelShader->SetParameters(RHICmdList, Texture, Gamma, BlendMode);
+					SetAllShaderParametersPS(RHICmdList, BasePixelShader, Texture, Gamma, BlendMode);
 					BasePixelShader->SetEditorCompositingParameters(RHICmdList, View);
 				}
 			}
 		}
 
 		// Set the simple element vertex shader parameters
-		VertexShader->SetParameters(RHICmdList, ViewMatrices);
+		SetAllShaderParametersVS(RHICmdList, VertexShader, ViewMatrices);
 	}
 }
 
diff --git a/Engine/Source/Runtime/Engine/Private/CubemapUnwrapUtils.cpp b/Engine/Source/Runtime/Engine/Private/CubemapUnwrapUtils.cpp
--- a/Engine/Source/Runtime/Engine/Private/CubemapUnwrapUtils.cpp
+++ b/Engine/Source/Runtime/Engine/Private/CubemapUnwrapUtils.cpp
@@ -18,6 +18,7 @@
 #include "Engine/TextureCube.h"
 #include "Engine/TextureRenderTargetCube.h"
 #include "RHIStaticStates.h"
+#include "DataDrivenShaderPlatformInfo.h"
 
 IMPLEMENT_SHADER_TYPE(,FCubemapTexturePropertiesVS,TEXT("/Engine/Private/SimpleElementVertexShader.usf"),TEXT("Main"),SF_Vertex);
 IMPLEMENT_SHADER_TYPE(,FCubemapTexturePropertiesPS,TEXT("/Engine/Private/SimpleElementPixelShader.usf"),TEXT("CubemapTextureProperties"),SF_Pixel);
@@ -131,9 +132,16 @@ bool FCubemapTexturePropertiesVS::ShouldCompilePermutation(const FGlobalShaderPe
 	return IsPCPlatform(Parameters.Platform);
 }
 
-void FCubemapTexturePropertiesVS::SetParameters( FRHICommandList& RHICmdList, const FMatrix& TransformValue )
+void FCubemapTexturePropertiesVS::SetParameters(FRHIBatchedShaderParameters& BatchedParameters, const FMatrix& TransformValue )
 {
+	SetShaderValue(BatchedParameters, Transform, (FMatrix44f)TransformValue);
+}
+
+void FCubemapTexturePropertiesVS::SetParameters(FRHICommandList& RHICmdList, const FMatrix& TransformValue)
+{
+	FRHIBatchedShaderParameters BatchedParameters;
 	SetShaderValue(RHICmdList, RHICmdList.GetBoundVertexShader(), Transform, (FMatrix44f)TransformValue);
+	RHICmdList.SetBatchedShaderParameters(RHICmdList.GetBoundVertexShader(), BatchedParameters);
 }
 
 FCubemapTexturePropertiesPS::FCubemapTexturePropertiesPS() = default;
@@ -166,21 +174,19 @@ bool FCubemapTexturePropertiesPS::ShouldCompilePermutation(const FGlobalShaderPe
 	return true;
 }
 
-void FCubemapTexturePropertiesPS::SetParameters(FRHICommandList& RHICmdList, const FTexture* InTexture, const FMatrix& InColorWeightsValue, float InMipLevel, float InSliceIndex, bool bInIsTextureCubeArray, const FMatrix44f& InViewMatrix, bool bInShowLongLatUnwrap, float InGammaValue, bool bInUsePointSampling)
+void FCubemapTexturePropertiesPS::SetParameters(FRHIBatchedShaderParameters& BatchedParameters, const FTexture* InTexture, const FMatrix& InColorWeightsValue, float InMipLevel, float InSliceIndex, bool bInIsTextureCubeArray, const FMatrix44f& InViewMatrix, bool bInShowLongLatUnwrap, float InGammaValue, bool bInUsePointSampling)
 {
-	FRHIPixelShader* ShaderRHI = RHICmdList.GetBoundPixelShader();
-
 	if (InTexture != nullptr)
 	{
 		FRHISamplerState* SamplerState = bInUsePointSampling ? TStaticSamplerState<SF_Point>::GetRHI() : InTexture->SamplerStateRHI.GetReference();
-		SetTextureParameter(RHICmdList, ShaderRHI, CubeTexture, CubeTextureSampler, SamplerState, InTexture->TextureRHI);
+		SetTextureParameter(BatchedParameters, CubeTexture, CubeTextureSampler, SamplerState, InTexture->TextureRHI);
 	}
 
 	FVector4f PackedProperties0Value(InMipLevel, bInShowLongLatUnwrap ? 1.0f : -1.0f, 0, 0);
-	SetShaderValue(RHICmdList, ShaderRHI, PackedProperties0, PackedProperties0Value);
-	SetShaderValue(RHICmdList, ShaderRHI, ColorWeights, (FMatrix44f)InColorWeightsValue);
-	SetShaderValue(RHICmdList, ShaderRHI, Gamma, InGammaValue);
-	SetShaderValue(RHICmdList, RHICmdList.GetBoundPixelShader(), ViewMatrix, InViewMatrix);
+	SetShaderValue(BatchedParameters, PackedProperties0, PackedProperties0Value);
+	SetShaderValue(BatchedParameters, ColorWeights, (FMatrix44f)InColorWeightsValue);
+	SetShaderValue(BatchedParameters, Gamma, InGammaValue);
+	SetShaderValue(BatchedParameters, ViewMatrix, InViewMatrix);
 
 	// store slice count and selected slice index for the texture cube array
 	if (bInIsTextureCubeArray)
@@ -188,11 +194,18 @@ void FCubemapTexturePropertiesPS::SetParameters(FRHICommandList& RHICmdList, con
 		// GetSizeZ() returns the total number of slices stored in the platform data
 		// for a TextureCube array this value is equal to the size of the array multiplied by 6
 		const float NumSlicesData = (float)(InTexture != nullptr ? FMath::Max((int32)InTexture->GetSizeZ() / 6, 1) : 1);
-		SetShaderValue(RHICmdList, ShaderRHI, NumSlices, NumSlicesData);
-		SetShaderValue(RHICmdList, ShaderRHI, SliceIndex, InSliceIndex);
+		SetShaderValue(BatchedParameters, NumSlices, NumSlicesData);
+		SetShaderValue(BatchedParameters, SliceIndex, InSliceIndex);
 	}
 }
 
+void FCubemapTexturePropertiesPS::SetParameters(FRHICommandList& RHICmdList, const FTexture* InTexture, const FMatrix& InColorWeightsValue, float InMipLevel, float InSliceIndex, bool bInIsTextureCubeArray, const FMatrix44f& InViewMatrix, bool bInShowLongLatUnwrap, float InGammaValue, bool bInUsePointSampling)
+{
+	FRHIBatchedShaderParameters BatchedParameters;
+	SetParameters(BatchedParameters, InTexture, InColorWeightsValue, InMipLevel, InSliceIndex, bInIsTextureCubeArray, InViewMatrix, bInShowLongLatUnwrap, InGammaValue, bInUsePointSampling);
+	RHICmdList.SetBatchedShaderParameters(RHICmdList.GetBoundPixelShader(), BatchedParameters);
+}
+
 void FMipLevelBatchedElementParameters::BindShaders(FRHICommandList& RHICmdList, FGraphicsPipelineStateInitializer& GraphicsPSOInit, ERHIFeatureLevel::Type InFeatureLevel, const FMatrix& InTransform, const float InGamma, const FMatrix& ColorWeights, const FTexture* Texture)
 {
 	GraphicsPSOInit.BlendState = TStaticBlendState<>::GetRHI();
@@ -213,8 +226,8 @@ void FMipLevelBatchedElementParameters::BindShaders(FRHICommandList& RHICmdList,
 	RHICmdList.ApplyCachedRenderTargets(GraphicsPSOInit);
 	SetGraphicsPipelineState(RHICmdList, GraphicsPSOInit, 0);
 
-	VertexShader->SetParameters(RHICmdList, InTransform);
-	PixelShader->SetParameters(RHICmdList, Texture, ColorWeights, MipLevel, SliceIndex, bIsTextureCubeArray, ViewMatrix, bShowLongLatUnwrap, InGamma, bUsePointSampling);
+	SetAllShaderParametersVS(RHICmdList, VertexShader, InTransform);
+	SetAllShaderParametersPS(RHICmdList, PixelShader, Texture, ColorWeights, MipLevel, SliceIndex, bIsTextureCubeArray, ViewMatrix, bShowLongLatUnwrap, InGamma, bUsePointSampling);
 }
 
 FIESLightProfilePS::FIESLightProfilePS() = default;
@@ -232,12 +245,17 @@ bool FIESLightProfilePS::ShouldCompilePermutation(const FGlobalShaderPermutation
 	return IsFeatureLevelSupported(Parameters.Platform, ERHIFeatureLevel::SM5) && !IsConsolePlatform(Parameters.Platform);
 }
 
-void FIESLightProfilePS::SetParameters( FRHICommandList& RHICmdList, const FTexture* Texture, float InBrightnessInLumens )
+void FIESLightProfilePS::SetParameters(FRHIBatchedShaderParameters& BatchedParameters, const FTexture* Texture, float InBrightnessInLumens )
 {
-	FRHIPixelShader* ShaderRHI = RHICmdList.GetBoundPixelShader();
-	SetTextureParameter(RHICmdList, ShaderRHI, IESTexture, IESTextureSampler, Texture);
+	SetTextureParameter(BatchedParameters, IESTexture, IESTextureSampler, Texture);
+	SetShaderValue(BatchedParameters, BrightnessInLumens, InBrightnessInLumens);
+}
 
-	SetShaderValue(RHICmdList, ShaderRHI, BrightnessInLumens, InBrightnessInLumens);
+void FIESLightProfilePS::SetParameters(FRHICommandList& RHICmdList, const FTexture* Texture, float InBrightnessInLumens)
+{
+	FRHIBatchedShaderParameters BatchedParameters;
+	SetParameters(BatchedParameters, Texture, InBrightnessInLumens);
+	RHICmdList.SetBatchedShaderParameters(RHICmdList.GetBoundPixelShader(), BatchedParameters);
 }
 
 void FIESLightProfileBatchedElementParameters::BindShaders( FRHICommandList& RHICmdList, FGraphicsPipelineStateInitializer& GraphicsPSOInit, ERHIFeatureLevel::Type InFeatureLevel, const FMatrix& InTransform, const float InGamma, const FMatrix& ColorWeights, const FTexture* Texture )
@@ -255,6 +273,6 @@ void FIESLightProfileBatchedElementParameters::BindShaders( FRHICommandList& RHI
 	RHICmdList.ApplyCachedRenderTargets(GraphicsPSOInit);
 	SetGraphicsPipelineState(RHICmdList, GraphicsPSOInit, 0);
 
-	VertexShader->SetParameters(RHICmdList, InTransform);
-	PixelShader->SetParameters(RHICmdList, Texture, BrightnessInLumens);
+	SetAllShaderParametersVS(RHICmdList, VertexShader, InTransform);
+	SetAllShaderParametersPS(RHICmdList, PixelShader, Texture, BrightnessInLumens);
 }
diff --git a/Engine/Source/Runtime/Engine/Private/GPUSkinCache.cpp b/Engine/Source/Runtime/Engine/Private/GPUSkinCache.cpp
--- a/Engine/Source/Runtime/Engine/Private/GPUSkinCache.cpp
+++ b/Engine/Source/Runtime/Engine/Private/GPUSkinCache.cpp
@@ -660,65 +660,61 @@ public:
 	}
 
 	void SetParameters(
-		FRHICommandListImmediate& RHICmdList, 
+		FRHIBatchedShaderParameters& BatchedParameters,
 		const FVertexBufferAndSRV& BoneBuffer,
 		FGPUSkinCacheEntry* Entry,
-		FGPUSkinCacheEntry::FSectionDispatchData& DispatchData,
-		FRHIUnorderedAccessView* PositionUAV, FRHIUnorderedAccessView* TangentUAV
+		const FGPUSkinCacheEntry::FSectionDispatchData& DispatchData,
+		FRHIUnorderedAccessView* PositionUAV,
+		FRHIUnorderedAccessView* TangentUAV
 		)
 	{
-		FRHIComputeShader* ShaderRHI = RHICmdList.GetBoundComputeShader();
-
-		SetShaderValue(RHICmdList, ShaderRHI, NumVertices, DispatchData.NumVertices);
-		SetShaderValue(RHICmdList, ShaderRHI, InputStreamStart, DispatchData.InputStreamStart);
+		SetShaderValue(BatchedParameters, NumVertices, DispatchData.NumVertices);
+		SetShaderValue(BatchedParameters, InputStreamStart, DispatchData.InputStreamStart);
 
 		check(BoneBuffer.VertexBufferSRV);
-		SetSRVParameter(RHICmdList, ShaderRHI, BoneMatrices, BoneBuffer.VertexBufferSRV);
+		SetSRVParameter(BatchedParameters, BoneMatrices, BoneBuffer.VertexBufferSRV);
 
-		SetSRVParameter(RHICmdList, ShaderRHI, TangentInputBuffer, DispatchData.TangentBufferSRV);
-		SetSRVParameter(RHICmdList, ShaderRHI, PositionInputBuffer, DispatchData.PositionBufferSRV);
+		SetSRVParameter(BatchedParameters, TangentInputBuffer, DispatchData.TangentBufferSRV);
+		SetSRVParameter(BatchedParameters, PositionInputBuffer, DispatchData.PositionBufferSRV);
 
-		SetShaderValue(RHICmdList, ShaderRHI, NumBoneInfluences, DispatchData.NumBoneInfluences);
-		SetShaderValue(RHICmdList, ShaderRHI, InputWeightIndexSize, Entry->InputWeightIndexSize);
-		SetShaderValue(RHICmdList, ShaderRHI, InputWeightStart, DispatchData.InputWeightStart);
-		SetShaderValue(RHICmdList, ShaderRHI, InputWeightStride, Entry->InputWeightStride);
-		SetSRVParameter(RHICmdList, ShaderRHI, InputWeightStream, Entry->InputWeightStreamSRV);
-		SetSRVParameter(RHICmdList, ShaderRHI, InputWeightLookupStream, Entry->InputWeightLookupStreamSRV);
+		SetShaderValue(BatchedParameters, NumBoneInfluences, DispatchData.NumBoneInfluences);
+		SetShaderValue(BatchedParameters, InputWeightIndexSize, Entry->InputWeightIndexSize);
+		SetShaderValue(BatchedParameters, InputWeightStart, DispatchData.InputWeightStart);
+		SetShaderValue(BatchedParameters, InputWeightStride, Entry->InputWeightStride);
+		SetSRVParameter(BatchedParameters, InputWeightStream, Entry->InputWeightStreamSRV);
+		SetSRVParameter(BatchedParameters, InputWeightLookupStream, Entry->InputWeightLookupStreamSRV);
 
 		// output UAV
-		SetUAVParameter(RHICmdList, ShaderRHI, PositionBufferUAV, PositionUAV);
-		SetUAVParameter(RHICmdList, ShaderRHI, TangentBufferUAV, TangentUAV);
-		SetShaderValue(RHICmdList, ShaderRHI, SkinCacheStart, DispatchData.OutputStreamStart);
+		SetUAVParameter(BatchedParameters, PositionBufferUAV, PositionUAV);
+		SetUAVParameter(BatchedParameters, TangentBufferUAV, TangentUAV);
+		SetShaderValue(BatchedParameters, SkinCacheStart, DispatchData.OutputStreamStart);
 
 		const bool bMorph = DispatchData.SkinType == 1;
 		if (bMorph)
 		{
-			SetSRVParameter(RHICmdList, ShaderRHI, MorphBuffer, Entry->MorphBuffer);
-			SetShaderValue(RHICmdList, ShaderRHI, MorphBufferOffset, DispatchData.MorphBufferOffset);
+			SetSRVParameter(BatchedParameters, MorphBuffer, Entry->MorphBuffer);
+			SetShaderValue(BatchedParameters, MorphBufferOffset, DispatchData.MorphBufferOffset);
 		}
 
 		const bool bCloth = DispatchData.SkinType == 2;
 		if (bCloth)
 		{
-			SetSRVParameter(RHICmdList, ShaderRHI, ClothBuffer, Entry->ClothBuffer);
-			SetSRVParameter(RHICmdList, ShaderRHI, ClothPositionsAndNormalsBuffer, DispatchData.ClothPositionsAndNormalsBuffer);
-			SetShaderValue(RHICmdList, ShaderRHI, ClothBufferOffset, DispatchData.ClothBufferOffset);
-			SetShaderValue(RHICmdList, ShaderRHI, ClothBlendWeight, DispatchData.ClothBlendWeight);
-			SetShaderValue(RHICmdList, ShaderRHI, ClothToLocal, DispatchData.ClothToLocal);
-			SetShaderValue(RHICmdList, ShaderRHI, ClothNumInfluencesPerVertex, DispatchData.ClothNumInfluencesPerVertex);
-			SetShaderValue(RHICmdList, ShaderRHI, WorldScale, DispatchData.WorldScale);
+			SetSRVParameter(BatchedParameters, ClothBuffer, Entry->ClothBuffer);
+			SetSRVParameter(BatchedParameters, ClothPositionsAndNormalsBuffer, DispatchData.ClothPositionsAndNormalsBuffer);
+			SetShaderValue(BatchedParameters, ClothBufferOffset, DispatchData.ClothBufferOffset);
+			SetShaderValue(BatchedParameters, ClothBlendWeight, DispatchData.ClothBlendWeight);
+			SetShaderValue(BatchedParameters, ClothToLocal, DispatchData.ClothToLocal);
+			SetShaderValue(BatchedParameters, ClothNumInfluencesPerVertex, DispatchData.ClothNumInfluencesPerVertex);
+			SetShaderValue(BatchedParameters, WorldScale, DispatchData.WorldScale);
 		}
 
-		SetShaderValue(RHICmdList, ShaderRHI, SkinCacheDebug, CVarGPUSkinCacheDebug.GetValueOnRenderThread());
+		SetShaderValue(BatchedParameters, SkinCacheDebug, CVarGPUSkinCacheDebug.GetValueOnRenderThread());
 	}
 
-
-	void UnsetParameters(FRHICommandList& RHICmdList)
+	void UnsetParameters(FRHIBatchedShaderParameters& BatchedParameters)
 	{
-		FRHIComputeShader* ShaderRHI = RHICmdList.GetBoundComputeShader();
-
-		SetUAVParameter(RHICmdList, ShaderRHI, PositionBufferUAV, 0);
-		SetUAVParameter(RHICmdList, ShaderRHI, TangentBufferUAV, 0);
+		SetUAVParameter(BatchedParameters, PositionBufferUAV, 0);
+		SetUAVParameter(BatchedParameters, TangentBufferUAV, 0);
 	}
 
 private:
@@ -912,45 +908,41 @@ public:
         DuplicatedIndicesIndices.Bind(Initializer.ParameterMap, TEXT("DuplicatedIndicesIndices"));
 	}
 
-	void SetParameters(FRHICommandListImmediate& RHICmdList, FGPUSkinCacheEntry* Entry, FGPUSkinCacheEntry::FSectionDispatchData& DispatchData, FRWBuffer& StagingBuffer)
+	void SetParameters(FRHIBatchedShaderParameters& BatchedParameters, FGPUSkinCacheEntry* Entry, const FGPUSkinCacheEntry::FSectionDispatchData& DispatchData, const FRWBuffer& StagingBuffer)
 	{
-		FRHIComputeShader* ShaderRHI = RHICmdList.GetBoundComputeShader();
-
 //later		FGlobalShader::SetParameters<FViewUniformShaderParameters>(RHICmdList, ShaderRHI, View);
 
-		SetShaderValue(RHICmdList, ShaderRHI, NumTriangles, DispatchData.NumTriangles);
+		SetShaderValue(BatchedParameters, NumTriangles, DispatchData.NumTriangles);
 
-		SetSRVParameter(RHICmdList, ShaderRHI, GPUPositionCacheBuffer, DispatchData.GetPositionRWBuffer()->Buffer.SRV);
-		SetSRVParameter(RHICmdList, ShaderRHI, GPUTangentCacheBuffer, DispatchData.GetActiveTangentRWBuffer()->Buffer.SRV);
-		SetSRVParameter(RHICmdList, ShaderRHI, UVsInputBuffer, DispatchData.UVsBufferSRV);
+		SetSRVParameter(BatchedParameters, GPUPositionCacheBuffer, DispatchData.GetPositionRWBuffer()->Buffer.SRV);
+		SetSRVParameter(BatchedParameters, GPUTangentCacheBuffer, DispatchData.GetActiveTangentRWBuffer()->Buffer.SRV);
+		SetSRVParameter(BatchedParameters, UVsInputBuffer, DispatchData.UVsBufferSRV);
 
-		SetShaderValue(RHICmdList, ShaderRHI, SkinCacheStart, DispatchData.OutputStreamStart);
+		SetShaderValue(BatchedParameters, SkinCacheStart, DispatchData.OutputStreamStart);
 
-		SetSRVParameter(RHICmdList, ShaderRHI, IndexBuffer, DispatchData.IndexBuffer);
-		SetShaderValue(RHICmdList, ShaderRHI, IndexBufferOffset, DispatchData.IndexBufferOffsetValue);
+		SetSRVParameter(BatchedParameters, IndexBuffer, DispatchData.IndexBuffer);
+		SetShaderValue(BatchedParameters, IndexBufferOffset, DispatchData.IndexBufferOffsetValue);
 		
-		SetShaderValue(RHICmdList, ShaderRHI, InputStreamStart, DispatchData.InputStreamStart);
-		SetShaderValue(RHICmdList, ShaderRHI, NumTexCoords, DispatchData.NumTexCoords);
-		SetShaderValue(RHICmdList, ShaderRHI, SelectedTexCoord, DispatchData.SelectedTexCoord);
-		SetSRVParameter(RHICmdList, ShaderRHI, TangentInputBuffer, DispatchData.TangentBufferSRV);
-		SetSRVParameter(RHICmdList, ShaderRHI, TangentInputBuffer, DispatchData.UVsBufferSRV);
+		SetShaderValue(BatchedParameters, InputStreamStart, DispatchData.InputStreamStart);
+		SetShaderValue(BatchedParameters, NumTexCoords, DispatchData.NumTexCoords);
+		SetShaderValue(BatchedParameters, SelectedTexCoord, DispatchData.SelectedTexCoord);
+		SetSRVParameter(BatchedParameters, TangentInputBuffer, DispatchData.TangentBufferSRV);
+		SetSRVParameter(BatchedParameters, TangentInputBuffer, DispatchData.UVsBufferSRV);
 
 		// UAV
-		SetUAVParameter(RHICmdList, ShaderRHI, IntermediateAccumBufferUAV, StagingBuffer.UAV);
-		SetShaderValue(RHICmdList, ShaderRHI, IntermediateAccumBufferOffset, GRecomputeTangentsParallelDispatch * DispatchData.IntermediateAccumulatedTangentBufferOffset);
+		SetUAVParameter(BatchedParameters, IntermediateAccumBufferUAV, StagingBuffer.UAV);
+		SetShaderValue(BatchedParameters, IntermediateAccumBufferOffset, GRecomputeTangentsParallelDispatch * DispatchData.IntermediateAccumulatedTangentBufferOffset);
 
         if (!GAllowDupedVertsForRecomputeTangents)
         {
-		    SetSRVParameter(RHICmdList, ShaderRHI, DuplicatedIndices, DispatchData.DuplicatedIndices);
-            SetSRVParameter(RHICmdList, ShaderRHI, DuplicatedIndicesIndices, DispatchData.DuplicatedIndicesIndices);
+		    SetSRVParameter(BatchedParameters, DuplicatedIndices, DispatchData.DuplicatedIndices);
+            SetSRVParameter(BatchedParameters, DuplicatedIndicesIndices, DispatchData.DuplicatedIndicesIndices);
         }
 	}
 
-	void UnsetParameters(FRHICommandList& RHICmdList)
+	void UnsetParameters(FRHIBatchedShaderParameters& BatchedParameters)
 	{
-		FRHIComputeShader* ShaderRHI = RHICmdList.GetBoundComputeShader();
-
-		SetUAVParameter(RHICmdList, ShaderRHI, IntermediateAccumBufferUAV, 0);
+		SetUAVParameter(BatchedParameters, IntermediateAccumBufferUAV, 0);
 	}
 
 	LAYOUT_FIELD(FShaderResourceParameter, IntermediateAccumBufferUAV);
@@ -1042,35 +1034,31 @@ public:
 		VertexColorChannel.Bind(Initializer.ParameterMap, TEXT("VertexColorChannel"));
 	}
 
-	void SetParameters(FRHICommandListImmediate& RHICmdList, FGPUSkinCacheEntry* Entry, FGPUSkinCacheEntry::FSectionDispatchData& DispatchData, FRWBuffer& StagingBuffer)
+	void SetParameters(FRHIBatchedShaderParameters& BatchedParameters, FGPUSkinCacheEntry* Entry, const FGPUSkinCacheEntry::FSectionDispatchData& DispatchData, const FRWBuffer& StagingBuffer)
 	{
-		FRHIComputeShader* ShaderRHI = RHICmdList.GetBoundComputeShader();
-
 		check(StagingBuffer.UAV);
 
-		//later		FGlobalShader::SetParameters<FViewUniformShaderParameters>(RHICmdList, ShaderRHI, View);
+		//later		FGlobalShader::SetParameters<FViewUniformShaderParameters>(BatchedParameters, View);
 
-		SetShaderValue(RHICmdList, ShaderRHI, SkinCacheStart, DispatchData.OutputStreamStart);
-		SetShaderValue(RHICmdList, ShaderRHI, NumVertices, DispatchData.NumVertices);
-		SetShaderValue(RHICmdList, ShaderRHI, InputStreamStart, DispatchData.InputStreamStart);
-		SetShaderValue(RHICmdList, ShaderRHI, VertexColorChannel, uint32(DispatchData.Section->RecomputeTangentsVertexMaskChannel));
+		SetShaderValue(BatchedParameters, SkinCacheStart, DispatchData.OutputStreamStart);
+		SetShaderValue(BatchedParameters, NumVertices, DispatchData.NumVertices);
+		SetShaderValue(BatchedParameters, InputStreamStart, DispatchData.InputStreamStart);
+		SetShaderValue(BatchedParameters, VertexColorChannel, uint32(DispatchData.Section->RecomputeTangentsVertexMaskChannel));
 
 		// UAVs
-		SetUAVParameter(RHICmdList, ShaderRHI, IntermediateAccumBufferUAV, StagingBuffer.UAV);
-		SetShaderValue(RHICmdList, ShaderRHI, IntermediateAccumBufferOffset, GRecomputeTangentsParallelDispatch * DispatchData.IntermediateAccumulatedTangentBufferOffset);
-		SetUAVParameter(RHICmdList, ShaderRHI, TangentBufferUAV, DispatchData.GetTangentRWBuffer()->Buffer.UAV);
+		SetUAVParameter(BatchedParameters, IntermediateAccumBufferUAV, StagingBuffer.UAV);
+		SetShaderValue(BatchedParameters, IntermediateAccumBufferOffset, GRecomputeTangentsParallelDispatch * DispatchData.IntermediateAccumulatedTangentBufferOffset);
+		SetUAVParameter(BatchedParameters, TangentBufferUAV, DispatchData.GetTangentRWBuffer()->Buffer.UAV);
 
-		SetSRVParameter(RHICmdList, ShaderRHI, TangentInputBuffer, DispatchData.IntermediateTangentBuffer ? DispatchData.IntermediateTangentBuffer->Buffer.SRV : nullptr);
+		SetSRVParameter(BatchedParameters, TangentInputBuffer, DispatchData.IntermediateTangentBuffer ? DispatchData.IntermediateTangentBuffer->Buffer.SRV : nullptr);
 
-		SetSRVParameter(RHICmdList, ShaderRHI, ColorInputBuffer, DispatchData.ColorBufferSRV);
+		SetSRVParameter(BatchedParameters, ColorInputBuffer, DispatchData.ColorBufferSRV);
 	}
 
-	void UnsetParameters(FRHICommandList& RHICmdList)
+	void UnsetParameters(FRHIBatchedShaderParameters& BatchedParameters)
 	{
-		FRHIComputeShader* ShaderRHI = RHICmdList.GetBoundComputeShader();
-
-		SetUAVParameter(RHICmdList, ShaderRHI, TangentBufferUAV, 0);
-		SetUAVParameter(RHICmdList, ShaderRHI, IntermediateAccumBufferUAV, 0);
+		SetUAVParameter(BatchedParameters, TangentBufferUAV, 0);
+		SetUAVParameter(BatchedParameters, IntermediateAccumBufferUAV, 0);
 	}
 };
 
@@ -1174,7 +1162,7 @@ void FGPUSkinCache::DispatchUpdateSkinTangents(FRHICommandListImmediate& RHICmdL
 				*RayTracingTag , *GetSkeletalMeshObjectName(Entry->GPUSkin), LODIndex, SectionIndex, DispatchData.IndexBufferOffsetValue, DispatchData.NumTriangles, Entry->BoneInfluenceType, bFullPrecisionUV);
 
 			FRHIComputeShader* ShaderRHI = Shader.GetComputeShader();
-			SetComputePipelineState(RHICmdList, ShaderRHI);
+			SetComputePipelineState(RHICmdList, Shader.GetComputeShader());
 
 			if (!GAllowDupedVertsForRecomputeTangents)
 			{
@@ -1195,9 +1183,13 @@ void FGPUSkinCache::DispatchUpdateSkinTangents(FRHICommandListImmediate& RHICmdL
 			}
 
 			INC_DWORD_STAT_BY(STAT_GPUSkinCache_NumTrianglesForRecomputeTangents, NumTriangles);
-			Shader->SetParameters(RHICmdList, Entry, DispatchData, GRecomputeTangentsParallelDispatch ? DispatchData.GetIntermediateAccumulatedTangentBuffer()->Buffer : StagingBuffer->Buffer);
+
+			const FRWBuffer& ShaderStagingBuffer = GRecomputeTangentsParallelDispatch ? DispatchData.GetIntermediateAccumulatedTangentBuffer()->Buffer : StagingBuffer->Buffer;
+
+			SetAllShaderParametersCS(RHICmdList, Shader, Entry, DispatchData, ShaderStagingBuffer);
 			DispatchComputeShader(RHICmdList, Shader.GetShader(), ThreadGroupCountValue, 1, 1);
-			Shader->UnsetParameters(RHICmdList);
+			UnsetAllShaderParametersCS(RHICmdList, Shader);
+
 			IncrementDispatchCounter(RHICmdList);
 		}
 	}
@@ -1229,9 +1221,10 @@ void FGPUSkinCache::DispatchUpdateSkinTangents(FRHICommandListImmediate& RHICmdL
 				});
 		}
 
-		ComputeShader->SetParameters(RHICmdList, Entry, DispatchData, GRecomputeTangentsParallelDispatch ? DispatchData.GetIntermediateAccumulatedTangentBuffer()->Buffer : StagingBuffer->Buffer);
+		SetAllShaderParametersCS(RHICmdList, ComputeShader, Entry, DispatchData, GRecomputeTangentsParallelDispatch ? DispatchData.GetIntermediateAccumulatedTangentBuffer()->Buffer : StagingBuffer->Buffer);
 		DispatchComputeShader(RHICmdList, ComputeShader.GetShader(), ThreadGroupCountValue, 1, 1);
-		ComputeShader->UnsetParameters(RHICmdList);
+		UnsetAllShaderParametersCS(RHICmdList, ComputeShader);
+
 		IncrementDispatchCounter(RHICmdList);
 	}
 }
@@ -1890,19 +1883,21 @@ void FGPUSkinCache::DispatchUpdateSkinning(FRHICommandListImmediate& RHICmdList,
 	{
 		SetComputePipelineState(RHICmdList, Shader.GetComputeShader());
 
-		Shader->SetParameters(
+		SetAllShaderParametersCS(
 			RHICmdList,
+			Shader,
 			PrevBoneBuffer,
 			Entry,
 			DispatchData,
 			DispatchData.GetPreviousPositionRWBuffer()->Buffer.UAV,
 			DispatchData.GetActiveTangentRWBuffer() ? DispatchData.GetActiveTangentRWBuffer()->Buffer.UAV : nullptr
-			);
+		);
 
 		uint32 VertexCountAlign64 = FMath::DivideAndRoundUp(DispatchData.NumVertices, (uint32)64);
 		INC_DWORD_STAT_BY(STAT_GPUSkinCache_TotalNumVertices, VertexCountAlign64 * 64);
 		RHICmdList.DispatchComputeShader(VertexCountAlign64, 1, 1);
-		Shader->UnsetParameters(RHICmdList);
+		UnsetAllShaderParametersCS(RHICmdList, Shader);
+
 		IncrementDispatchCounter(RHICmdList);
 		BuffersToTransitionToRead.Add(DispatchData.GetPreviousPositionRWBuffer());
 	}
@@ -1911,19 +1906,22 @@ void FGPUSkinCache::DispatchUpdateSkinning(FRHICommandListImmediate& RHICmdList,
 	{
 		SetComputePipelineState(RHICmdList, Shader.GetComputeShader());
 
-		Shader->SetParameters(
-			RHICmdList, 
-			BoneBuffer, 
-			Entry, 
-			DispatchData, 
-			DispatchData.GetPositionRWBuffer()->Buffer.UAV, 
+		SetAllShaderParametersCS(
+			RHICmdList,
+			Shader,
+			BoneBuffer,
+			Entry,
+			DispatchData,
+			DispatchData.GetPositionRWBuffer()->Buffer.UAV,
 			DispatchData.GetActiveTangentRWBuffer() ? DispatchData.GetActiveTangentRWBuffer()->Buffer.UAV : nullptr
-			);
+		);
 
 		uint32 VertexCountAlign64 = FMath::DivideAndRoundUp(DispatchData.NumVertices, (uint32)64);
 		INC_DWORD_STAT_BY(STAT_GPUSkinCache_TotalNumVertices, VertexCountAlign64 * 64);
 		RHICmdList.DispatchComputeShader(VertexCountAlign64, 1, 1);
-		Shader->UnsetParameters(RHICmdList);
+
+		UnsetAllShaderParametersCS(RHICmdList, Shader);
+
 		IncrementDispatchCounter(RHICmdList);
 		BuffersToTransitionToRead.Add(DispatchData.GetPositionRWBuffer());
 	}
diff --git a/Engine/Source/Runtime/Engine/Private/GPUSort.cpp b/Engine/Source/Runtime/Engine/Private/GPUSort.cpp
--- a/Engine/Source/Runtime/Engine/Private/GPUSort.cpp
+++ b/Engine/Source/Runtime/Engine/Private/GPUSort.cpp
@@ -216,20 +216,17 @@ TGlobalResource<FRadixSortParametersBuffer> GRadixSortParametersBuffer;
 class FRadixSortClearOffsetsCS : public FGlobalShader
 {
 	DECLARE_SHADER_TYPE(FRadixSortClearOffsetsCS,Global);
-
 public:
 
-	static void ModifyCompilationEnvironment(const FGlobalShaderPermutationParameters& Parameters, FShaderCompilerEnvironment& OutEnvironment )
+	static void ModifyCompilationEnvironment(const FGlobalShaderPermutationParameters& Parameters, FShaderCompilerEnvironment& OutEnvironment)
 	{
-		FGlobalShader::ModifyCompilationEnvironment( Parameters, OutEnvironment );
-		OutEnvironment.SetDefine( TEXT("RADIX_SORT_CLEAR_OFFSETS"), 1 );
-		SetRadixSortShaderCompilerEnvironment( OutEnvironment );
+		FGlobalShader::ModifyCompilationEnvironment(Parameters, OutEnvironment);
+		OutEnvironment.SetDefine(TEXT("RADIX_SORT_CLEAR_OFFSETS"), 1);
+		SetRadixSortShaderCompilerEnvironment(OutEnvironment);
 	}
 
 	/** Default constructor. */
-	FRadixSortClearOffsetsCS()
-	{
-	}
+	FRadixSortClearOffsetsCS() = default;
 
 	/** Initialization constructor. */
 	explicit FRadixSortClearOffsetsCS( const ShaderMetaType::CompiledShaderInitializerType& Initializer )
@@ -238,31 +235,14 @@ public:
 		OutOffsets.Bind( Initializer.ParameterMap, TEXT("OutOffsets") );
 	}
 
-	/**
-	 * Set parameters for this shader.
-	 */
-	void SetParameters(FRHICommandList& RHICmdList, const FRadixSortUniformBufferRef& UniformBuffer )
+	void SetParameters(FRHIBatchedShaderParameters& BatchedParameters, FRHIUnorderedAccessView* OutOffsetsUAV)
 	{
-		FRHIComputeShader* ComputeShaderRHI = RHICmdList.GetBoundComputeShader();
-		SetUniformBufferParameter(RHICmdList, ComputeShaderRHI, GetUniformBufferParameter<FRadixSortParameters>(), UniformBuffer );
+		SetUAVParameter(BatchedParameters, OutOffsets, OutOffsetsUAV);
 	}
 
-	/**
-	 * Set output buffer for this shader.
-	 */
-	void SetOutput(FRHICommandList& RHICmdList, FRHIUnorderedAccessView* OutOffsetsUAV)
+	void UnsetParameters(FRHIBatchedShaderParameters& BatchedParameters)
 	{
-		FRHIComputeShader* ComputeShaderRHI = RHICmdList.GetBoundComputeShader();
-		SetUAVParameter(RHICmdList, ComputeShaderRHI, OutOffsets, OutOffsetsUAV);
-	}
-
-	/**
-	 * Unbinds any buffers that have been bound.
-	 */
-	void UnbindBuffers(FRHICommandList& RHICmdList)
-	{
-		FRHIComputeShader* ComputeShaderRHI = RHICmdList.GetBoundComputeShader();
-		SetUAVParameter(RHICmdList, ComputeShaderRHI, OutOffsets, nullptr);
+		SetUAVParameter(BatchedParameters, OutOffsets, nullptr);
 	}
 
 private:
@@ -314,35 +294,18 @@ public:
 		return RadixSortParameterBuffer.IsBound();
 	}
 
-	/**
-	 * Set parameters for this shader.
-	 */
-	void SetParameters(FRHICommandList& RHICmdList, FRHIShaderResourceView* InKeysSRV, FRadixSortUniformBufferRef& RadixSortUniformBuffer, FRHIShaderResourceView* RadixSortParameterBufferSRV)
+	void SetParameters(FRHIBatchedShaderParameters& BatchedParameters, FRHIUnorderedAccessView* OutOffsetsUAV, FRHIShaderResourceView* InKeysSRV, const FRadixSortUniformBufferRef& RadixSortUniformBuffer, FRHIShaderResourceView* RadixSortParameterBufferSRV)
 	{
-		FRHIComputeShader* ComputeShaderRHI = RHICmdList.GetBoundComputeShader();
-		SetUniformBufferParameter(RHICmdList, ComputeShaderRHI, GetUniformBufferParameter<FRadixSortParameters>(), RadixSortUniformBuffer );
-		SetSRVParameter(RHICmdList, ComputeShaderRHI, InKeys, InKeysSRV);
-		SetSRVParameter(RHICmdList, ComputeShaderRHI, RadixSortParameterBuffer, RadixSortParameterBufferSRV);
+		SetUAVParameter(BatchedParameters, OutOffsets, OutOffsetsUAV);
+		SetSRVParameter(BatchedParameters, InKeys, InKeysSRV);
+		SetSRVParameter(BatchedParameters, RadixSortParameterBuffer, RadixSortParameterBufferSRV);
 	}
 
-	/**
-	 * Set output buffer for this shader.
-	 */
-	void SetOutput(FRHICommandList& RHICmdList, FRHIUnorderedAccessView* OutOffsetsUAV)
+	void UnsetParameters(FRHIBatchedShaderParameters& BatchedParameters)
 	{
-		FRHIComputeShader* ComputeShaderRHI = RHICmdList.GetBoundComputeShader();
-		SetUAVParameter(RHICmdList, ComputeShaderRHI, OutOffsets, OutOffsetsUAV);
-	}
-
-	/**
-	 * Unbinds any buffers that have been bound.
-	 */
-	void UnbindBuffers(FRHICommandList& RHICmdList)
-	{
-		FRHIComputeShader* ComputeShaderRHI = RHICmdList.GetBoundComputeShader();
-		SetSRVParameter(RHICmdList, ComputeShaderRHI, RadixSortParameterBuffer, nullptr);
-		SetSRVParameter(RHICmdList, ComputeShaderRHI, InKeys, nullptr);
-		SetUAVParameter(RHICmdList, ComputeShaderRHI, OutOffsets, nullptr);
+		SetSRVParameter(BatchedParameters, RadixSortParameterBuffer, nullptr);
+		SetSRVParameter(BatchedParameters, InKeys, nullptr);
+		SetUAVParameter(BatchedParameters, OutOffsets, nullptr);
 	}
 
 private:
@@ -387,32 +350,16 @@ public:
 		OutOffsets.Bind( Initializer.ParameterMap, TEXT("OutOffsets") );
 	}
 
-	/**
-	 * Set parameters for this shader.
-	 */
-	void SetParameters(FRHICommandList& RHICmdList, FRHIShaderResourceView* InOffsetsSRV)
-	{
-		FRHIComputeShader* ComputeShaderRHI = RHICmdList.GetBoundComputeShader();
-		SetSRVParameter(RHICmdList, ComputeShaderRHI, InOffsets, InOffsetsSRV);
-	}
-
-	/**
-	 * Set output buffer for this shader.
-	 */
-	void SetOutput(FRHICommandList& RHICmdList, FRHIUnorderedAccessView* OutOffsetsUAV)
+	void SetParameters(FRHIBatchedShaderParameters& BatchedParameters, FRHIUnorderedAccessView* OutOffsetsUAV, FRHIShaderResourceView* InOffsetsSRV)
 	{
-		FRHIComputeShader* ComputeShaderRHI = RHICmdList.GetBoundComputeShader();
-		SetUAVParameter(RHICmdList, ComputeShaderRHI, OutOffsets, OutOffsetsUAV);
+		SetUAVParameter(BatchedParameters, OutOffsets, OutOffsetsUAV);
+		SetSRVParameter(BatchedParameters, InOffsets, InOffsetsSRV);
 	}
 
-	/**
-	 * Unbinds any buffers that have been bound.
-	 */
-	void UnbindBuffers(FRHICommandList& RHICmdList)
+	void UnsetParameters(FRHIBatchedShaderParameters& BatchedParameters)
 	{
-		FRHIComputeShader* ComputeShaderRHI = RHICmdList.GetBoundComputeShader();
-		SetSRVParameter(RHICmdList, ComputeShaderRHI, InOffsets, nullptr);
-		SetUAVParameter(RHICmdList, ComputeShaderRHI, OutOffsets, nullptr);
+		SetSRVParameter(BatchedParameters, InOffsets, nullptr);
+		SetUAVParameter(BatchedParameters, OutOffsets, nullptr);
 	}
 
 private:
@@ -471,47 +418,32 @@ public:
 		return RadixSortParameterBuffer.IsBound();
 	}
 
-	/**
-	 * Set parameters for this shader.
-	 */
 	void SetParameters(
-		FRHICommandList& RHICmdList,
+		FRHIBatchedShaderParameters& BatchedParameters,
+		FRHIUnorderedAccessView* OutKeysUAV,
+		FRHIUnorderedAccessView* OutValuesUAV,
 		FRHIShaderResourceView* InKeysSRV,
 		FRHIShaderResourceView* InValuesSRV,
 		FRHIShaderResourceView* InOffsetsSRV,
-		FRadixSortUniformBufferRef& RadixSortUniformBuffer,
 		FRHIShaderResourceView* RadixSortParameterBufferSRV )
 	{
-		FRHIComputeShader* ComputeShaderRHI = RHICmdList.GetBoundComputeShader();
-		SetUniformBufferParameter(RHICmdList, ComputeShaderRHI, GetUniformBufferParameter<FRadixSortParameters>(), RadixSortUniformBuffer );
-		SetSRVParameter(RHICmdList, ComputeShaderRHI, RadixSortParameterBuffer, RadixSortParameterBufferSRV);
-		SetSRVParameter(RHICmdList, ComputeShaderRHI, InKeys, InKeysSRV);
-		SetSRVParameter(RHICmdList, ComputeShaderRHI, InValues, InValuesSRV);
-		SetSRVParameter(RHICmdList, ComputeShaderRHI, InOffsets, InOffsetsSRV);
-	}
+		SetUAVParameter(BatchedParameters, OutKeys, OutKeysUAV);
+		SetUAVParameter(BatchedParameters, OutValues, OutValuesUAV);
 
-	/**
-	 * Set output buffer for this shader.
-	 */
-	void SetOutput(FRHICommandList& RHICmdList, FRHIUnorderedAccessView* OutKeysUAV, FRHIUnorderedAccessView* OutValuesUAV)
-	{
-		FRHIComputeShader* ComputeShaderRHI = RHICmdList.GetBoundComputeShader();
-		SetUAVParameter(RHICmdList, ComputeShaderRHI, OutKeys, OutKeysUAV);
-		SetUAVParameter(RHICmdList, ComputeShaderRHI, OutValues, OutValuesUAV);
+		SetSRVParameter(BatchedParameters, RadixSortParameterBuffer, RadixSortParameterBufferSRV);
+		SetSRVParameter(BatchedParameters, InKeys, InKeysSRV);
+		SetSRVParameter(BatchedParameters, InValues, InValuesSRV);
+		SetSRVParameter(BatchedParameters, InOffsets, InOffsetsSRV);
 	}
 
-	/**
-	 * Unbinds any buffers that have been bound.
-	 */
-	void UnbindBuffers(FRHICommandList& RHICmdList)
+	void UnsetParameters(FRHIBatchedShaderParameters& BatchedParameters)
 	{
-		FRHIComputeShader* ComputeShaderRHI = RHICmdList.GetBoundComputeShader();
-		SetSRVParameter(RHICmdList, ComputeShaderRHI, RadixSortParameterBuffer, nullptr);
-		SetSRVParameter(RHICmdList, ComputeShaderRHI, InKeys, nullptr);
-		SetSRVParameter(RHICmdList, ComputeShaderRHI, InValues, nullptr);
-		SetSRVParameter(RHICmdList, ComputeShaderRHI, InOffsets, nullptr);
-		SetUAVParameter(RHICmdList, ComputeShaderRHI, OutKeys, nullptr);
-		SetUAVParameter(RHICmdList, ComputeShaderRHI, OutValues, nullptr);
+		SetSRVParameter(BatchedParameters, RadixSortParameterBuffer, nullptr);
+		SetSRVParameter(BatchedParameters, InKeys, nullptr);
+		SetSRVParameter(BatchedParameters, InValues, nullptr);
+		SetSRVParameter(BatchedParameters, InOffsets, nullptr);
+		SetUAVParameter(BatchedParameters, OutKeys, nullptr);
+		SetUAVParameter(BatchedParameters, OutValues, nullptr);
 	}
 
 private:
@@ -648,9 +580,11 @@ int32 SortGPUBuffers(FRHICommandList& RHICmdList, FGPUSortBuffers SortBuffers, i
 			
 			// Clear the offsets buffer.
 			SetComputePipelineState(RHICmdList, ClearOffsetsCS.GetComputeShader());
-			ClearOffsetsCS->SetOutput(RHICmdList, GSortOffsetBuffers.BufferUAVs[0]);
+			SetAllShaderParametersCS(RHICmdList, ClearOffsetsCS, GSortOffsetBuffers.BufferUAVs[0]);
+
 			DispatchComputeShader(RHICmdList, ClearOffsetsCS.GetShader(), 1, 1 ,1 );
-			ClearOffsetsCS->UnbindBuffers(RHICmdList);
+
+			UnsetAllShaderParametersCS(RHICmdList, ClearOffsetsCS);
 
 			//make UAV safe for readback
 			RHICmdList.Transition({
@@ -660,10 +594,13 @@ int32 SortGPUBuffers(FRHICommandList& RHICmdList, FGPUSortBuffers SortBuffers, i
 
 			// Phase 1: Scan upsweep to compute per-digit totals.
 			SetComputePipelineState(RHICmdList, UpsweepCS.GetComputeShader());
-			UpsweepCS->SetOutput(RHICmdList, GSortOffsetBuffers.BufferUAVs[0]);
-			UpsweepCS->SetParameters(RHICmdList, SortBuffers.RemoteKeySRVs[BufferIndex], SortUniformBufferRef, GRadixSortParametersBuffer.SortParametersBufferSRV );
-			DispatchComputeShader(RHICmdList, UpsweepCS.GetShader(), GroupCount, 1, 1 );
-			UpsweepCS->UnbindBuffers(RHICmdList);
+
+			SetUniformBufferParameter(RHICmdList, UpsweepCS.GetComputeShader(), UpsweepCS->GetUniformBufferParameter<FRadixSortParameters>(), SortUniformBufferRef);
+			SetAllShaderParametersCS(RHICmdList, UpsweepCS, GSortOffsetBuffers.BufferUAVs[0], SortBuffers.RemoteKeySRVs[BufferIndex], SortUniformBufferRef, GRadixSortParametersBuffer.SortParametersBufferSRV);
+
+			DispatchComputeShader(RHICmdList, UpsweepCS.GetShader(), GroupCount, 1, 1);
+
+			UnsetAllShaderParametersCS(RHICmdList, UpsweepCS);
 
 			//barrier both UAVS since for next step.
 			RHICmdList.Transition({
@@ -679,10 +616,11 @@ int32 SortGPUBuffers(FRHICommandList& RHICmdList, FGPUSortBuffers SortBuffers, i
 
 			// Phase 2: Parallel prefix scan on the offsets buffer.
 			SetComputePipelineState(RHICmdList, SpineCS.GetComputeShader());
-			SpineCS->SetOutput(RHICmdList, GSortOffsetBuffers.BufferUAVs[1]);
-			SpineCS->SetParameters(RHICmdList, GSortOffsetBuffers.BufferSRVs[0] );
+			SetAllShaderParametersCS(RHICmdList, SpineCS, GSortOffsetBuffers.BufferUAVs[1], GSortOffsetBuffers.BufferSRVs[0]);
+
 			DispatchComputeShader(RHICmdList, SpineCS.GetShader(), 1, 1, 1 );
-			SpineCS->UnbindBuffers(RHICmdList);
+
+			UnsetAllShaderParametersCS(RHICmdList, SpineCS);
 
 			if (bDebugOffsets)
 			{
@@ -699,6 +637,9 @@ int32 SortGPUBuffers(FRHICommandList& RHICmdList, FGPUSortBuffers SortBuffers, i
 			const bool bIsLastPass = ((PassBits << RADIX_BITS) & KeyMask) == 0;
 			// Phase 3: Downsweep to compute final offsets and scatter keys.
 			SetComputePipelineState(RHICmdList, DownsweepCS.GetComputeShader());
+
+			SetUniformBufferParameter(RHICmdList, DownsweepCS.GetComputeShader(), DownsweepCS->GetUniformBufferParameter<FRadixSortParameters>(), SortUniformBufferRef);
+
 			{
 				FRHIUnorderedAccessView* ValuesUAV = nullptr;
 				if (bIsLastPass && SortBuffers.FinalValuesUAV)
@@ -711,14 +652,13 @@ int32 SortGPUBuffers(FRHICommandList& RHICmdList, FGPUSortBuffers SortBuffers, i
 				{
 					ValuesUAV = SortBuffers.RemoteValueUAVs[BufferIndex ^ 0x1];
 				}
-				DownsweepCS->SetOutput(RHICmdList, SortBuffers.RemoteKeyUAVs[BufferIndex ^ 0x1], ValuesUAV);
-			}
-			{
+
 				FRHIShaderResourceView* ValuesSRV = (PassIndex == 0 && SortBuffers.FirstValuesSRV) ? SortBuffers.FirstValuesSRV : SortBuffers.RemoteValueSRVs[BufferIndex];
-				DownsweepCS->SetParameters(RHICmdList, SortBuffers.RemoteKeySRVs[BufferIndex], ValuesSRV, GSortOffsetBuffers.BufferSRVs[1], SortUniformBufferRef, GRadixSortParametersBuffer.SortParametersBufferSRV );
+
+				SetAllShaderParametersCS(RHICmdList, DownsweepCS, SortBuffers.RemoteKeyUAVs[BufferIndex ^ 0x1], ValuesUAV, SortBuffers.RemoteKeySRVs[BufferIndex], ValuesSRV, GSortOffsetBuffers.BufferSRVs[1], GRadixSortParametersBuffer.SortParametersBufferSRV);
 			}
 			DispatchComputeShader(RHICmdList, DownsweepCS.GetShader(), GroupCount, 1, 1 );
-			DownsweepCS->UnbindBuffers(RHICmdList);
+			UnsetAllShaderParametersCS(RHICmdList, DownsweepCS);
 
 
 			RHICmdList.Transition({
diff --git a/Engine/Source/Runtime/Engine/Private/GPUSortManager.cpp b/Engine/Source/Runtime/Engine/Private/GPUSortManager.cpp
--- a/Engine/Source/Runtime/Engine/Private/GPUSortManager.cpp
+++ b/Engine/Source/Runtime/Engine/Private/GPUSortManager.cpp
@@ -107,13 +107,15 @@ public:
 	 * Set parameters.
 	 */
 	void SetParameters(
-		FRHICommandList& RHICmdList,	
+		FRHIBatchedShaderParameters& BatchedParameters,
 		FRHIShaderResourceView* InSourceData,
 		FRHIUnorderedAccessView* const* InDestDatas,
 		const int32* InUsedIndexCounts, 
 		int32 StartingIndex,
 		int32 DestCount);
 
+	void UnsetParameters(FRHIBatchedShaderParameters& BatchedParameters);
+
 	void Begin(FRHICommandList& RHICmdList);
 	void End(FRHICommandList& RHICmdList);
 
@@ -142,32 +144,40 @@ FCopyUIntBufferCS::FCopyUIntBufferCS(const ShaderMetaType::CompiledShaderInitial
 }
 
 void FCopyUIntBufferCS::SetParameters(
-	FRHICommandList& RHICmdList,
+	FRHIBatchedShaderParameters& BatchedParameters,
 	FRHIShaderResourceView* InSourceData,
 	FRHIUnorderedAccessView* const* InDestDatas,
 	const int32* InUsedIndexCounts,
 	int32 StartingIndex,
 	int32 DestCount)
 {
-	FRHIComputeShader* ComputeShaderRHI = RHICmdList.GetBoundComputeShader();
 	check(DestCount > 0 && DestCount <= COPYUINTCS_BUFFER_COUNT);
 
-	SetSRVParameter(RHICmdList, ComputeShaderRHI, SourceData, InSourceData);
+	SetSRVParameter(BatchedParameters, SourceData, InSourceData);
 
 	FUintVector4 CopyParamsValue(StartingIndex, 0, 0, 0);
 	for (int32 Index = 0; Index < DestCount; ++Index)
 	{
-		SetUAVParameter(RHICmdList, ComputeShaderRHI, DestData[Index], InDestDatas[Index]);
+		SetUAVParameter(BatchedParameters, DestData[Index], InDestDatas[Index]);
 		CopyParamsValue[Index + 1] = InUsedIndexCounts[Index];
 	}
 
 	for (int32 Index = DestCount; Index < COPYUINTCS_BUFFER_COUNT; ++Index)
 	{
 		// TR-DummyUAVs : those buffers are only ever used here, but there content is never accessed.
-		SetUAVParameter(RHICmdList, ComputeShaderRHI, DestData[Index], NiagaraSortingDummyUAV[Index].Buffer.UAV);
+		SetUAVParameter(BatchedParameters, DestData[Index], NiagaraSortingDummyUAV[Index].Buffer.UAV);
 	}
 
-	SetShaderValue(RHICmdList, ComputeShaderRHI, CopyParams, CopyParamsValue);
+	SetShaderValue(BatchedParameters, CopyParams, CopyParamsValue);
+}
+
+void FCopyUIntBufferCS::UnsetParameters(FRHIBatchedShaderParameters& BatchedParameters)
+{
+	SetSRVParameter(BatchedParameters, SourceData, nullptr);
+	for (int32 Index = 0; Index < COPYUINTCS_BUFFER_COUNT; ++Index)
+	{
+		SetUAVParameter(BatchedParameters, DestData[Index], nullptr);
+	}
 }
 
 void FCopyUIntBufferCS::Begin(FRHICommandList& RHICmdList)
@@ -182,14 +192,10 @@ void FCopyUIntBufferCS::Begin(FRHICommandList& RHICmdList)
 
 void FCopyUIntBufferCS::End(FRHICommandList& RHICmdList)
 {
-	FRHIComputeShader* ComputeShaderRHI = RHICmdList.GetBoundComputeShader();
-	SetSRVParameter(RHICmdList, ComputeShaderRHI, SourceData, nullptr);
-
 	FRHIUnorderedAccessView* Views[COPYUINTCS_BUFFER_COUNT];
 	for (int32 Index = 0; Index < COPYUINTCS_BUFFER_COUNT; ++Index)
 	{
 		Views[Index] = NiagaraSortingDummyUAV[Index].Buffer.UAV;
-		SetUAVParameter(RHICmdList, ComputeShaderRHI, DestData[Index], nullptr);
 	}
 	RHICmdList.EndUAVOverlap(Views);
 }
@@ -211,13 +217,16 @@ void CopyUIntBufferToTargets(FRHICommandListImmediate& RHICmdList, ERHIFeatureLe
 		const int32 NumTargetsInPass = FMath::Min<int32>(NumTargets - Index0InPass, COPYUINTCS_BUFFER_COUNT);
 		const int32 NumElementsInPass = TargetSizes[Index0InPass + NumTargetsInPass - 1] - StartingOffset;
 
-		CopyBufferCS->SetParameters(RHICmdList, SourceSRV, TargetUAVs + Index0InPass, TargetSizes + Index0InPass, StartingOffset, NumTargetsInPass);
+		SetAllShaderParametersCS(RHICmdList, CopyBufferCS, SourceSRV, TargetUAVs + Index0InPass, TargetSizes + Index0InPass, StartingOffset, NumTargetsInPass);
+
 		DispatchComputeShader(RHICmdList, CopyBufferCS, FMath::DivideAndRoundUp(NumElementsInPass, COPYUINTCS_THREAD_COUNT), 1, 1);
 
 		StartingOffset += NumElementsInPass;
 		Index0InPass += COPYUINTCS_BUFFER_COUNT;
 	};
 
+	UnsetAllShaderParametersCS(RHICmdList, CopyBufferCS);
+
 	CopyBufferCS->End(RHICmdList);
 }
 
diff --git a/Engine/Source/Runtime/Engine/Private/Particles/ParticleSortingGPU.cpp b/Engine/Source/Runtime/Engine/Private/Particles/ParticleSortingGPU.cpp
--- a/Engine/Source/Runtime/Engine/Private/Particles/ParticleSortingGPU.cpp
+++ b/Engine/Source/Runtime/Engine/Private/Particles/ParticleSortingGPU.cpp
@@ -76,49 +76,25 @@ public:
 		OutParticleIndices.Bind( Initializer.ParameterMap, TEXT("OutParticleIndices") );
 	}
 
-	/**
-	 * Set output buffers for this shader.
-	 */
-	void SetOutput(FRHICommandList& RHICmdList, FRHIUnorderedAccessView* OutKeysUAV, FRHIUnorderedAccessView* OutIndicesUAV )
-	{
-		FRHIComputeShader* ComputeShaderRHI = RHICmdList.GetBoundComputeShader();
-		SetUAVParameter(RHICmdList, ComputeShaderRHI, OutKeys, OutKeysUAV);
-		SetUAVParameter(RHICmdList, ComputeShaderRHI, OutParticleIndices, OutIndicesUAV);
-	}
-
-	/**
-	 * Set input parameters.
-	 */
 	void SetParameters(
-		FRHICommandList& RHICmdList,
-		FParticleKeyGenUniformBufferRef& UniformBuffer,
+		FRHIBatchedShaderParameters& BatchedParameters,
+		FRHIUnorderedAccessView* OutKeysUAV,
+		FRHIUnorderedAccessView* OutIndicesUAV,
+		FRHITexture2D* PositionTextureRHI,
 		FRHIShaderResourceView* InIndicesSRV
 		)
 	{
-		FRHIComputeShader* ComputeShaderRHI = RHICmdList.GetBoundComputeShader();
-		SetUniformBufferParameter(RHICmdList, ComputeShaderRHI, GetUniformBufferParameter<FParticleKeyGenParameters>(), UniformBuffer );
-		SetSRVParameter(RHICmdList, ComputeShaderRHI, InParticleIndices, InIndicesSRV);
+		SetUAVParameter(BatchedParameters, OutKeys, OutKeysUAV);
+		SetUAVParameter(BatchedParameters, OutParticleIndices, OutIndicesUAV);
+		SetTextureParameter(BatchedParameters, PositionTexture, PositionTextureRHI);
+		SetSRVParameter(BatchedParameters, InParticleIndices, InIndicesSRV);
 	}
 
-	/**
-	 * Set the texture from which particle positions can be read.
-	 */
-	void SetPositionTextures(FRHICommandList& RHICmdList, FRHITexture2D* PositionTextureRHI)
+	void UnsetParameters(FRHIBatchedShaderParameters& BatchedParameters)
 	{
-		FRHIComputeShader* ComputeShaderRHI = RHICmdList.GetBoundComputeShader();
-		SetTextureParameter(RHICmdList, ComputeShaderRHI, PositionTexture, PositionTextureRHI);
-	}
-
-	/**
-	 * Unbinds any buffers that have been bound.
-	 */
-	void UnbindBuffers(FRHICommandList& RHICmdList)
-	{
-		FRHIComputeShader* ComputeShaderRHI = RHICmdList.GetBoundComputeShader();
-
-		SetSRVParameter(RHICmdList, ComputeShaderRHI, InParticleIndices, nullptr);
-		SetUAVParameter(RHICmdList, ComputeShaderRHI, OutKeys, nullptr);
-		SetUAVParameter(RHICmdList, ComputeShaderRHI, OutParticleIndices, nullptr);
+		SetSRVParameter(BatchedParameters, InParticleIndices, nullptr);
+		SetUAVParameter(BatchedParameters, OutKeys, nullptr);
+		SetUAVParameter(BatchedParameters, OutParticleIndices, nullptr);
 	}
 
 private:
@@ -155,15 +131,14 @@ int32 GenerateParticleSortKeys(
 	check(FeatureLevel >= ERHIFeatureLevel::SM5);
 
 	FParticleKeyGenParameters KeyGenParameters;
-	FParticleKeyGenUniformBufferRef KeyGenUniformBuffer;
 	const uint32 MaxGroupCount = 128;
 	int32 TotalParticleCount = 0;
 
 	// Grab the shader, set output.
 	TShaderMapRef<FParticleSortKeyGenCS> KeyGenCS(GetGlobalShaderMap(FeatureLevel));
 	SetComputePipelineState(RHICmdList, KeyGenCS.GetComputeShader());
-	KeyGenCS->SetOutput(RHICmdList, KeyBufferUAV, SortedVertexBufferUAV);
-	KeyGenCS->SetPositionTextures(RHICmdList, PositionTextureRHI);
+
+	FRHIBatchedShaderParameters BatchedParameters;
 
 	// TR-KeyGen : No sync needed between tasks since they update different parts of the data (assuming it's ok if cache line overlap).
 	RHICmdList.BeginUAVOverlap({ KeyBufferUAV, SortedVertexBufferUAV });
@@ -184,16 +159,22 @@ int32 GenerateParticleSortKeys(
 			KeyGenParameters.OutputOffset = SortInfo.AllocationInfo.BufferOffset;
 			KeyGenParameters.EmitterKey = (uint32)SortInfo.AllocationInfo.ElementIndex << 16;
 			KeyGenParameters.KeyCount = ParticleCount;
-			KeyGenUniformBuffer = FParticleKeyGenUniformBufferRef::CreateUniformBufferImmediate( KeyGenParameters, UniformBuffer_SingleDraw );
 
-			// Dispatch.
-			KeyGenCS->SetParameters(RHICmdList, KeyGenUniformBuffer, SortInfo.VertexBufferSRV);
+			FParticleKeyGenUniformBufferRef KeyGenUniformBuffer = FParticleKeyGenUniformBufferRef::CreateUniformBufferImmediate( KeyGenParameters, UniformBuffer_SingleDraw );
+			SetUniformBufferParameter(RHICmdList, KeyGenCS.GetComputeShader(), KeyGenCS->GetUniformBufferParameter<FParticleKeyGenParameters>(), KeyGenUniformBuffer);
+
+			KeyGenCS->SetParameters(BatchedParameters, KeyBufferUAV, SortedVertexBufferUAV, PositionTextureRHI, SortInfo.VertexBufferSRV);
+
+			RHICmdList.SetBatchedShaderParameters(KeyGenCS.GetComputeShader(), BatchedParameters);
+			BatchedParameters.Reset();
+
 			DispatchComputeShader(RHICmdList, KeyGenCS.GetShader(), GroupCount, 1, 1);
 		}
 	}
 
 	// Clear the output buffer.
-	KeyGenCS->UnbindBuffers(RHICmdList);
+	KeyGenCS->UnsetParameters(BatchedParameters);
+	RHICmdList.SetBatchedShaderParameters(KeyGenCS.GetComputeShader(), BatchedParameters);
 
 	RHICmdList.EndUAVOverlap({ KeyBufferUAV, SortedVertexBufferUAV });
 
diff --git a/Engine/Source/Runtime/Engine/Private/PhysicsField/PhysicsFieldComponent.cpp b/Engine/Source/Runtime/Engine/Private/PhysicsField/PhysicsFieldComponent.cpp
--- a/Engine/Source/Runtime/Engine/Private/PhysicsField/PhysicsFieldComponent.cpp
+++ b/Engine/Source/Runtime/Engine/Private/PhysicsField/PhysicsFieldComponent.cpp
@@ -227,44 +227,37 @@ public:
 	{
 	}
 
-	void SetParameters(FRHICommandList& RHICmdList, FPhysicsFieldResource* FieldResource, const float InTimeSeconds, const int32 InNumCells)
+	void SetParameters(FRHIBatchedShaderParameters& BatchedParameters, FPhysicsFieldResource* FieldResource, const float InTimeSeconds, const int32 InNumCells)
 	{
-		FRHIComputeShader* ShaderRHI = RHICmdList.GetBoundComputeShader();
-
 		if (FieldResource)
 		{
-			RHICmdList.Transition(FRHITransitionInfo(FieldResource->ClipmapBuffer.UAV, ERHIAccess::Unknown, ERHIAccess::UAVCompute));
-
-			SetSRVParameter(RHICmdList, ShaderRHI, NodesParams, FieldResource->NodesParams.SRV);
-			SetSRVParameter(RHICmdList, ShaderRHI, NodesOffsets, FieldResource->NodesOffsets.SRV);
-			SetSRVParameter(RHICmdList, ShaderRHI, TargetsOffsets, FieldResource->TargetsOffsets.SRV);
-			SetUAVParameter(RHICmdList, ShaderRHI, FieldClipmap, FieldResource->ClipmapBuffer.UAV);
-
-			SetSRVParameter(RHICmdList, ShaderRHI, CellsOffsets, FieldResource->CellsOffsets.SRV);
-			SetSRVParameter(RHICmdList, ShaderRHI, CellsMin, FieldResource->CellsMin.SRV);
-			SetSRVParameter(RHICmdList, ShaderRHI, CellsMax, FieldResource->CellsMax.SRV);
-
-			SetShaderValue(RHICmdList, ShaderRHI, ClipmapResolution, FieldResource->FieldInfos.ClipmapResolution);
-			SetShaderValue(RHICmdList, ShaderRHI, ClipmapDistance, FieldResource->FieldInfos.ClipmapDistance);
-			SetShaderValue(RHICmdList, ShaderRHI, ClipmapCount, FieldResource->FieldInfos.ClipmapCount);
-			SetShaderValue(RHICmdList, ShaderRHI, ClipmapCenter, (FVector3f)FieldResource->FieldInfos.ClipmapCenter);
-			SetShaderValue(RHICmdList, ShaderRHI, ClipmapExponent, FieldResource->FieldInfos.ClipmapExponent);
-
-			SetShaderValue(RHICmdList, ShaderRHI, ValidTargets, FieldResource->FieldInfos.ValidTargets);
-			SetShaderValue(RHICmdList, ShaderRHI, ValidCount, FieldResource->FieldInfos.ValidCount);
-			SetShaderValue(RHICmdList, ShaderRHI, TargetsIndex, FieldResource->FieldInfos.PhysicsTargets);
-			SetShaderValue(RHICmdList, ShaderRHI, TimeSeconds, InTimeSeconds);
-			SetShaderValue(RHICmdList, ShaderRHI, NumCells, InNumCells);
-			SetShaderValue(RHICmdList, ShaderRHI, TargetCount, FieldResource->FieldInfos.TargetCount);
+			SetSRVParameter(BatchedParameters, NodesParams, FieldResource->NodesParams.SRV);
+			SetSRVParameter(BatchedParameters, NodesOffsets, FieldResource->NodesOffsets.SRV);
+			SetSRVParameter(BatchedParameters, TargetsOffsets, FieldResource->TargetsOffsets.SRV);
+			SetUAVParameter(BatchedParameters, FieldClipmap, FieldResource->ClipmapBuffer.UAV);
+
+			SetSRVParameter(BatchedParameters, CellsOffsets, FieldResource->CellsOffsets.SRV);
+			SetSRVParameter(BatchedParameters, CellsMin, FieldResource->CellsMin.SRV);
+			SetSRVParameter(BatchedParameters, CellsMax, FieldResource->CellsMax.SRV);
+
+			SetShaderValue(BatchedParameters, ClipmapResolution, FieldResource->FieldInfos.ClipmapResolution);
+			SetShaderValue(BatchedParameters, ClipmapDistance, FieldResource->FieldInfos.ClipmapDistance);
+			SetShaderValue(BatchedParameters, ClipmapCount, FieldResource->FieldInfos.ClipmapCount);
+			SetShaderValue(BatchedParameters, ClipmapCenter, (FVector3f)FieldResource->FieldInfos.ClipmapCenter);
+			SetShaderValue(BatchedParameters, ClipmapExponent, FieldResource->FieldInfos.ClipmapExponent);
+
+			SetShaderValue(BatchedParameters, ValidTargets, FieldResource->FieldInfos.ValidTargets);
+			SetShaderValue(BatchedParameters, ValidCount, FieldResource->FieldInfos.ValidCount);
+			SetShaderValue(BatchedParameters, TargetsIndex, FieldResource->FieldInfos.PhysicsTargets);
+			SetShaderValue(BatchedParameters, TimeSeconds, InTimeSeconds);
+			SetShaderValue(BatchedParameters, NumCells, InNumCells);
+			SetShaderValue(BatchedParameters, TargetCount, FieldResource->FieldInfos.TargetCount);
 		}
 	}
 
-	void UnsetParameters(FRHICommandList& RHICmdList, FPhysicsFieldResource* FieldResource)
+	void UnsetParameters(FRHIBatchedShaderParameters& BatchedParameters)
 	{
-		FRHIComputeShader* ShaderRHI = RHICmdList.GetBoundComputeShader();
-
-		SetUAVParameter(RHICmdList, ShaderRHI, FieldClipmap, nullptr);
-		RHICmdList.Transition(FRHITransitionInfo(FieldResource->ClipmapBuffer.UAV, ERHIAccess::Unknown, ERHIAccess::SRVCompute));
+		SetUAVParameter(BatchedParameters, FieldClipmap, nullptr);
 	}
 
 private:
@@ -341,37 +334,30 @@ public:
 	{
 	}
 
-	void SetParameters(FRHICommandList& RHICmdList, FPhysicsFieldResource* FieldResource, const float InTimeSeconds, const int32 InNumCells)
+	void SetParameters(FRHIBatchedShaderParameters& BatchedParameters, FPhysicsFieldResource* FieldResource, const float InTimeSeconds, const int32 InNumCells)
 	{
-		FRHIComputeShader* ShaderRHI = RHICmdList.GetBoundComputeShader();
-
 		if (FieldResource)
 		{
-			RHICmdList.Transition(FRHITransitionInfo(FieldResource->ClipmapBuffer.UAV, ERHIAccess::Unknown, ERHIAccess::UAVCompute));
+			SetUAVParameter(BatchedParameters, FieldClipmap, FieldResource->ClipmapBuffer.UAV);
+			SetShaderValue(BatchedParameters, NumCells, InNumCells);
 
-			SetUAVParameter(RHICmdList, ShaderRHI, FieldClipmap, FieldResource->ClipmapBuffer.UAV);
-			SetShaderValue(RHICmdList, ShaderRHI, NumCells, InNumCells);
+			SetSRVParameter(BatchedParameters, CellsOffsets, FieldResource->CellsOffsets.SRV);
+			SetSRVParameter(BatchedParameters, CellsMin, FieldResource->CellsMin.SRV);
+			SetSRVParameter(BatchedParameters, CellsMax, FieldResource->CellsMax.SRV);
 
-			SetSRVParameter(RHICmdList, ShaderRHI, CellsOffsets, FieldResource->CellsOffsets.SRV);
-			SetSRVParameter(RHICmdList, ShaderRHI, CellsMin, FieldResource->CellsMin.SRV);
-			SetSRVParameter(RHICmdList, ShaderRHI, CellsMax, FieldResource->CellsMax.SRV);
+			SetShaderValue(BatchedParameters, ClipmapResolution, FieldResource->FieldInfos.ClipmapResolution);
+			SetShaderValue(BatchedParameters, ClipmapCount, FieldResource->FieldInfos.ClipmapCount);
 
-			SetShaderValue(RHICmdList, ShaderRHI, ClipmapResolution, FieldResource->FieldInfos.ClipmapResolution);
-			SetShaderValue(RHICmdList, ShaderRHI, ClipmapCount, FieldResource->FieldInfos.ClipmapCount);
-
-			SetShaderValue(RHICmdList, ShaderRHI, ValidTargets, FieldResource->FieldInfos.ValidTargets);
-			SetShaderValue(RHICmdList, ShaderRHI, ValidCount, FieldResource->FieldInfos.ValidCount);
-			SetShaderValue(RHICmdList, ShaderRHI, TargetsIndex, FieldResource->FieldInfos.PhysicsTargets);
-			SetShaderValue(RHICmdList, ShaderRHI, TargetCount, FieldResource->FieldInfos.TargetCount);
+			SetShaderValue(BatchedParameters, ValidTargets, FieldResource->FieldInfos.ValidTargets);
+			SetShaderValue(BatchedParameters, ValidCount, FieldResource->FieldInfos.ValidCount);
+			SetShaderValue(BatchedParameters, TargetsIndex, FieldResource->FieldInfos.PhysicsTargets);
+			SetShaderValue(BatchedParameters, TargetCount, FieldResource->FieldInfos.TargetCount);
 		}
 	}
 
-	void UnsetParameters(FRHICommandList& RHICmdList, FPhysicsFieldResource* FieldResource)
+	void UnsetParameters(FRHIBatchedShaderParameters& BatchedParameters)
 	{
-		FRHIComputeShader* ShaderRHI = RHICmdList.GetBoundComputeShader();
-
-		SetUAVParameter(RHICmdList, ShaderRHI, FieldClipmap, nullptr);
-		RHICmdList.Transition(FRHITransitionInfo(FieldResource->ClipmapBuffer.UAV, ERHIAccess::Unknown, ERHIAccess::SRVCompute));
+		SetUAVParameter(BatchedParameters, FieldClipmap, nullptr);
 	}
 
 private:
@@ -573,15 +559,19 @@ void FPhysicsFieldResource::UpdateResource(FRHICommandListImmediate& RHICmdList,
 			SCOPED_DRAW_EVENT(RHICmdList, PhysicsFields_ClearClipmap);
 			SCOPED_GPU_STAT(RHICmdList, PhysicsFields_ClearClipmap);
 
+			RHICmdList.Transition(FRHITransitionInfo(this->ClipmapBuffer.UAV, ERHIAccess::Unknown, ERHIAccess::UAVCompute));
+
 			TShaderMapRef<FResetPhysicsFieldClipmapCS> ComputeShader(GetGlobalShaderMap(GetFeatureLevel()));
 			SetComputePipelineState(RHICmdList, ComputeShader.GetComputeShader());
 
 			const uint32 NumCells = FieldInfos.CellsOffsets[CellsCount];
 			const uint32 NumGroups = FMath::DivideAndRoundUp<int32>(NumCells, FResetPhysicsFieldClipmapCS::ThreadGroupSize);
 
-			ComputeShader->SetParameters(RHICmdList, this, TimeSeconds, NumCells);
+			SetAllShaderParametersCS(RHICmdList, ComputeShader, this, TimeSeconds, NumCells);
 			DispatchComputeShader(RHICmdList, ComputeShader.GetShader(), NumGroups, 1, 1);
-			ComputeShader->UnsetParameters(RHICmdList, this);
+			UnsetAllShaderParametersCS(RHICmdList, ComputeShader);
+
+			RHICmdList.Transition(FRHITransitionInfo(this->ClipmapBuffer.UAV, ERHIAccess::Unknown, ERHIAccess::SRVCompute));
 		}
 		{
 			SCOPED_DRAW_EVENT(RHICmdList, PhysicsFields_UpdateBuffers);
@@ -597,15 +587,20 @@ void FPhysicsFieldResource::UpdateResource(FRHICommandListImmediate& RHICmdList,
 		{
 			SCOPED_DRAW_EVENT(RHICmdList, PhysicsFields_BuildClipmap);
 			SCOPED_GPU_STAT(RHICmdList, PhysicsFields_BuildClipmap);
+
+			RHICmdList.Transition(FRHITransitionInfo(this->ClipmapBuffer.UAV, ERHIAccess::Unknown, ERHIAccess::UAVCompute));
+
 			TShaderMapRef<FBuildPhysicsFieldClipmapCS> ComputeShader(GetGlobalShaderMap(GetFeatureLevel()));
 			SetComputePipelineState(RHICmdList, ComputeShader.GetComputeShader());
 
 			const uint32 NumCells = FieldInfos.CellsOffsets[CellsCount];
 			const uint32 NumGroups = FMath::DivideAndRoundUp<int32>(NumCells, FBuildPhysicsFieldClipmapCS::ThreadGroupSize);
 
-			ComputeShader->SetParameters(RHICmdList, this, TimeSeconds, NumCells);
+			SetAllShaderParametersCS(RHICmdList, ComputeShader, this, TimeSeconds, NumCells);
 			DispatchComputeShader(RHICmdList, ComputeShader.GetShader(), NumGroups, 1, 1);
-			ComputeShader->UnsetParameters(RHICmdList, this);
+			UnsetAllShaderParametersCS(RHICmdList, ComputeShader);
+
+			RHICmdList.Transition(FRHITransitionInfo(this->ClipmapBuffer.UAV, ERHIAccess::Unknown, ERHIAccess::SRVCompute));
 		}
 	}
 }
diff --git a/Engine/Source/Runtime/Engine/Private/SimpleElementShaders.cpp b/Engine/Source/Runtime/Engine/Private/SimpleElementShaders.cpp
--- a/Engine/Source/Runtime/Engine/Private/SimpleElementShaders.cpp
+++ b/Engine/Source/Runtime/Engine/Private/SimpleElementShaders.cpp
@@ -11,6 +11,7 @@
 #include "SceneView.h"
 #include "RHIContext.h"
 #include "SceneRelativeViewMatrices.h"
+#include "DataDrivenShaderPlatformInfo.h"
 
 /*------------------------------------------------------------------------------
 	Simple element vertex shader.
@@ -23,16 +24,30 @@ FSimpleElementVS::FSimpleElementVS(const ShaderMetaType::CompiledShaderInitializ
 	TransformTilePosition.Bind(Initializer.ParameterMap, TEXT("TransformTilePosition"), SPF_Optional); // TransformTilePosition may be optimized out if LWC is disabled
 }
 
+void FSimpleElementVS::SetParameters(FRHIBatchedShaderParameters& BatchedParameters, const FMatrix& WorldToClipMatrix)
+{
+	SetShaderValue(BatchedParameters, RelativeTransform, FMatrix44f(WorldToClipMatrix));
+	SetShaderValue(BatchedParameters, TransformTilePosition, FVector3f::ZeroVector);
+}
+
+void FSimpleElementVS::SetParameters(FRHIBatchedShaderParameters& BatchedParameters, const FRelativeViewMatrices& Matrices)
+{
+	SetShaderValue(BatchedParameters, RelativeTransform, Matrices.RelativeWorldToClip);
+	SetShaderValue(BatchedParameters, TransformTilePosition, Matrices.TilePosition);
+}
+
 void FSimpleElementVS::SetParameters(FRHICommandList& RHICmdList, const FMatrix& WorldToClipMatrix)
 {
-	SetShaderValue(RHICmdList, RHICmdList.GetBoundVertexShader(), RelativeTransform, FMatrix44f(WorldToClipMatrix));
-	SetShaderValue(RHICmdList, RHICmdList.GetBoundVertexShader(), TransformTilePosition, FVector3f::ZeroVector);
+	FRHIBatchedShaderParameters BatchedParameters;
+	SetParameters(BatchedParameters, WorldToClipMatrix);
+	RHICmdList.SetBatchedShaderParameters(RHICmdList.GetBoundVertexShader(), BatchedParameters);
 }
 
 void FSimpleElementVS::SetParameters(FRHICommandList& RHICmdList, const FRelativeViewMatrices& Matrices)
 {
-	SetShaderValue(RHICmdList, RHICmdList.GetBoundVertexShader(), RelativeTransform, Matrices.RelativeWorldToClip);
-	SetShaderValue(RHICmdList, RHICmdList.GetBoundVertexShader(), TransformTilePosition, Matrices.TilePosition);
+	FRHIBatchedShaderParameters BatchedParameters;
+	SetParameters(BatchedParameters, Matrices);
+	RHICmdList.SetBatchedShaderParameters(RHICmdList.GetBoundVertexShader(), BatchedParameters);
 }
 
 void FSimpleElementVS::ModifyCompilationEnvironment(const FGlobalShaderPermutationParameters& Parameters, FShaderCompilerEnvironment& OutEnvironment)
@@ -54,6 +69,14 @@ FSimpleElementPS::FSimpleElementPS(const ShaderMetaType::CompiledShaderInitializ
 	TextureComponentReplicateAlpha.Bind(Initializer.ParameterMap,TEXT("TextureComponentReplicateAlpha"));
 }
 
+void FSimpleElementPS::SetParameters(FRHIBatchedShaderParameters& BatchedParameters, const FTexture* TextureValue)
+{
+	SetTextureParameter(BatchedParameters, InTexture, InTextureSampler, TextureValue);
+
+	SetShaderValue(BatchedParameters, TextureComponentReplicate, TextureValue->bGreyScaleFormat ? FLinearColor(1, 0, 0, 0) : FLinearColor(0, 0, 0, 0));
+	SetShaderValue(BatchedParameters, TextureComponentReplicateAlpha, TextureValue->bGreyScaleFormat ? FLinearColor(1, 0, 0, 0) : FLinearColor(0, 0, 0, 1));
+}
+
 void FSimpleElementPS::SetEditorCompositingParameters(FRHICommandList& RHICmdList, const FSceneView* View)
 {
 	if( View )
@@ -69,23 +92,11 @@ void FSimpleElementPS::SetEditorCompositingParameters(FRHICommandList& RHICmdLis
 
 void FSimpleElementPS::SetParameters(FRHICommandList& RHICmdList, const FTexture* TextureValue)
 {
-	SetTextureParameter(RHICmdList, RHICmdList.GetBoundPixelShader(),InTexture,InTextureSampler,TextureValue);
-	
-	SetShaderValue(RHICmdList, RHICmdList.GetBoundPixelShader(),TextureComponentReplicate,TextureValue->bGreyScaleFormat ? FLinearColor(1,0,0,0) : FLinearColor(0,0,0,0));
-	SetShaderValue(RHICmdList, RHICmdList.GetBoundPixelShader(),TextureComponentReplicateAlpha,TextureValue->bGreyScaleFormat ? FLinearColor(1,0,0,0) : FLinearColor(0,0,0,1));
+	FRHIBatchedShaderParameters BatchedParameters;
+	SetParameters(BatchedParameters, TextureValue);
+	RHICmdList.SetBatchedShaderParameters(RHICmdList.GetBoundPixelShader(), BatchedParameters);
 }
 
-/*bool FSimpleElementPS::Serialize(FArchive& Ar)
-{
-	bool bShaderHasOutdatedParameters = FGlobalShader::Serialize(Ar);
-	Ar << InTexture;
-	Ar << InTextureSampler;
-	Ar << TextureComponentReplicate;
-	Ar << TextureComponentReplicateAlpha;
-	Ar << EditorCompositeDepthTestParameter;
-	Ar << ScreenToPixel;
-	return bShaderHasOutdatedParameters;
-}*/
 
 FSimpleElementAlphaOnlyPS::FSimpleElementAlphaOnlyPS(const ShaderMetaType::CompiledShaderInitializerType& Initializer) :
 	FSimpleElementPS(Initializer)
@@ -98,18 +109,18 @@ FSimpleElementGammaBasePS::FSimpleElementGammaBasePS(const ShaderMetaType::Compi
 	Gamma.Bind(Initializer.ParameterMap,TEXT("Gamma"));
 }
 
-void FSimpleElementGammaBasePS::SetParameters(FRHICommandList& RHICmdList, const FTexture* Texture, float GammaValue, ESimpleElementBlendMode BlendMode)
+void FSimpleElementGammaBasePS::SetParameters(FRHIBatchedShaderParameters& BatchedParameters, const FTexture* Texture, float GammaValue, ESimpleElementBlendMode BlendMode)
 {
-	FSimpleElementPS::SetParameters(RHICmdList, Texture);
-	SetShaderValue(RHICmdList, RHICmdList.GetBoundPixelShader(),Gamma,GammaValue);
+	FSimpleElementPS::SetParameters(BatchedParameters, Texture);
+	SetShaderValue(BatchedParameters, Gamma, GammaValue);
 }
 
-/*bool FSimpleElementGammaBasePS::Serialize(FArchive& Ar)
+void FSimpleElementGammaBasePS::SetParameters(FRHICommandList& RHICmdList, const FTexture* Texture, float GammaValue, ESimpleElementBlendMode BlendMode)
 {
-	bool bShaderHasOutdatedParameters = FSimpleElementPS::Serialize(Ar);
-	Ar << Gamma;
-	return bShaderHasOutdatedParameters;
-}*/
+	FRHIBatchedShaderParameters BatchedParameters;
+	SetParameters(BatchedParameters, Texture, GammaValue, BlendMode);
+	RHICmdList.SetBatchedShaderParameters(RHICmdList.GetBoundPixelShader(), BatchedParameters);
+}
 
 FSimpleElementMaskedGammaBasePS::FSimpleElementMaskedGammaBasePS(const ShaderMetaType::CompiledShaderInitializerType& Initializer) :
 	FSimpleElementGammaBasePS(Initializer)
@@ -117,18 +128,18 @@ FSimpleElementMaskedGammaBasePS::FSimpleElementMaskedGammaBasePS(const ShaderMet
 	ClipRef.Bind(Initializer.ParameterMap,TEXT("ClipRef"), SPF_Mandatory);
 }
 
-void FSimpleElementMaskedGammaBasePS::SetParameters(FRHICommandList& RHICmdList, const FTexture* Texture, float InGamma, float ClipRefValue, ESimpleElementBlendMode BlendMode)
+void FSimpleElementMaskedGammaBasePS::SetParameters(FRHIBatchedShaderParameters& BatchedParameters, const FTexture* Texture, float InGamma, float ClipRefValue, ESimpleElementBlendMode BlendMode)
 {
-	FSimpleElementGammaBasePS::SetParameters(RHICmdList, Texture,InGamma,BlendMode);
-	SetShaderValue(RHICmdList, RHICmdList.GetBoundPixelShader(),ClipRef,ClipRefValue);
+	FSimpleElementGammaBasePS::SetParameters(BatchedParameters, Texture, InGamma, BlendMode);
+	SetShaderValue(BatchedParameters, ClipRef, ClipRefValue);
 }
 
-/*bool FSimpleElementMaskedGammaBasePS::Serialize(FArchive& Ar)
+void FSimpleElementMaskedGammaBasePS::SetParameters(FRHICommandList& RHICmdList, const FTexture* Texture, float InGamma, float ClipRefValue, ESimpleElementBlendMode BlendMode)
 {
-	bool bShaderHasOutdatedParameters = FSimpleElementGammaBasePS::Serialize(Ar);
-	Ar << ClipRef;
-	return bShaderHasOutdatedParameters;
-}*/
+	FRHIBatchedShaderParameters BatchedParameters;
+	SetParameters(BatchedParameters, Texture, InGamma, ClipRefValue, BlendMode);
+	RHICmdList.SetBatchedShaderParameters(RHICmdList.GetBoundPixelShader(), BatchedParameters);
+}
 
 /**
 * Constructor
@@ -149,21 +160,8 @@ FSimpleElementDistanceFieldGammaPS::FSimpleElementDistanceFieldGammaPS(const Sha
 	GlowInnerRadius.Bind(Initializer.ParameterMap,TEXT("GlowInnerRadius"));
 }
 
-/**
-* Sets all the constant parameters for this shader
-*
-* @param Texture - 2d tile texture
-* @param Gamma - if gamma != 1.0 then a pow(color,Gamma) is applied
-* @param ClipRef - reference value to compare with alpha for killing pixels
-* @param SmoothWidth - The width to smooth the edge the texture
-* @param EnableShadow - Toggles drop shadow rendering
-* @param ShadowDirection - 2D vector specifying the direction of shadow
-* @param ShadowColor - Color of the shadowed pixels
-* @param ShadowSmoothWidth - The width to smooth the edge the shadow of the texture
-* @param BlendMode - current batched element blend mode being rendered
-*/
 void FSimpleElementDistanceFieldGammaPS::SetParameters(
-	FRHICommandList& RHICmdList, 
+	FRHIBatchedShaderParameters& BatchedParameters,
 	const FTexture* Texture,
 	float InGamma,
 	float InClipRef,
@@ -174,28 +172,56 @@ void FSimpleElementDistanceFieldGammaPS::SetParameters(
 	float ShadowSmoothWidthValue,
 	const FDepthFieldGlowInfo& GlowInfo,
 	ESimpleElementBlendMode BlendMode
-	)
+)
 {
-	FSimpleElementMaskedGammaBasePS::SetParameters(RHICmdList, Texture,InGamma,InClipRef,BlendMode);
-	SetShaderValue(RHICmdList, RHICmdList.GetBoundPixelShader(),SmoothWidth,SmoothWidthValue);		
+	FSimpleElementMaskedGammaBasePS::SetParameters(BatchedParameters, Texture, InGamma, InClipRef, BlendMode);
+	SetShaderValue(BatchedParameters, SmoothWidth, SmoothWidthValue);
 	const uint32 bEnableShadowValueUInt = (bEnableShadowValue ? 1 : 0);
-	SetShaderValue(RHICmdList, RHICmdList.GetBoundPixelShader(),EnableShadow,bEnableShadowValueUInt);
+	SetShaderValue(BatchedParameters, EnableShadow, bEnableShadowValueUInt);
 	if (bEnableShadowValue)
 	{
-		SetShaderValue(RHICmdList, RHICmdList.GetBoundPixelShader(),ShadowDirection,FVector2f(ShadowDirectionValue));
-		SetShaderValue(RHICmdList, RHICmdList.GetBoundPixelShader(),ShadowColor,ShadowColorValue);
-		SetShaderValue(RHICmdList, RHICmdList.GetBoundPixelShader(),ShadowSmoothWidth,ShadowSmoothWidthValue);
+		SetShaderValue(BatchedParameters, ShadowDirection, FVector2f(ShadowDirectionValue));
+		SetShaderValue(BatchedParameters, ShadowColor, ShadowColorValue);
+		SetShaderValue(BatchedParameters, ShadowSmoothWidth, ShadowSmoothWidthValue);
 	}
-	SetShaderValue(RHICmdList, RHICmdList.GetBoundPixelShader(),EnableGlow,GlowInfo.bEnableGlow);
+	SetShaderValue(BatchedParameters, EnableGlow, GlowInfo.bEnableGlow);
 	if (GlowInfo.bEnableGlow)
 	{
-		SetShaderValue(RHICmdList, RHICmdList.GetBoundPixelShader(),GlowColor,GlowInfo.GlowColor);
-		SetShaderValue(RHICmdList, RHICmdList.GetBoundPixelShader(),GlowOuterRadius,FVector2f(GlowInfo.GlowOuterRadius));
-		SetShaderValue(RHICmdList, RHICmdList.GetBoundPixelShader(),GlowInnerRadius,FVector2f(GlowInfo.GlowInnerRadius));
+		SetShaderValue(BatchedParameters, GlowColor, GlowInfo.GlowColor);
+		SetShaderValue(BatchedParameters, GlowOuterRadius, FVector2f(GlowInfo.GlowOuterRadius));
+		SetShaderValue(BatchedParameters, GlowInnerRadius, FVector2f(GlowInfo.GlowInnerRadius));
 	}
+}
 
-	// This shader does not use editor compositing
-	SetEditorCompositingParameters(RHICmdList, NULL);
+void FSimpleElementDistanceFieldGammaPS::SetParameters(
+	FRHICommandList& RHICmdList, 
+	const FTexture* Texture,
+	float InGamma,
+	float InClipRef,
+	float SmoothWidthValue,
+	bool bEnableShadowValue,
+	const FVector2D& ShadowDirectionValue,
+	const FLinearColor& ShadowColorValue,
+	float ShadowSmoothWidthValue,
+	const FDepthFieldGlowInfo& GlowInfo,
+	ESimpleElementBlendMode BlendMode
+	)
+{
+	FRHIBatchedShaderParameters BatchedParameters;
+	SetParameters(
+		BatchedParameters
+		, Texture
+		, InGamma
+		, InClipRef
+		, SmoothWidthValue
+		, bEnableShadowValue
+		, ShadowDirectionValue
+		, ShadowColorValue
+		, ShadowSmoothWidthValue
+		, GlowInfo
+		, BlendMode
+	);
+	RHICmdList.SetBatchedShaderParameters(RHICmdList.GetBoundPixelShader(), BatchedParameters);
 }
 
 /**
@@ -226,19 +252,22 @@ FSimpleElementHitProxyPS::FSimpleElementHitProxyPS(const ShaderMetaType::Compile
 	InTextureSampler.Bind(Initializer.ParameterMap,TEXT("InTextureSampler"));
 }
 
-void FSimpleElementHitProxyPS::SetParameters(FRHICommandList& RHICmdList, const FTexture* TextureValue)
+bool FSimpleElementHitProxyPS::ShouldCompilePermutation(const FGlobalShaderPermutationParameters& Parameters)
 {
-	SetTextureParameter(RHICmdList, RHICmdList.GetBoundPixelShader(),InTexture,InTextureSampler,TextureValue);
+	return IsPCPlatform(Parameters.Platform);
 }
 
-/*bool FSimpleElementHitProxyPS::Serialize(FArchive& Ar)
+void FSimpleElementHitProxyPS::SetParameters(FRHIBatchedShaderParameters& BatchedParameters, const FTexture* TextureValue)
 {
-	bool bShaderHasOutdatedParameters = FGlobalShader::Serialize(Ar);
-	Ar << InTexture;
-	Ar << InTextureSampler;
-	return bShaderHasOutdatedParameters;
-}*/
+	SetTextureParameter(BatchedParameters, InTexture, InTextureSampler, TextureValue);
+}
 
+void FSimpleElementHitProxyPS::SetParameters(FRHICommandList& RHICmdList, const FTexture* TextureValue)
+{
+	FRHIBatchedShaderParameters BatchedParameters;
+	SetParameters(BatchedParameters, TextureValue);
+	RHICmdList.SetBatchedShaderParameters(RHICmdList.GetBoundPixelShader(), BatchedParameters);
+}
 
 FSimpleElementColorChannelMaskPS::FSimpleElementColorChannelMaskPS(const ShaderMetaType::CompiledShaderInitializerType& Initializer)
 : FGlobalShader(Initializer)
@@ -249,18 +278,23 @@ FSimpleElementColorChannelMaskPS::FSimpleElementColorChannelMaskPS(const ShaderM
 	Gamma.Bind(Initializer.ParameterMap,TEXT("Gamma"));
 }
 
-/**
-* Sets all the constant parameters for this shader
-*
-* @param Texture - 2d tile texture
-* @param ColorWeights - reference value to compare with alpha for killing pixels
-* @param Gamma - if gamma != 1.0 then a pow(color,Gamma) is applied
-*/
+bool FSimpleElementColorChannelMaskPS::ShouldCompilePermutation(const FGlobalShaderPermutationParameters& Parameters)
+{
+	return IsPCPlatform(Parameters.Platform);
+}
+
+void FSimpleElementColorChannelMaskPS::SetParameters(FRHIBatchedShaderParameters& BatchedParameters, const FTexture* TextureValue, const FMatrix& ColorWeightsValue, float GammaValue)
+{
+	SetTextureParameter(BatchedParameters, InTexture, InTextureSampler, TextureValue);
+	SetShaderValue(BatchedParameters, ColorWeights, (FMatrix44f)ColorWeightsValue);
+	SetShaderValue(BatchedParameters, Gamma, GammaValue);
+}
+
 void FSimpleElementColorChannelMaskPS::SetParameters(FRHICommandList& RHICmdList, const FTexture* TextureValue, const FMatrix& ColorWeightsValue, float GammaValue)
 {
-	SetTextureParameter(RHICmdList, RHICmdList.GetBoundPixelShader(),InTexture,InTextureSampler,TextureValue);
-	SetShaderValue(RHICmdList, RHICmdList.GetBoundPixelShader(),ColorWeights, (FMatrix44f)ColorWeightsValue);
-	SetShaderValue(RHICmdList, RHICmdList.GetBoundPixelShader(),Gamma,GammaValue);
+	FRHIBatchedShaderParameters BatchedParameters;
+	SetParameters(BatchedParameters, TextureValue, ColorWeightsValue, GammaValue);
+	RHICmdList.SetBatchedShaderParameters(RHICmdList.GetBoundPixelShader(), BatchedParameters);
 }
 
 /*bool FSimpleElementColorChannelMaskPS::Serialize(FArchive& Ar)
diff --git a/Engine/Source/Runtime/Engine/Private/VectorField.cpp b/Engine/Source/Runtime/Engine/Private/VectorField.cpp
--- a/Engine/Source/Runtime/Engine/Private/VectorField.cpp
+++ b/Engine/Source/Runtime/Engine/Private/VectorField.cpp
@@ -996,41 +996,22 @@ public:
 		OutVolumeTextureSampler.Bind( Initializer.ParameterMap, TEXT("OutVolumeTextureSampler") );
 	}
 
-	/**
-	 * Set parameters for this shader.
-	 * @param UniformBuffer - Uniform buffer containing parameters for compositing vector fields.
-	 * @param AtlasTextureRHI - The atlas texture with which to create the vector field.
-	 * @param NoiseVolumeTextureRHI - The volume texture to use to add noise to the vector field.
-	 */
 	void SetParameters(
-		FRHICommandList& RHICmdList, 
-		const FCompositeAnimatedVectorFieldUniformBufferRef& UniformBuffer,
+		FRHIBatchedShaderParameters& BatchedParameters,
+		FRHIUnorderedAccessView* VolumeTextureUAV,
 		FRHITexture* AtlasTextureRHI,
-		FRHITexture* NoiseVolumeTextureRHI )
+		FRHITexture* NoiseVolumeTextureRHI)
 	{
-		FRHIComputeShader* ComputeShaderRHI = RHICmdList.GetBoundComputeShader();
-		FRHISamplerState* SamplerStateLinear = TStaticSamplerState<SF_Bilinear,AM_Clamp,AM_Clamp,AM_Clamp>::GetRHI();
-		SetUniformBufferParameter(RHICmdList, ComputeShaderRHI, GetUniformBufferParameter<FCompositeAnimatedVectorFieldUniformParameters>(), UniformBuffer );
-		SetTextureParameter(RHICmdList, ComputeShaderRHI, AtlasTexture, AtlasTextureSampler, SamplerStateLinear, AtlasTextureRHI );
-		SetTextureParameter(RHICmdList, ComputeShaderRHI, NoiseVolumeTexture, NoiseVolumeTextureSampler, SamplerStateLinear, NoiseVolumeTextureRHI );
-	}
+		FRHISamplerState* SamplerStateLinear = TStaticSamplerState<SF_Bilinear, AM_Clamp, AM_Clamp, AM_Clamp>::GetRHI();
 
-	/**
-	 * Set output buffer for this shader.
-	 */
-	void SetOutput(FRHICommandList& RHICmdList, FRHIUnorderedAccessView* VolumeTextureUAV)
-	{
-		FRHIComputeShader* ComputeShaderRHI = RHICmdList.GetBoundComputeShader();
-		SetUAVParameter(RHICmdList, ComputeShaderRHI, OutVolumeTexture, VolumeTextureUAV);
+		SetUAVParameter(BatchedParameters, OutVolumeTexture, VolumeTextureUAV);
+		SetTextureParameter(BatchedParameters, AtlasTexture, AtlasTextureSampler, SamplerStateLinear, AtlasTextureRHI);
+		SetTextureParameter(BatchedParameters, NoiseVolumeTexture, NoiseVolumeTextureSampler, SamplerStateLinear, NoiseVolumeTextureRHI);
 	}
 
-	/**
-	 * Unbinds any buffers that have been bound.
-	 */
-	void UnbindBuffers(FRHICommandList& RHICmdList)
+	void UnsetParameters(FRHIBatchedShaderParameters& BatchedParameters)
 	{
-		FRHIComputeShader* ComputeShaderRHI = RHICmdList.GetBoundComputeShader();
-		SetUAVParameter(RHICmdList, ComputeShaderRHI, OutVolumeTexture, nullptr);
+		SetUAVParameter(BatchedParameters, OutVolumeTexture, nullptr);
 	}
 
 private:
@@ -1184,20 +1165,25 @@ public:
 
 			RHICmdList.Transition(FRHITransitionInfo(VolumeTextureUAV, ERHIAccess::Unknown, ERHIAccess::UAVCompute));
 			SetComputePipelineState(RHICmdList, CompositeCS.GetComputeShader());
-			CompositeCS->SetOutput(RHICmdList, VolumeTextureUAV);
-			/// ?
-			CompositeCS->SetParameters(
+
+			SetUniformBufferParameter(RHICmdList, CompositeCS.GetComputeShader(), CompositeCS->GetUniformBufferParameter<FCompositeAnimatedVectorFieldUniformParameters>(), UniformBuffer);
+
+			SetAllShaderParametersCS(
 				RHICmdList,
-				UniformBuffer,
+				CompositeCS,
+				VolumeTextureUAV,
 				AnimatedVectorField->Texture->GetResource()->TextureRHI,
-				NoiseVolumeTextureRHI );
+				NoiseVolumeTextureRHI);
+
 			DispatchComputeShader(
 				RHICmdList,
 				CompositeCS.GetShader(),
 				SizeX / THREADS_PER_AXIS,
 				SizeY / THREADS_PER_AXIS,
 				SizeZ / THREADS_PER_AXIS );
-			CompositeCS->UnbindBuffers(RHICmdList);
+
+			UnsetAllShaderParametersCS(RHICmdList, CompositeCS);
+
 			RHICmdList.Transition(FRHITransitionInfo(VolumeTextureUAV, ERHIAccess::UAVCompute, ERHIAccess::SRVMask));
 		}
 	}
diff --git a/Engine/Source/Runtime/Engine/Public/CubemapUnwrapUtils.h b/Engine/Source/Runtime/Engine/Public/CubemapUnwrapUtils.h
--- a/Engine/Source/Runtime/Engine/Public/CubemapUnwrapUtils.h
+++ b/Engine/Source/Runtime/Engine/Public/CubemapUnwrapUtils.h
@@ -52,6 +52,9 @@ public:
 
 	static bool ShouldCompilePermutation(const FGlobalShaderPermutationParameters& Parameters);
 
+	void SetParameters(FRHIBatchedShaderParameters& BatchedParameters, const FMatrix& TransformValue);
+
+	UE_DEPRECATED(5.3, "SetParameters with FRHIBatchedShaderParameters should be used.")
 	void SetParameters(FRHICommandList& RHICmdList, const FMatrix& TransformValue);
 
 private:
@@ -74,6 +77,10 @@ public:
 	FCubemapTexturePropertiesPS(const ShaderMetaType::CompiledShaderInitializerType& Initializer);
 
 	static bool ShouldCompilePermutation(const FGlobalShaderPermutationParameters& Parameters);
+
+	void SetParameters(FRHIBatchedShaderParameters& BatchedParameters, const FTexture* InTexture, const FMatrix& InColorWeightsValue, float InMipLevel, float InSliceIndex, bool bInIsTextureCubeArray, const FMatrix44f& InViewMatrix, bool bInShowLongLatUnwrap, float InGammaValue, bool bInUsePointSampling);
+
+	UE_DEPRECATED(5.3, "SetParameters with FRHIBatchedShaderParameters should be used.")
 	void SetParameters(FRHICommandList& RHICmdList, const FTexture* InTexture, const FMatrix& InColorWeightsValue, float InMipLevel, float InSliceIndex, bool bInIsTextureCubeArray, const FMatrix44f& InViewMatrix, bool bInShowLongLatUnwrap, float InGammaValue, bool bInUsePointSampling);
 
 private:
@@ -141,6 +148,9 @@ public:
 		OutEnvironment.SetDefine(TEXT("USE_IES_STANDALONE_TEXTURE"), 1);
 	}
 
+	void SetParameters(FRHIBatchedShaderParameters& BatchedParameters, const FTexture* Texture, float InBrightnessInLumens);
+
+	UE_DEPRECATED(5.3, "SetParameters with FRHIBatchedShaderParameters should be used.")
 	void SetParameters(FRHICommandList& RHICmdList, const FTexture* Texture, float InBrightnessInLumens);
 
 private:
diff --git a/Engine/Source/Runtime/Engine/Public/SimpleElementShaders.h b/Engine/Source/Runtime/Engine/Public/SimpleElementShaders.h
--- a/Engine/Source/Runtime/Engine/Public/SimpleElementShaders.h
+++ b/Engine/Source/Runtime/Engine/Public/SimpleElementShaders.h
@@ -24,17 +24,17 @@ class ENGINE_API FSimpleElementVS : public FGlobalShader
 {
 	DECLARE_SHADER_TYPE(FSimpleElementVS,Global);
 public:
-
-	static bool ShouldCompilePermutation(const FGlobalShaderPermutationParameters& Parameters) { return true; }
-
 	FSimpleElementVS(const ShaderMetaType::CompiledShaderInitializerType& Initializer);
 	FSimpleElementVS() {}
 
+	void SetParameters(FRHIBatchedShaderParameters& BatchedParameters, const FMatrix& WorldToClipMatrix);
+	void SetParameters(FRHIBatchedShaderParameters& BatchedParameters, const FRelativeViewMatrices& Matrices);
+
+	UE_DEPRECATED(5.3, "SetParameters with FRHIBatchedShaderParameters should be used.")
 	void SetParameters(FRHICommandList& RHICmdList, const FMatrix& WorldToClipMatrix);
+	UE_DEPRECATED(5.3, "SetParameters with FRHIBatchedShaderParameters should be used.")
 	void SetParameters(FRHICommandList& RHICmdList, const FRelativeViewMatrices& Matrices);
 
-	//virtual bool Serialize(FArchive& Ar) override;
-
 	static void ModifyCompilationEnvironment(const FGlobalShaderPermutationParameters& Parameters, FShaderCompilerEnvironment& OutEnvironment);
 
 private:
@@ -50,9 +50,6 @@ class FSimpleElementPS : public FGlobalShader
 {
 	DECLARE_SHADER_TYPE(FSimpleElementPS, Global);
 public:
-
-	static bool ShouldCompilePermutation(const FGlobalShaderPermutationParameters& Parameters) { return true; }
-
 	FSimpleElementPS(const ShaderMetaType::CompiledShaderInitializerType& Initializer);
 	FSimpleElementPS() {}
 
@@ -60,11 +57,13 @@ public:
 	 * Sets parameters for compositing editor primitives
 	 *
 	 * @param View			SceneView for view constants when compositing
-	 * @param DepthTexture	Depth texture to read from when depth testing for compositing.  If not set no compositing will occur
 	 */
 	void SetEditorCompositingParameters(FRHICommandList& RHICmdList, const FSceneView* View);
 
-	void SetParameters(FRHICommandList& RHICmdList, const FTexture* TextureValue );
+	void SetParameters(FRHIBatchedShaderParameters& BatchedParameters, const FTexture* TextureValue);
+
+	UE_DEPRECATED(5.3, "SetParameters with FRHIBatchedShaderParameters should be used.")
+	void SetParameters(FRHICommandList& RHICmdList, const FTexture* TextureValue);
 
 	//virtual bool Serialize(FArchive& Ar) override;
 
@@ -100,14 +99,13 @@ public:
 	FSimpleElementGammaBasePS(const ShaderMetaType::CompiledShaderInitializerType& Initializer);
 	FSimpleElementGammaBasePS() {}
 
-	void SetParameters(FRHICommandList& RHICmdList, const FTexture* Texture,float GammaValue,ESimpleElementBlendMode BlendMode);
+	void SetParameters(FRHIBatchedShaderParameters& BatchedParameters, const FTexture* Texture, float GammaValue, ESimpleElementBlendMode BlendMode);
 
-	//virtual bool Serialize(FArchive& Ar) override;
+	UE_DEPRECATED(5.3, "SetParameters with FRHIBatchedShaderParameters should be used.")
+	void SetParameters(FRHICommandList& RHICmdList, const FTexture* Texture,float GammaValue,ESimpleElementBlendMode BlendMode);
 
 private:
-	
-		LAYOUT_FIELD(FShaderParameter, Gamma)
-	
+	LAYOUT_FIELD(FShaderParameter, Gamma)
 };
 
 template <bool bSRGBTexture>
@@ -115,12 +113,9 @@ class FSimpleElementGammaPS : public FSimpleElementGammaBasePS
 {
 	DECLARE_SHADER_TYPE(FSimpleElementGammaPS, Global);
 public:
-
 	FSimpleElementGammaPS(const ShaderMetaType::CompiledShaderInitializerType& Initializer) : FSimpleElementGammaBasePS(Initializer) {}
 	FSimpleElementGammaPS() {}
 
-	static bool ShouldCompilePermutation(const FGlobalShaderPermutationParameters& Parameters) { return true; }
-
 	static void ModifyCompilationEnvironment(const FGlobalShaderPermutationParameters& Parameters, FShaderCompilerEnvironment& OutEnvironment)
 	{
 		OutEnvironment.SetDefine(TEXT("SRGB_INPUT_TEXTURE"), bSRGBTexture);
@@ -134,11 +129,8 @@ class FSimpleElementGammaAlphaOnlyPS : public FSimpleElementGammaBasePS
 {
 	DECLARE_SHADER_TYPE(FSimpleElementGammaAlphaOnlyPS, Global);
 public:
-
 	FSimpleElementGammaAlphaOnlyPS(const ShaderMetaType::CompiledShaderInitializerType& Initializer) : FSimpleElementGammaBasePS(Initializer) {}
 	FSimpleElementGammaAlphaOnlyPS() {}
-
-	static bool ShouldCompilePermutation(const FGlobalShaderPermutationParameters& Parameters) { return true; }
 };
 
 /**
@@ -148,18 +140,16 @@ class FSimpleElementMaskedGammaBasePS : public FSimpleElementGammaBasePS
 {
 	DECLARE_TYPE_LAYOUT(FSimpleElementMaskedGammaBasePS, NonVirtual);
 public:
-
 	FSimpleElementMaskedGammaBasePS(const ShaderMetaType::CompiledShaderInitializerType& Initializer);
 	FSimpleElementMaskedGammaBasePS() {}
 
-	void SetParameters(FRHICommandList& RHICmdList, const FTexture* Texture,float Gamma,float ClipRefValue,ESimpleElementBlendMode BlendMode);
+	void SetParameters(FRHIBatchedShaderParameters& BatchedParameters, const FTexture* Texture,float Gamma,float ClipRefValue,ESimpleElementBlendMode BlendMode);
 
-	//virtual bool Serialize(FArchive& Ar) override;
+	UE_DEPRECATED(5.3, "SetParameters with FRHIBatchedShaderParameters should be used.")
+	void SetParameters(FRHICommandList& RHICmdList, const FTexture* Texture,float Gamma,float ClipRefValue,ESimpleElementBlendMode BlendMode);
 
 private:
-	
-		LAYOUT_FIELD(FShaderParameter, ClipRef)
-	
+	LAYOUT_FIELD(FShaderParameter, ClipRef)
 };
 
 template <bool bSRGBTexture>
@@ -167,11 +157,9 @@ class FSimpleElementMaskedGammaPS : public FSimpleElementMaskedGammaBasePS
 {
 	DECLARE_SHADER_TYPE(FSimpleElementMaskedGammaPS, Global);
 public:
-
 	FSimpleElementMaskedGammaPS(const ShaderMetaType::CompiledShaderInitializerType& Initializer) : FSimpleElementMaskedGammaBasePS(Initializer) {}
 	FSimpleElementMaskedGammaPS() {}
 
-	static bool ShouldCompilePermutation(const FGlobalShaderPermutationParameters& Parameters) { return true; }
 	static void ModifyCompilationEnvironment(const FGlobalShaderPermutationParameters& Parameters, FShaderCompilerEnvironment& OutEnvironment)
 	{
 		OutEnvironment.SetDefine(TEXT("SRGB_INPUT_TEXTURE"), bSRGBTexture);
@@ -185,28 +173,7 @@ class FSimpleElementDistanceFieldGammaPS : public FSimpleElementMaskedGammaBaseP
 {
 	DECLARE_SHADER_TYPE(FSimpleElementDistanceFieldGammaPS,Global);
 public:
-
-	/**
-	* Determine if this shader should be compiled
-	*
-	* @param Platform - current shader platform being compiled
-	* @return true if this shader should be cached for the given platform
-	*/
-	static bool ShouldCompilePermutation(const FGlobalShaderPermutationParameters& Parameters) 
-	{ 
-		return true; 
-	}
-
-	/**
-	* Constructor
-	*
-	* @param Initializer - shader initialization container
-	*/
 	FSimpleElementDistanceFieldGammaPS(const ShaderMetaType::CompiledShaderInitializerType& Initializer);
-
-	/**
-	* Default constructor
-	*/
 	FSimpleElementDistanceFieldGammaPS() {}
 
 	/**
@@ -223,7 +190,7 @@ public:
 	* @param BlendMode - current batched element blend mode being rendered
 	*/
 	void SetParameters(
-		FRHICommandList& RHICmdList, 
+		FRHIBatchedShaderParameters& BatchedParameters,
 		const FTexture* Texture,
 		float Gamma,
 		float ClipRef,
@@ -236,13 +203,20 @@ public:
 		ESimpleElementBlendMode BlendMode
 		);
 
-	/**
-	* Serialize constant paramaters for this shader
-	* 
-	* @param Ar - archive to serialize to
-	* @return true if any of the parameters were outdated
-	*/
-	//virtual bool Serialize(FArchive& Ar) override;
+	UE_DEPRECATED(5.3, "SetParameters with FRHIBatchedShaderParameters should be used.")
+	void SetParameters(
+		FRHICommandList& RHICmdList, 
+		const FTexture* Texture,
+		float Gamma,
+		float ClipRef,
+		float SmoothWidthValue,
+		bool EnableShadowValue,
+		const FVector2D& ShadowDirectionValue,
+		const FLinearColor& ShadowColorValue,
+		float ShadowSmoothWidthValue,
+		const FDepthFieldGlowInfo& GlowInfo,
+		ESimpleElementBlendMode BlendMode
+		);
 	
 private:
 	/** The width to smooth the edge the texture */
@@ -272,26 +246,21 @@ class FSimpleElementHitProxyPS : public FGlobalShader
 {
 	DECLARE_SHADER_TYPE(FSimpleElementHitProxyPS,Global);
 public:
-
-	static bool ShouldCompilePermutation(const FGlobalShaderPermutationParameters& Parameters) 
-	{ 
-		return IsPCPlatform(Parameters.Platform); 
-	}
-
-
 	FSimpleElementHitProxyPS(const ShaderMetaType::CompiledShaderInitializerType& Initializer);
 	FSimpleElementHitProxyPS() {}
 
-	void SetParameters(FRHICommandList& RHICmdList, const FTexture* TextureValue);
+	static bool ShouldCompilePermutation(const FGlobalShaderPermutationParameters& Parameters);
 
-	//virtual bool Serialize(FArchive& Ar) override;
+	void SetParameters(FRHIBatchedShaderParameters& BatchedParameters, const FTexture* TextureValue);
+
+	UE_DEPRECATED(5.3, "SetParameters with FRHIBatchedShaderParameters should be used.")
+	void SetParameters(FRHICommandList& RHICmdList, const FTexture* TextureValue);
 
 private:
 	LAYOUT_FIELD(FShaderResourceParameter, InTexture)
 	LAYOUT_FIELD(FShaderResourceParameter, InTextureSampler)
 };
 
-
 /**
 * A pixel shader for rendering a texture with the ability to weight the colors for each channel.
 * The shader also features the ability to view alpha channels and desaturate any red, green or blue channels
@@ -301,16 +270,11 @@ class FSimpleElementColorChannelMaskPS : public FGlobalShader
 {
 	DECLARE_SHADER_TYPE(FSimpleElementColorChannelMaskPS,Global);
 public:
-
-	static bool ShouldCompilePermutation(const FGlobalShaderPermutationParameters& Parameters) 
-	{ 
-		return IsPCPlatform(Parameters.Platform); 
-	}
-
-
 	FSimpleElementColorChannelMaskPS(const ShaderMetaType::CompiledShaderInitializerType& Initializer);
 	FSimpleElementColorChannelMaskPS() {}
 
+	static bool ShouldCompilePermutation(const FGlobalShaderPermutationParameters& Parameters);
+
 	/**
 	* Sets all the constant parameters for this shader
 	*
@@ -318,9 +282,10 @@ public:
 	* @param ColorWeights - reference value to compare with alpha for killing pixels
 	* @param Gamma - if gamma != 1.0 then a pow(color,Gamma) is applied
 	*/
-	void SetParameters(FRHICommandList& RHICmdList, const FTexture* TextureValue, const FMatrix& ColorWeightsValue, float GammaValue);
+	void SetParameters(FRHIBatchedShaderParameters& BatchedParameters, const FTexture* TextureValue, const FMatrix& ColorWeightsValue, float GammaValue);
 
-	//virtual bool Serialize(FArchive& Ar) override;
+	UE_DEPRECATED(5.3, "SetParameters with FRHIBatchedShaderParameters should be used.")
+	void SetParameters(FRHICommandList& RHICmdList, const FTexture* TextureValue, const FMatrix& ColorWeightsValue, float GammaValue);
 
 private:
 	LAYOUT_FIELD(FShaderResourceParameter, InTexture)
diff --git a/Engine/Source/Runtime/Landscape/Private/LandscapeEditLayers.cpp b/Engine/Source/Runtime/Landscape/Private/LandscapeEditLayers.cpp
--- a/Engine/Source/Runtime/Landscape/Private/LandscapeEditLayers.cpp
+++ b/Engine/Source/Runtime/Landscape/Private/LandscapeEditLayers.cpp
@@ -443,9 +443,9 @@ public:
 	FLandscapeLayersVS()
 	{}
 
-	void SetParameters(FRHICommandList& RHICmdList, const FMatrix44f& InProjectionMatrix)
+	void SetParameters(FRHIBatchedShaderParameters& BatchedParameters, const FMatrix44f& InProjectionMatrix)
 	{
-		SetShaderValue(RHICmdList, RHICmdList.GetBoundVertexShader(), TransformParam, InProjectionMatrix);
+		SetShaderValue(BatchedParameters, TransformParam, InProjectionMatrix);
 	}
 
 private:
@@ -517,20 +517,20 @@ public:
 	FLandscapeLayersHeightmapPS()
 	{}
 
-	void SetParameters(FRHICommandList& RHICmdList, const FLandscapeLayersHeightmapShaderParameters& InParams)
+	void SetParameters(FRHIBatchedShaderParameters& BatchedParameters, const FLandscapeLayersHeightmapShaderParameters& InParams)
 	{
-		SetTextureParameter(RHICmdList, RHICmdList.GetBoundPixelShader(), ReadTexture1Param, ReadTexture1SamplerParam, TStaticSamplerState<SF_Point, AM_Clamp, AM_Clamp, AM_Clamp>::GetRHI(), InParams.ReadHeightmap1->GetResource()->TextureRHI);
-		SetTextureParameter(RHICmdList, RHICmdList.GetBoundPixelShader(), ReadTexture2Param, ReadTexture2SamplerParam, TStaticSamplerState<SF_Point, AM_Clamp, AM_Clamp, AM_Clamp>::GetRHI(), InParams.ReadHeightmap2 != nullptr ? InParams.ReadHeightmap2->GetResource()->TextureRHI : GWhiteTexture->TextureRHI);
+		SetTextureParameter(BatchedParameters, ReadTexture1Param, ReadTexture1SamplerParam, TStaticSamplerState<SF_Point, AM_Clamp, AM_Clamp, AM_Clamp>::GetRHI(), InParams.ReadHeightmap1->GetResource()->TextureRHI);
+		SetTextureParameter(BatchedParameters, ReadTexture2Param, ReadTexture2SamplerParam, TStaticSamplerState<SF_Point, AM_Clamp, AM_Clamp, AM_Clamp>::GetRHI(), InParams.ReadHeightmap2 != nullptr ? InParams.ReadHeightmap2->GetResource()->TextureRHI : GWhiteTexture->TextureRHI);
 
 		FVector4f LayerInfo(InParams.LayerAlpha, InParams.LayerVisible ? 1.0f : 0.0f, InParams.LayerBlendMode == LSBM_AlphaBlend ? 1.0f : 0.f, 0.f);
 		FVector4f OutputConfig(InParams.ApplyLayerModifiers ? 1.0f : 0.0f, 0.0f /*unused*/, InParams.ReadHeightmap2 ? 1.0f : 0.0f, InParams.GenerateNormals ? 1.0f : 0.0f);
 		FVector2f TextureSize(static_cast<float>(InParams.HeightmapSize.X), static_cast<float>(InParams.HeightmapSize.Y));
 
-		SetShaderValue(RHICmdList, RHICmdList.GetBoundPixelShader(), LayerInfoParam, LayerInfo);
-		SetShaderValue(RHICmdList, RHICmdList.GetBoundPixelShader(), OutputConfigParam, OutputConfig);
-		SetShaderValue(RHICmdList, RHICmdList.GetBoundPixelShader(), TextureSizeParam, TextureSize);
-		SetShaderValue(RHICmdList, RHICmdList.GetBoundPixelShader(), LandscapeGridScaleParam, FVector3f(InParams.GridSize));
-		SetShaderValue(RHICmdList, RHICmdList.GetBoundPixelShader(), ComponentVertexCountParam, (float)InParams.CurrentMipComponentVertexCount);
+		SetShaderValue(BatchedParameters, LayerInfoParam, LayerInfo);
+		SetShaderValue(BatchedParameters, OutputConfigParam, OutputConfig);
+		SetShaderValue(BatchedParameters, TextureSizeParam, TextureSize);
+		SetShaderValue(BatchedParameters, LandscapeGridScaleParam, FVector3f(InParams.GridSize));
+		SetShaderValue(BatchedParameters, ComponentVertexCountParam, (float)InParams.CurrentMipComponentVertexCount);
 	}
 
 private:
@@ -574,13 +574,13 @@ public:
 	FLandscapeLayersHeightmapMipsPS()
 	{}
 
-	void SetParameters(FRHICommandList& RHICmdList, const FLandscapeLayersHeightmapShaderParameters& InParams)
+	void SetParameters(FRHIBatchedShaderParameters& BatchedParameters, const FLandscapeLayersHeightmapShaderParameters& InParams)
 	{
-		SetTextureParameter(RHICmdList, RHICmdList.GetBoundPixelShader(), ReadTexture1Param, ReadTexture1SamplerParam, TStaticSamplerState<SF_Point, AM_Clamp, AM_Clamp, AM_Clamp>::GetRHI(), InParams.ReadHeightmap1->GetResource()->TextureRHI);
+		SetTextureParameter(BatchedParameters, ReadTexture1Param, ReadTexture1SamplerParam, TStaticSamplerState<SF_Point, AM_Clamp, AM_Clamp, AM_Clamp>::GetRHI(), InParams.ReadHeightmap1->GetResource()->TextureRHI);
 
-		SetShaderValue(RHICmdList, RHICmdList.GetBoundPixelShader(), CurrentMipSizeParam, FVector2f(static_cast<float>(InParams.CurrentMipSize.X), static_cast<float>(InParams.CurrentMipSize.Y)));
-		SetShaderValue(RHICmdList, RHICmdList.GetBoundPixelShader(), ParentMipSizeParam, FVector2f(static_cast<float>(InParams.ParentMipSize.X), static_cast<float>(InParams.ParentMipSize.Y)));
-		SetShaderValue(RHICmdList, RHICmdList.GetBoundPixelShader(), CurrentMipComponentVertexCountParam, (float)InParams.CurrentMipComponentVertexCount);
+		SetShaderValue(BatchedParameters, CurrentMipSizeParam, FVector2f(static_cast<float>(InParams.CurrentMipSize.X), static_cast<float>(InParams.CurrentMipSize.Y)));
+		SetShaderValue(BatchedParameters, ParentMipSizeParam, FVector2f(static_cast<float>(InParams.ParentMipSize.X), static_cast<float>(InParams.ParentMipSize.Y)));
+		SetShaderValue(BatchedParameters, CurrentMipComponentVertexCountParam, (float)InParams.CurrentMipComponentVertexCount);
 	}
 
 private:
@@ -651,17 +651,17 @@ public:
 	FLandscapeLayersWeightmapPS()
 	{}
 
-	void SetParameters(FRHICommandList& RHICmdList, const FLandscapeLayersWeightmapShaderParameters& InParams)
+	void SetParameters(FRHIBatchedShaderParameters& BatchedParameters, const FLandscapeLayersWeightmapShaderParameters& InParams)
 	{
-		SetTextureParameter(RHICmdList, RHICmdList.GetBoundPixelShader(), ReadTexture1Param, ReadTexture1SamplerParam, TStaticSamplerState<SF_Point, AM_Clamp, AM_Clamp, AM_Clamp>::GetRHI(), InParams.ReadWeightmap1->GetResource()->TextureRHI);
-		SetTextureParameter(RHICmdList, RHICmdList.GetBoundPixelShader(), ReadTexture2Param, ReadTexture2SamplerParam, TStaticSamplerState<SF_Point, AM_Clamp, AM_Clamp, AM_Clamp>::GetRHI(), InParams.ReadWeightmap2 != nullptr ? InParams.ReadWeightmap2->GetResource()->TextureRHI : GWhiteTexture->TextureRHI);
+		SetTextureParameter(BatchedParameters, ReadTexture1Param, ReadTexture1SamplerParam, TStaticSamplerState<SF_Point, AM_Clamp, AM_Clamp, AM_Clamp>::GetRHI(), InParams.ReadWeightmap1->GetResource()->TextureRHI);
+		SetTextureParameter(BatchedParameters, ReadTexture2Param, ReadTexture2SamplerParam, TStaticSamplerState<SF_Point, AM_Clamp, AM_Clamp, AM_Clamp>::GetRHI(), InParams.ReadWeightmap2 != nullptr ? InParams.ReadWeightmap2->GetResource()->TextureRHI : GWhiteTexture->TextureRHI);
 
 		FVector4f LayerInfo(InParams.LayerAlpha, InParams.LayerVisible ? 1.0f : 0.0f, InParams.LayerBlendMode == LSBM_AlphaBlend ? 1.0f : 0.f, 0.f);
 		FVector4f OutputConfig(InParams.ApplyLayerModifiers ? 1.0f : 0.0f, InParams.OutputAsSubstractive ? 1.0f : 0.0f, InParams.ReadWeightmap2 != nullptr ? 1.0f : 0.0f, InParams.OutputAsNormalized ? 1.0f : 0.0f);
 
-		SetShaderValue(RHICmdList, RHICmdList.GetBoundPixelShader(), LayerInfoParam, LayerInfo);
-		SetShaderValue(RHICmdList, RHICmdList.GetBoundPixelShader(), OutputConfigParam, OutputConfig);
-		SetShaderValue(RHICmdList, RHICmdList.GetBoundPixelShader(), ComponentVertexCountParam, (float)InParams.CurrentMipComponentVertexCount);
+		SetShaderValue(BatchedParameters, LayerInfoParam, LayerInfo);
+		SetShaderValue(BatchedParameters, OutputConfigParam, OutputConfig);
+		SetShaderValue(BatchedParameters, ComponentVertexCountParam, (float)InParams.CurrentMipComponentVertexCount);
 	}
 
 private:
@@ -703,13 +703,13 @@ public:
 	FLandscapeLayersWeightmapMipsPS()
 	{}
 
-	void SetParameters(FRHICommandList& RHICmdList, const FLandscapeLayersWeightmapShaderParameters& InParams)
+	void SetParameters(FRHIBatchedShaderParameters& BatchedParameters, const FLandscapeLayersWeightmapShaderParameters& InParams)
 	{
-		SetTextureParameter(RHICmdList, RHICmdList.GetBoundPixelShader(), ReadTexture1Param, ReadTexture1SamplerParam, TStaticSamplerState<SF_Point, AM_Clamp, AM_Clamp, AM_Clamp>::GetRHI(), InParams.ReadWeightmap1->GetResource()->TextureRHI);
+		SetTextureParameter(BatchedParameters, ReadTexture1Param, ReadTexture1SamplerParam, TStaticSamplerState<SF_Point, AM_Clamp, AM_Clamp, AM_Clamp>::GetRHI(), InParams.ReadWeightmap1->GetResource()->TextureRHI);
 
-		SetShaderValue(RHICmdList, RHICmdList.GetBoundPixelShader(), CurrentMipSizeParam, FVector2f(static_cast<float>(InParams.CurrentMipSize.X), static_cast<float>(InParams.CurrentMipSize.Y)));
-		SetShaderValue(RHICmdList, RHICmdList.GetBoundPixelShader(), ParentMipSizeParam, FVector2f(static_cast<float>(InParams.ParentMipSize.X), static_cast<float>(InParams.ParentMipSize.Y)));
-		SetShaderValue(RHICmdList, RHICmdList.GetBoundPixelShader(), CurrentMipComponentVertexCountParam, (float)InParams.CurrentMipComponentVertexCount);
+		SetShaderValue(BatchedParameters, CurrentMipSizeParam, FVector2f(static_cast<float>(InParams.CurrentMipSize.X), static_cast<float>(InParams.CurrentMipSize.Y)));
+		SetShaderValue(BatchedParameters, ParentMipSizeParam, FVector2f(static_cast<float>(InParams.ParentMipSize.X), static_cast<float>(InParams.ParentMipSize.Y)));
+		SetShaderValue(BatchedParameters, CurrentMipComponentVertexCountParam, (float)InParams.CurrentMipComponentVertexCount);
 	}
 
 private:
@@ -764,9 +764,9 @@ public:
 	FLandscapeCopyTexturePS()
 	{}
 
-	void SetParameters(FRHICommandList& RHICmdList, FRHITexture* TextureRHI)
+	void SetParameters(FRHIBatchedShaderParameters& BatchedParameters, FRHITexture* TextureRHI)
 	{
-		SetTextureParameter(RHICmdList, RHICmdList.GetBoundPixelShader(), ReadTexture1Param, ReadTexture1SamplerParam, TStaticSamplerState<SF_Point, AM_Clamp, AM_Clamp, AM_Clamp>::GetRHI(), TextureRHI);
+		SetTextureParameter(BatchedParameters, ReadTexture1Param, ReadTexture1SamplerParam, TStaticSamplerState<SF_Point, AM_Clamp, AM_Clamp, AM_Clamp>::GetRHI(), TextureRHI);
 	}
 
 private:
@@ -879,20 +879,17 @@ public:
 	FLandscapeLayerWeightmapExtractMaterialLayersCS()
 	{}
 
-	void SetParameters(FRHICommandList& RHICmdList, const FLandscapeLayerWeightmapExtractMaterialLayersComputeShaderParameters& InParams)
+	void SetParameters(FRHIBatchedShaderParameters& BatchedParameters, const FLandscapeLayerWeightmapExtractMaterialLayersComputeShaderParameters& InParams)
 	{
-		SetTextureParameter(RHICmdList, RHICmdList.GetBoundComputeShader(), ComponentWeightmapParam, InParams.ComponentWeightmapResource->TextureRHI);
-		SetUAVParameter(RHICmdList, RHICmdList.GetBoundComputeShader(), AtlasPaintListsParam, InParams.AtlasWeightmapsPerLayer->GetTextureUAV(/*InMipLevel = */0));
-		SetSRVParameter(RHICmdList, RHICmdList.GetBoundComputeShader(), ComponentsDataParam, InParams.ComputeShaderResource->ComponentsDataSRV);
-		SetShaderValue(RHICmdList, RHICmdList.GetBoundComputeShader(), ComponentSizeParam, InParams.ComponentSize);
-		
-		// In case the CS is executed twice in a row, we need a barrier since we want to prevent UAV overlaps:
-		RHICmdList.Transition(FRHITransitionInfo(InParams.AtlasWeightmapsPerLayer->TextureRHI, ERHIAccess::UAVMask, ERHIAccess::UAVMask));
+		SetTextureParameter(BatchedParameters, ComponentWeightmapParam, InParams.ComponentWeightmapResource->TextureRHI);
+		SetUAVParameter(BatchedParameters, AtlasPaintListsParam, InParams.AtlasWeightmapsPerLayer->GetTextureUAV(/*InMipLevel = */0));
+		SetSRVParameter(BatchedParameters, ComponentsDataParam, InParams.ComputeShaderResource->ComponentsDataSRV);
+		SetShaderValue(BatchedParameters, ComponentSizeParam, InParams.ComponentSize);
 	}
 
-	void UnsetParameters(FRHICommandList& RHICmdList)
+	void UnsetParameters(FRHIBatchedShaderParameters& BatchedParameters)
 	{
-		SetUAVParameter(RHICmdList, RHICmdList.GetBoundComputeShader(), AtlasPaintListsParam, nullptr);
+		SetUAVParameter(BatchedParameters, AtlasPaintListsParam, nullptr);
 	}
 
 private:
@@ -918,14 +915,20 @@ public:
 
 		TShaderMapRef<FLandscapeLayerWeightmapExtractMaterialLayersCS> ComputeShader(GetGlobalShaderMap(GMaxRHIFeatureLevel));
 		SetComputePipelineState(InRHICmdList, ComputeShader.GetComputeShader());
-		ComputeShader->SetParameters(InRHICmdList, ShaderParams);
+
+		SetAllShaderParametersCS(InRHICmdList, ComputeShader, ShaderParams);
+
+		// In case the CS is executed twice in a row, we need a barrier since we want to prevent UAV overlaps:
+		InRHICmdList.Transition(FRHITransitionInfo(ShaderParams.AtlasWeightmapsPerLayer->TextureRHI, ERHIAccess::UAVMask, ERHIAccess::UAVMask));
 
 		uint32 ThreadGroupCountX = FMath::CeilToInt((float)ShaderParams.ComponentSize / (float)GLandscapeLayerWeightmapThreadGroupSizeX);
 		uint32 ThreadGroupCountY = FMath::CeilToInt((float)ShaderParams.ComponentSize / (float)GLandscapeLayerWeightmapThreadGroupSizeY);
 		check(ThreadGroupCountX > 0 && ThreadGroupCountY > 0);
 
 		DispatchComputeShader(InRHICmdList, ComputeShader.GetShader(), ThreadGroupCountX, ThreadGroupCountY, ShaderParams.ComputeShaderResource->GetComponentsDataCount());
-		ComputeShader->UnsetParameters(InRHICmdList);
+
+		UnsetAllShaderParametersCS(InRHICmdList, ComputeShader);
+
 		ShaderParams.ComputeShaderResource->ReleaseResource();
 		delete ShaderParams.ComputeShaderResource;
 	}
@@ -1066,19 +1069,19 @@ public:
 	FLandscapeLayerWeightmapPackMaterialLayersCS()
 	{}
 
-	void SetParameters(FRHICommandList& RHICmdList, const FLandscapeLayerWeightmapPackMaterialLayersComputeShaderParameters& InParams)
+	void SetParameters(FRHIBatchedShaderParameters& BatchedParameters, const FLandscapeLayerWeightmapPackMaterialLayersComputeShaderParameters& InParams)
 	{
-		SetUAVParameter(RHICmdList, RHICmdList.GetBoundComputeShader(), ComponentWeightmapParam, InParams.ComponentWeightmapResource->GetTextureUAV(/*InMipLevel = */0));
-		SetTextureParameter(RHICmdList, RHICmdList.GetBoundComputeShader(), AtlasPaintListsParam, InParams.AtlasWeightmapsPerLayer->TextureRHI);
-		SetSRVParameter(RHICmdList, RHICmdList.GetBoundComputeShader(), ComponentsDataParam, InParams.ComputeShaderResource->ComponentsDataSRV);
-		SetShaderValue(RHICmdList, RHICmdList.GetBoundComputeShader(), ComponentSizeParam, InParams.ComponentSize);
-		SetSRVParameter(RHICmdList, RHICmdList.GetBoundComputeShader(), WeightmapWeightBlendModeParam, InParams.ComputeShaderResource->WeightmapWeightBlendModeSRV);
-		SetSRVParameter(RHICmdList, RHICmdList.GetBoundComputeShader(), WeightmapTextureOutputOffsetParam, InParams.ComputeShaderResource->WeightmapTextureOutputOffsetSRV);
+		SetUAVParameter(BatchedParameters, ComponentWeightmapParam, InParams.ComponentWeightmapResource->GetTextureUAV(/*InMipLevel = */0));
+		SetTextureParameter(BatchedParameters, AtlasPaintListsParam, InParams.AtlasWeightmapsPerLayer->TextureRHI);
+		SetSRVParameter(BatchedParameters, ComponentsDataParam, InParams.ComputeShaderResource->ComponentsDataSRV);
+		SetShaderValue(BatchedParameters, ComponentSizeParam, InParams.ComponentSize);
+		SetSRVParameter(BatchedParameters, WeightmapWeightBlendModeParam, InParams.ComputeShaderResource->WeightmapWeightBlendModeSRV);
+		SetSRVParameter(BatchedParameters, WeightmapTextureOutputOffsetParam, InParams.ComputeShaderResource->WeightmapTextureOutputOffsetSRV);
 	}
 
-	void UnsetParameters(FRHICommandList& RHICmdList)
+	void UnsetParameters(FRHIBatchedShaderParameters& BatchedParameters)
 	{
-		SetUAVParameter(RHICmdList, RHICmdList.GetBoundComputeShader(), ComponentWeightmapParam, nullptr);
+		SetUAVParameter(BatchedParameters, ComponentWeightmapParam, nullptr);
 	}
 
 private:
@@ -1106,14 +1109,17 @@ public:
 
 		TShaderMapRef<FLandscapeLayerWeightmapPackMaterialLayersCS> ComputeShader(GetGlobalShaderMap(GMaxRHIFeatureLevel));
 		SetComputePipelineState(InRHICmdList, ComputeShader.GetComputeShader());
-		ComputeShader->SetParameters(InRHICmdList, ShaderParams);
+
+		SetAllShaderParametersCS(InRHICmdList, ComputeShader, ShaderParams);
 
 		uint32 ThreadGroupCountX = FMath::CeilToInt((float)ShaderParams.ComponentSize / (float)GLandscapeLayerWeightmapThreadGroupSizeX);
 		uint32 ThreadGroupCountY = FMath::CeilToInt((float)ShaderParams.ComponentSize / (float)GLandscapeLayerWeightmapThreadGroupSizeY);
 		check(ThreadGroupCountX > 0 && ThreadGroupCountY > 0);
 
 		DispatchComputeShader(InRHICmdList, ComputeShader.GetShader(), ThreadGroupCountX, ThreadGroupCountY, ShaderParams.ComputeShaderResource->GetComponentsDataCount());
-		ComputeShader->UnsetParameters(InRHICmdList);
+
+		UnsetAllShaderParametersCS(InRHICmdList, ComputeShader);
+
 		ShaderParams.ComputeShaderResource->ReleaseResource();
 		delete ShaderParams.ComputeShaderResource;
 	}
@@ -1641,8 +1647,8 @@ public:
 			SetGraphicsPipelineState(InRHICmdList, GraphicsPSOInit, 0);
 
 			// Set shader params
-			VertexShader->SetParameters(InRHICmdList, FMatrix44f(ProjectionMatrix));		// LWC_TODo: Precision loss?
-			PixelShader->SetParameters(InRHICmdList, ShaderParams);
+			SetAllShaderParametersVS(InRHICmdList, VertexShader, FMatrix44f(ProjectionMatrix));		// LWC_TODo: Precision loss?
+			SetAllShaderParametersPS(InRHICmdList, PixelShader, ShaderParams);
 		}
 		else
 		{
@@ -1659,8 +1665,8 @@ public:
 			SetGraphicsPipelineState(InRHICmdList, GraphicsPSOInit, 0);
 
 			// Set shader params
-			VertexShader->SetParameters(InRHICmdList, FMatrix44f(ProjectionMatrix));		// LWC_TODo: Precision loss?
-			PixelShader->SetParameters(InRHICmdList, ShaderParams);
+			SetAllShaderParametersVS(InRHICmdList, VertexShader, FMatrix44f(ProjectionMatrix));		// LWC_TODo: Precision loss?
+			SetAllShaderParametersPS(InRHICmdList, PixelShader, ShaderParams);
 		}
 
 		InRHICmdList.SetScissorRect(false, 0, 0, 0, 0);
@@ -2613,7 +2619,7 @@ void ALandscape::CopyTexturePS(const FString& InSourceDebugName, FTextureResourc
 		GraphicsPSOInit.BoundShaderState.PixelShaderRHI = PixelShader.GetPixelShader();
 		SetGraphicsPipelineState(RHICmdList, GraphicsPSOInit, 0);
 
-		PixelShader->SetParameters(RHICmdList, InSourceResource->TextureRHI);
+		SetAllShaderParametersPS(RHICmdList, PixelShader, InSourceResource->TextureRHI);
 
 		RHICmdList.SetViewport(0.0f, 0.0f, 0.0f, (float)InDestResource->GetSizeX(), (float)InDestResource->GetSizeY(), 1.0f);
 		RHICmdList.DrawIndexedPrimitive(GTwoTrianglesIndexBuffer.IndexBufferRHI, 0, 0, 4, 0, 2, 1);
diff --git a/Engine/Source/Runtime/RenderCore/Private/BuiltInRayTracingShaders.cpp b/Engine/Source/Runtime/RenderCore/Private/BuiltInRayTracingShaders.cpp
--- a/Engine/Source/Runtime/RenderCore/Private/BuiltInRayTracingShaders.cpp
+++ b/Engine/Source/Runtime/RenderCore/Private/BuiltInRayTracingShaders.cpp
@@ -56,18 +56,24 @@ void FRayTracingDispatchDescCS::Dispatch(FRHICommandList& RHICmdList,
 	FUintVector4 DispatchDescData[DispatchDescMaxSizeUint4s] = {};
 	FMemory::Memcpy(DispatchDescData, DispatchDescInput, DispatchDescSize);
 
-	SetShaderValueArray(RHICmdList, ShaderRHI, ComputeShader->DispatchDescInputParam, DispatchDescData, DispatchDescMaxSizeUint4s);
-	SetShaderValue(RHICmdList, ShaderRHI, ComputeShader->DispatchDescSizeDwordsParam, DispatchDescSizeDwords);
-	SetShaderValue(RHICmdList, ShaderRHI, ComputeShader->DispatchDescDimensionsOffsetDwordsParam, DispatchDescDimensionsOffsetDwords);
-	SetShaderValue(RHICmdList, ShaderRHI, ComputeShader->DimensionsBufferOffsetDwordsParam, DimensionsBufferOffsetDwords);
+	FRHIBatchedShaderParameters BatchedParameters;
 
-	SetSRVParameter(RHICmdList, ShaderRHI, ComputeShader->DispatchDimensionsParam, DispatchDimensionsSRV);
-	SetUAVParameter(RHICmdList, ShaderRHI, ComputeShader->DispatchDescOutputParam, DispatchDescOutputUAV);
+	SetShaderValueArray(BatchedParameters, ComputeShader->DispatchDescInputParam, DispatchDescData, DispatchDescMaxSizeUint4s);
+	SetShaderValue(BatchedParameters, ComputeShader->DispatchDescSizeDwordsParam, DispatchDescSizeDwords);
+	SetShaderValue(BatchedParameters, ComputeShader->DispatchDescDimensionsOffsetDwordsParam, DispatchDescDimensionsOffsetDwords);
+	SetShaderValue(BatchedParameters, ComputeShader->DimensionsBufferOffsetDwordsParam, DimensionsBufferOffsetDwords);
+
+	SetSRVParameter(BatchedParameters, ComputeShader->DispatchDimensionsParam, DispatchDimensionsSRV);
+	SetUAVParameter(BatchedParameters, ComputeShader->DispatchDescOutputParam, DispatchDescOutputUAV);
+	RHICmdList.SetBatchedShaderParameters(ComputeShader.GetComputeShader(), BatchedParameters);
 
 	RHICmdList.DispatchComputeShader(1, 1, 1);
 
-	SetSRVParameter(RHICmdList, ShaderRHI, ComputeShader->DispatchDimensionsParam, nullptr);
-	SetUAVParameter(RHICmdList, ShaderRHI, ComputeShader->DispatchDescOutputParam, nullptr);
+	BatchedParameters.Reset();
+
+	SetSRVParameter(BatchedParameters, ComputeShader->DispatchDimensionsParam, nullptr);
+	SetUAVParameter(BatchedParameters, ComputeShader->DispatchDescOutputParam, nullptr);
+	RHICmdList.SetBatchedShaderParameters(ComputeShader.GetComputeShader(), BatchedParameters);
 }
 
 #endif // RHI_RAYTRACING
diff --git a/Engine/Source/Runtime/RenderCore/Public/ShaderParameterUtils.h b/Engine/Source/Runtime/RenderCore/Public/ShaderParameterUtils.h
--- a/Engine/Source/Runtime/RenderCore/Public/ShaderParameterUtils.h
+++ b/Engine/Source/Runtime/RenderCore/Public/ShaderParameterUtils.h
@@ -13,8 +13,234 @@
 #include "ShaderCore.h"
 #include "Misc/App.h"
 
+class FShaderMapPointerTable;
+
 template<typename TBufferStruct> class TUniformBuffer;
 template<typename TBufferStruct> class TUniformBufferRef;
+template<typename ShaderType, typename PointerTableType> class TShaderRefBase;
+template<typename ShaderType> using TShaderRef = TShaderRefBase<ShaderType, FShaderMapPointerTable>;
+
+template<class ParameterType>
+void SetShaderValue(
+	FRHIBatchedShaderParameters& BatchedParameters
+	, const FShaderParameter& Parameter
+	, const ParameterType& Value
+	, uint32 ElementIndex = 0
+)
+{
+	// This will trigger if the parameter was not serialized
+	checkSlow(Parameter.IsInitialized());
+
+	static_assert(!TIsPointer<ParameterType>::Value, "Passing by value is not valid.");
+
+	const uint32 AlignedTypeSize = Align<uint32>(sizeof(ParameterType), SHADER_PARAMETER_ARRAY_ELEMENT_ALIGNMENT);
+	const uint32 ElementByteOffset = ElementIndex * AlignedTypeSize;
+	const int32 NumBytesToSet = FMath::Min<int32>(sizeof(ParameterType), static_cast<int32>(Parameter.GetNumBytes()) - ElementByteOffset);
+
+	if (NumBytesToSet > 0)
+	{
+		BatchedParameters.SetShaderParameter(
+			Parameter.GetBufferIndex(),
+			Parameter.GetBaseIndex() + ElementByteOffset,
+			(uint32)NumBytesToSet,
+			&Value);
+	}
+}
+
+template<class ParameterType>
+void SetShaderValueArray(
+	FRHIBatchedShaderParameters& BatchedParameters
+	, const FShaderParameter& Parameter
+	, const ParameterType* Values
+	, uint32 NumElements
+	, uint32 ElementIndex = 0
+)
+{
+	const uint32 AlignedTypeSize = Align<uint32>(sizeof(ParameterType), SHADER_PARAMETER_ARRAY_ELEMENT_ALIGNMENT);
+	const uint32 ElementByteOffset = ElementIndex * AlignedTypeSize;
+	const int32 NumBytesToSet = FMath::Min<int32>(NumElements * AlignedTypeSize, Parameter.GetNumBytes() - ElementByteOffset);
+
+	// This will trigger if the parameter was not serialized
+	checkSlow(Parameter.IsInitialized());
+
+	if (NumBytesToSet > 0)
+	{
+		BatchedParameters.SetShaderParameter(
+			Parameter.GetBufferIndex(),
+			Parameter.GetBaseIndex() + ElementByteOffset,
+			(uint32)NumBytesToSet,
+			Values
+		);
+	}
+}
+
+inline void SetTextureParameter(FRHIBatchedShaderParameters& BatchedParameters, const FShaderResourceParameter& Parameter, FRHITexture* TextureRHI)
+{
+	if (Parameter.IsBound())
+	{
+#if PLATFORM_SUPPORTS_BINDLESS_RENDERING
+		if (Parameter.GetType() == EShaderParameterType::BindlessResourceIndex)
+		{
+			BatchedParameters.SetBindlessTexture(Parameter.GetBaseIndex(), TextureRHI);
+		}
+		else
+#endif
+		{
+			BatchedParameters.SetShaderTexture(Parameter.GetBaseIndex(), TextureRHI);
+		}
+	}
+}
+
+inline void SetSamplerParameter(FRHIBatchedShaderParameters& BatchedParameters, const FShaderResourceParameter& Parameter, FRHISamplerState* SamplerStateRHI)
+{
+	if (Parameter.IsBound())
+	{
+#if PLATFORM_SUPPORTS_BINDLESS_RENDERING
+		if (Parameter.GetType() == EShaderParameterType::BindlessSamplerIndex)
+		{
+			BatchedParameters.SetBindlessSampler(Parameter.GetBaseIndex(), SamplerStateRHI);
+		}
+		else
+#endif
+		{
+			BatchedParameters.SetShaderSampler(Parameter.GetBaseIndex(), SamplerStateRHI);
+		}
+	}
+}
+
+inline void SetTextureParameter(
+	FRHIBatchedShaderParameters& BatchedParameters,
+	const FShaderResourceParameter& TextureParameter,
+	const FShaderResourceParameter& SamplerParameter,
+	FRHISamplerState* SamplerStateRHI,
+	FRHITexture* TextureRHI
+)
+{
+	SetTextureParameter(BatchedParameters, TextureParameter, TextureRHI);
+	SetSamplerParameter(BatchedParameters, SamplerParameter, SamplerStateRHI);
+}
+
+inline void SetTextureParameter(
+	FRHIBatchedShaderParameters& BatchedParameters,
+	const FShaderResourceParameter& TextureParameter,
+	const FShaderResourceParameter& SamplerParameter,
+	const FTexture* Texture
+)
+{
+	if (TextureParameter.IsBound())
+	{
+		Texture->LastRenderTime = FApp::GetCurrentTime();
+	}
+
+	SetTextureParameter(BatchedParameters, TextureParameter, Texture->TextureRHI);
+	SetSamplerParameter(BatchedParameters, SamplerParameter, Texture->SamplerStateRHI);
+}
+
+inline void SetSRVParameter(FRHIBatchedShaderParameters& BatchedParameters, const FShaderResourceParameter& Parameter, FRHIShaderResourceView* SRV)
+{
+	if (Parameter.IsBound())
+	{
+#if PLATFORM_SUPPORTS_BINDLESS_RENDERING
+		if (Parameter.GetType() == EShaderParameterType::BindlessResourceIndex)
+		{
+			BatchedParameters.SetBindlessResourceView(Parameter.GetBaseIndex(), SRV);
+		}
+		else
+#endif
+		{
+			BatchedParameters.SetShaderResourceViewParameter(Parameter.GetBaseIndex(), SRV);
+		}
+	}
+}
+
+inline void SetUAVParameter(FRHIBatchedShaderParameters& BatchedParameters, const FShaderResourceParameter& Parameter, FRHIUnorderedAccessView* UAV)
+{
+	if (Parameter.IsBound())
+	{
+#if PLATFORM_SUPPORTS_BINDLESS_RENDERING
+		if (Parameter.GetType() == EShaderParameterType::BindlessResourceIndex)
+		{
+			BatchedParameters.SetBindlessUAV(Parameter.GetBaseIndex(), UAV);
+		}
+		else
+#endif
+		{
+			BatchedParameters.SetUAVParameter(Parameter.GetBaseIndex(), UAV);
+		}
+	}
+}
+
+/// Utility to set all parameters for a Vertex shader. Requires the shader type to implement SetParameters(FRHIBatchedShaderParameters& BatchedParameters, ...)
+template<typename TCmdList, typename TShaderType, typename... TArguments>
+inline void SetAllShaderParametersVS(TCmdList& RHICmdList, const TShaderRef<TShaderType>& InShader, TArguments&&... InArguments)
+{
+	FRHIBatchedShaderParameters BatchedParameters;
+	InShader->SetParameters(BatchedParameters, Forward<TArguments>(InArguments)...);
+	RHICmdList.SetBatchedShaderParameters(InShader.GetVertexShader(), BatchedParameters);
+}
+
+/// Utility to set all parameters for a Mesh shader. Requires the shader type to implement SetParameters(FRHIBatchedShaderParameters& BatchedParameters, ...)
+template<typename TCmdList, typename TShaderType, typename... TArguments>
+inline void SetAllShaderParametersMS(TCmdList& RHICmdList, const TShaderRef<TShaderType>& InShader, TArguments&&... InArguments)
+{
+	FRHIBatchedShaderParameters BatchedParameters;
+	InShader->SetParameters(BatchedParameters, Forward<TArguments>(InArguments)...);
+	RHICmdList.SetBatchedShaderParameters(InShader.GetMeshShader(), BatchedParameters);
+}
+
+/// Utility to set all parameters for an Amplification shader. Requires the shader type to implement SetParameters(FRHIBatchedShaderParameters& BatchedParameters, ...)
+template<typename TCmdList, typename TShaderType, typename... TArguments>
+inline void SetAllShaderParametersAS(TCmdList& RHICmdList, const TShaderRef<TShaderType>& InShader, TArguments&&... InArguments)
+{
+	FRHIBatchedShaderParameters BatchedParameters;
+	InShader->SetParameters(BatchedParameters, Forward<TArguments>(InArguments)...);
+	RHICmdList.SetBatchedShaderParameters(InShader.GetAmplificationShader(), BatchedParameters);
+}
+
+/// Utility to set all parameters for a Pixel shader. Requires the shader type to implement SetParameters(FRHIBatchedShaderParameters& BatchedParameters, ...)
+template<typename TCmdList, typename TShaderType, typename... TArguments>
+inline void SetAllShaderParametersPS(TCmdList& RHICmdList, const TShaderRef<TShaderType>& InShader, TArguments&&... InArguments)
+{
+	FRHIBatchedShaderParameters BatchedParameters;
+	InShader->SetParameters(BatchedParameters, Forward<TArguments>(InArguments)...);
+	RHICmdList.SetBatchedShaderParameters(InShader.GetPixelShader(), BatchedParameters);
+}
+
+/// Utility to set all parameters for a Geometry shader. Requires the shader type to implement SetParameters(FRHIBatchedShaderParameters& BatchedParameters, ...)
+template<typename TCmdList, typename TShaderType, typename... TArguments>
+inline void SetAllShaderParametersGS(TCmdList& RHICmdList, const TShaderRef<TShaderType>& InShader, TArguments&&... InArguments)
+{
+	FRHIBatchedShaderParameters BatchedParameters;
+	InShader->SetParameters(BatchedParameters, Forward<TArguments>(InArguments)...);
+	RHICmdList.SetBatchedShaderParameters(InShader.GetGeometryShader(), BatchedParameters);
+}
+
+/// Utility to set all parameters for a Compute shader. Requires the shader type to implement SetParameters(FRHIBatchedShaderParameters& BatchedParameters, ...)
+template<typename TCmdList, typename TShaderType, typename... TArguments>
+inline void SetAllShaderParametersCS(TCmdList& RHICmdList, const TShaderRef<TShaderType>& InShader, TArguments&&... InArguments)
+{
+	FRHIBatchedShaderParameters BatchedParameters;
+	InShader->SetParameters(BatchedParameters, Forward<TArguments>(InArguments)...);
+	RHICmdList.SetBatchedShaderParameters(InShader.GetComputeShader(), BatchedParameters);
+}
+
+/// Utility to unset all parameters for a Pixel shader. Requires the shader type to implement UnsetParameters(FRHIBatchedShaderParameters& BatchedParameters)
+template<typename TCmdList, typename TShaderType>
+inline void UnsetAllShaderParametersPS(TCmdList& RHICmdList, const TShaderRef<TShaderType>& InShader)
+{
+	FRHIBatchedShaderParameters BatchedParameters;
+	InShader->UnsetParameters(BatchedParameters);
+	RHICmdList.SetBatchedShaderParameters(InShader.GetPixelShader(), BatchedParameters);
+}
+
+/// Utility to unset all parameters for a Compute shader. Requires the shader type to implement UnsetParameters(FRHIBatchedShaderParameters& BatchedParameters)
+template<typename TCmdList, typename TShaderType>
+inline void UnsetAllShaderParametersCS(TCmdList& RHICmdList, const TShaderRef<TShaderType>& InShader)
+{
+	FRHIBatchedShaderParameters BatchedParameters;
+	InShader->UnsetParameters(BatchedParameters);
+	RHICmdList.SetBatchedShaderParameters(InShader.GetComputeShader(), BatchedParameters);
+}
 
 /**
  * Sets the value of a  shader parameter.  Template'd on shader type
diff --git a/Engine/Source/Runtime/Renderer/Private/PostProcess/PostProcessCombineLUTs.cpp b/Engine/Source/Runtime/Renderer/Private/PostProcess/PostProcessCombineLUTs.cpp
--- a/Engine/Source/Runtime/Renderer/Private/PostProcess/PostProcessCombineLUTs.cpp
+++ b/Engine/Source/Runtime/Renderer/Private/PostProcess/PostProcessCombineLUTs.cpp
@@ -514,8 +514,7 @@ FRDGTextureRef AddCombineLUTPass(FRDGBuilder& GraphBuilder, const FViewInfo& Vie
 				GraphicsPSOInit.BoundShaderState.PixelShaderRHI = PixelShader.GetPixelShader();
 				SetGraphicsPipelineState(RHICmdList, GraphicsPSOInit, 0);
 
-				VertexShader->SetParameters(RHICmdList, VolumeBounds, FIntVector(VolumeBounds.MaxX - VolumeBounds.MinX));
-
+				SetAllShaderParametersVS(RHICmdList, VertexShader, VolumeBounds, FIntVector(VolumeBounds.MaxX - VolumeBounds.MinX));
 				SetShaderParameters(RHICmdList, PixelShader, PixelShader.GetPixelShader(), *PassParameters);
 
 				RasterizeToVolumeTexture(RHICmdList, VolumeBounds);
diff --git a/Engine/Source/Runtime/Renderer/Private/TranslucentLighting.cpp b/Engine/Source/Runtime/Renderer/Private/TranslucentLighting.cpp
--- a/Engine/Source/Runtime/Renderer/Private/TranslucentLighting.cpp
+++ b/Engine/Source/Runtime/Renderer/Private/TranslucentLighting.cpp
@@ -963,10 +963,10 @@ void InjectTranslucencyLightingVolumeAmbientCubemap(
 					GraphicsPSOInit.PrimitiveType = PT_TriangleStrip;
 					SetGraphicsPipelineState(RHICmdList, GraphicsPSOInit, 0);
 
-					VertexShader->SetParameters(RHICmdList, VolumeBounds, FIntVector(TranslucencyLightingVolumeDim));
+					SetAllShaderParametersVS(RHICmdList, VertexShader, VolumeBounds, FIntVector(TranslucencyLightingVolumeDim));
 					if (GeometryShader.IsValid())
 					{
-						GeometryShader->SetParameters(RHICmdList, VolumeBounds.MinZ);
+						SetAllShaderParametersGS(RHICmdList, GeometryShader, VolumeBounds.MinZ);
 					}
 					SetShaderParameters(RHICmdList, PixelShader, PixelShader.GetPixelShader(), *PassParameters);
 				});
@@ -1217,10 +1217,10 @@ void InjectTranslucencyLightingVolume(
 
 					const int32 TranslucencyLightingVolumeDim = GetTranslucencyLightingVolumeDim();
 
-					VertexShader->SetParameters(RHICmdList, VolumeBounds, FIntVector(TranslucencyLightingVolumeDim));
+					SetAllShaderParametersVS(RHICmdList, VertexShader, VolumeBounds, FIntVector(TranslucencyLightingVolumeDim));
 					if (GeometryShader.IsValid())
 					{
-						GeometryShader->SetParameters(RHICmdList, VolumeBounds.MinZ);
+						SetAllShaderParametersGS(RHICmdList, GeometryShader, VolumeBounds.MinZ);
 					}
 
 					PixelShader->SetParameters(RHICmdList, View, InjectionData.LightFunctionMaterialProxy);
@@ -1344,10 +1344,10 @@ void InjectSimpleTranslucencyLightingVolumeArray(
 							GraphicsPSOInit.BoundShaderState.PixelShaderRHI = PixelShader.GetPixelShader();
 							SetGraphicsPipelineState(RHICmdList, GraphicsPSOInit, 0);
 
-							VertexShader->SetParameters(RHICmdList, VolumeBounds, FIntVector(TranslucencyLightingVolumeDim));
+							SetAllShaderParametersVS(RHICmdList, VertexShader, VolumeBounds, FIntVector(TranslucencyLightingVolumeDim));
 							if (GeometryShader.IsValid())
 							{
-								GeometryShader->SetParameters(RHICmdList, VolumeBounds.MinZ);
+								SetAllShaderParametersGS(RHICmdList, GeometryShader, VolumeBounds.MinZ);
 							}
 							SetShaderParameters(RHICmdList, PixelShader, PixelShader.GetPixelShader(), *PassParameters);
 							RasterizeToVolumeTexture(RHICmdList, VolumeBounds);
@@ -1428,10 +1428,10 @@ void FilterTranslucencyLightingVolume(
 				GraphicsPSOInit.PrimitiveType = PT_TriangleStrip;
 				SetGraphicsPipelineState(RHICmdList, GraphicsPSOInit, 0);
 
-				VertexShader->SetParameters(RHICmdList, VolumeBounds, FIntVector(TranslucencyLightingVolumeDim));
+				SetAllShaderParametersVS(RHICmdList, VertexShader, VolumeBounds, FIntVector(TranslucencyLightingVolumeDim));
 				if (GeometryShader.IsValid())
 				{
-					GeometryShader->SetParameters(RHICmdList, VolumeBounds.MinZ);
+					SetAllShaderParametersGS(RHICmdList, GeometryShader, VolumeBounds.MinZ);
 				}
 				SetShaderParameters(RHICmdList, PixelShader, PixelShader.GetPixelShader(), *PassParameters);
 				RasterizeToVolumeTexture(RHICmdList, VolumeBounds);
diff --git a/Engine/Source/Runtime/Renderer/Private/VolumetricFog.cpp b/Engine/Source/Runtime/Renderer/Private/VolumetricFog.cpp
--- a/Engine/Source/Runtime/Renderer/Private/VolumetricFog.cpp
+++ b/Engine/Source/Runtime/Renderer/Private/VolumetricFog.cpp
@@ -619,7 +619,7 @@ void FDeferredShadingSceneRenderer::RenderLocalLightsForVolumetricFog(
 
 						if (GeometryShader.IsValid())
 						{
-							GeometryShader->SetParameters(RHICmdList, VolumeZBounds.X);
+							SetAllShaderParametersGS(RHICmdList, GeometryShader, VolumeZBounds.X);
 						}
 
 						RHICmdList.SetStreamSource(0, GCircleRasterizeVertexBuffer.VertexBufferRHI, 0);
diff --git a/Engine/Source/Runtime/Renderer/Public/VolumeRendering.h b/Engine/Source/Runtime/Renderer/Public/VolumeRendering.h
--- a/Engine/Source/Runtime/Renderer/Public/VolumeRendering.h
+++ b/Engine/Source/Runtime/Renderer/Public/VolumeRendering.h
@@ -54,17 +54,25 @@ public:
 	static bool ShouldCompilePermutation(const FGlobalShaderPermutationParameters& Parameters);
 	static void ModifyCompilationEnvironment(const FGlobalShaderPermutationParameters& Parameters, FShaderCompilerEnvironment& OutEnvironment);
 
-	template <typename TRHICommandList>
-	void SetParameters(TRHICommandList& RHICmdList, const FVolumeBounds& VolumeBounds, FIntVector VolumeResolution)
+	void SetParameters(FRHIBatchedShaderParameters& BatchedParameters, const FVolumeBounds& VolumeBounds, const FIntVector& VolumeResolution)
 	{
 		const float InvVolumeResolutionX = 1.0f / VolumeResolution.X;
 		const float InvVolumeResolutionY = 1.0f / VolumeResolution.Y;
-		SetShaderValue(RHICmdList, RHICmdList.GetBoundVertexShader(), UVScaleBias, FVector4f(
+		SetShaderValue(BatchedParameters, UVScaleBias, FVector4f(
 			(VolumeBounds.MaxX - VolumeBounds.MinX) * InvVolumeResolutionX,
 			(VolumeBounds.MaxY - VolumeBounds.MinY) * InvVolumeResolutionY,
 			VolumeBounds.MinX * InvVolumeResolutionX,
 			VolumeBounds.MinY * InvVolumeResolutionY));
-		SetShaderValue(RHICmdList, RHICmdList.GetBoundVertexShader(), MinZ, VolumeBounds.MinZ);
+		SetShaderValue(BatchedParameters, MinZ, VolumeBounds.MinZ);
+	}
+
+	template <typename TRHICommandList>
+	UE_DEPRECATED(5.3, "SetParameters with FRHIBatchedShaderParameters should be used.")
+	void SetParameters(TRHICommandList& RHICmdList, const FVolumeBounds& VolumeBounds, const FIntVector& VolumeResolution)
+	{
+		FRHIBatchedShaderParameters BatchedParameters;
+		SetParameters(BatchedParameters, VolumeBounds, VolumeResolution);
+		RHICmdList.SetBatchedShaderParameters(RHICmdList.GetBoundVertexShader(), BatchedParameters);
 	}
 
 private:
@@ -82,10 +90,18 @@ public:
 
 	static bool ShouldCompilePermutation(const FGlobalShaderPermutationParameters& Parameters);
 
+	void SetParameters(FRHIBatchedShaderParameters& BatchedParameters, int32 MinZValue)
+	{
+		SetShaderValue(BatchedParameters, MinZ, MinZValue);
+	}
+
 	template <typename TRHICommandList>
+	UE_DEPRECATED(5.3, "SetParameters with FRHIBatchedShaderParameters should be used.")
 	void SetParameters(TRHICommandList& RHICmdList, int32 MinZValue)
 	{
-		SetShaderValue(RHICmdList, RHICmdList.GetBoundGeometryShader(), MinZ, MinZValue);
+		FRHIBatchedShaderParameters BatchedParameters;
+		SetParameters(BatchedParameters, MinZValue);
+		RHICmdList.SetBatchedShaderParameters(RHICmdList.GetBoundGeometryShader(), BatchedParameters);
 	}
 
 private:
diff --git a/Engine/Source/Runtime/SlateRHIRenderer/Private/SlateRHIRenderer.cpp b/Engine/Source/Runtime/SlateRHIRenderer/Private/SlateRHIRenderer.cpp
--- a/Engine/Source/Runtime/SlateRHIRenderer/Private/SlateRHIRenderer.cpp
+++ b/Engine/Source/Runtime/SlateRHIRenderer/Private/SlateRHIRenderer.cpp
@@ -626,7 +626,7 @@ public:
 	}
 	FCompositeLUTGenerationPS() {}
 
-	void SetParameters(FRHICommandList& RHICmdList, EDisplayOutputFormat DisplayOutputFormat, EDisplayColorGamut DisplayColorGamut, float DisplayMaxLuminance)
+	void SetParameters(FRHIBatchedShaderParameters& BatchedParameters, EDisplayOutputFormat DisplayOutputFormat, EDisplayColorGamut DisplayColorGamut, float DisplayMaxLuminance)
 	{
 		static const auto CVarOutputGamma = IConsoleManager::Get().FindTConsoleVariableDataFloat(TEXT("r.TonemapperGamma"));
 
@@ -646,20 +646,20 @@ public:
 			OutputDeviceValue = FMath::Max(OutputDeviceValue, (int32)EDisplayOutputFormat::SDR_ExplicitGammaMapping);
 		}
 
-		SetShaderValue(RHICmdList, RHICmdList.GetBoundPixelShader(), OutputDevice, OutputDeviceValue);
-		SetShaderValue(RHICmdList, RHICmdList.GetBoundPixelShader(), OutputGamut, OutputGamutValue);
-		SetShaderValue(RHICmdList, RHICmdList.GetBoundPixelShader(), OutputMaxLuminance, DisplayMaxLuminance);
+		SetShaderValue(BatchedParameters, OutputDevice, OutputDeviceValue);
+		SetShaderValue(BatchedParameters, OutputGamut, OutputGamutValue);
+		SetShaderValue(BatchedParameters, OutputMaxLuminance, DisplayMaxLuminance);
 
 		FACESTonemapParams TmpInternalACESParams;
 		GetACESTonemapParameters(TmpInternalACESParams);
 
-		SetShaderValue(RHICmdList, RHICmdList.GetBoundPixelShader(), ACESMinMaxData, TmpInternalACESParams.ACESMinMaxData);
-		SetShaderValue(RHICmdList, RHICmdList.GetBoundPixelShader(), ACESMidData, TmpInternalACESParams.ACESMidData);
-		SetShaderValue(RHICmdList, RHICmdList.GetBoundPixelShader(), ACESCoefsLow_0, TmpInternalACESParams.ACESCoefsLow_0);
-		SetShaderValue(RHICmdList, RHICmdList.GetBoundPixelShader(), ACESCoefsHigh_0, TmpInternalACESParams.ACESCoefsHigh_0);
-		SetShaderValue(RHICmdList, RHICmdList.GetBoundPixelShader(), ACESCoefsLow_4, TmpInternalACESParams.ACESCoefsLow_4);
-		SetShaderValue(RHICmdList, RHICmdList.GetBoundPixelShader(), ACESCoefsHigh_4, TmpInternalACESParams.ACESCoefsHigh_4);
-		SetShaderValue(RHICmdList, RHICmdList.GetBoundPixelShader(), ACESSceneColorMultiplier, TmpInternalACESParams.ACESSceneColorMultiplier);
+		SetShaderValue(BatchedParameters, ACESMinMaxData, TmpInternalACESParams.ACESMinMaxData);
+		SetShaderValue(BatchedParameters, ACESMidData, TmpInternalACESParams.ACESMidData);
+		SetShaderValue(BatchedParameters, ACESCoefsLow_0, TmpInternalACESParams.ACESCoefsLow_0);
+		SetShaderValue(BatchedParameters, ACESCoefsHigh_0, TmpInternalACESParams.ACESCoefsHigh_0);
+		SetShaderValue(BatchedParameters, ACESCoefsLow_4, TmpInternalACESParams.ACESCoefsLow_4);
+		SetShaderValue(BatchedParameters, ACESCoefsHigh_4, TmpInternalACESParams.ACESCoefsHigh_4);
+		SetShaderValue(BatchedParameters, ACESSceneColorMultiplier, TmpInternalACESParams.ACESSceneColorMultiplier);
 
 	}
 
@@ -717,30 +717,28 @@ public:
 	}
 	FCompositeShaderBase() = default;
 
-	template<class TRHIShader>
-	void SetParametersBase(FRHICommandList& RHICmdList, TRHIShader* BoundShader, FRHITexture* UITextureRHI, FRHITexture* UITextureWriteMaskRHI, FRHITexture* ColorSpaceLUTRHI)
+	void SetParametersBase(FRHIBatchedShaderParameters& BatchedParameters, FRHITexture* UITextureRHI, FRHITexture* UITextureWriteMaskRHI, FRHITexture* ColorSpaceLUTRHI)
 	{
 		static const auto CVarOutputDevice = IConsoleManager::Get().FindTConsoleVariableDataInt(TEXT("r.HDR.Display.OutputDevice"));
 
-		SetTextureParameter(RHICmdList, BoundShader, UITexture, UISampler, TStaticSamplerState<SF_Point>::GetRHI(), UITextureRHI);
-		SetTextureParameter(RHICmdList, BoundShader, ColorSpaceLUT, ColorSpaceLUTSampler, TStaticSamplerState<SF_Bilinear, AM_Clamp, AM_Clamp, AM_Clamp>::GetRHI(), ColorSpaceLUTRHI);
-		SetShaderValue(RHICmdList, BoundShader, UILevel, CVarUILevel.GetValueOnRenderThread());
-		SetShaderValue(RHICmdList, BoundShader, OutputDevice, CVarOutputDevice->GetValueOnRenderThread());
-		SetShaderValue(RHICmdList, BoundShader, UILuminance, CVarHDRUILuminance.GetValueOnRenderThread());
+		SetTextureParameter(BatchedParameters, UITexture, UISampler, TStaticSamplerState<SF_Point>::GetRHI(), UITextureRHI);
+		SetTextureParameter(BatchedParameters, ColorSpaceLUT, ColorSpaceLUTSampler, TStaticSamplerState<SF_Bilinear, AM_Clamp, AM_Clamp, AM_Clamp>::GetRHI(), ColorSpaceLUTRHI);
+		SetShaderValue(BatchedParameters, UILevel, CVarUILevel.GetValueOnRenderThread());
+		SetShaderValue(BatchedParameters, OutputDevice, CVarOutputDevice->GetValueOnRenderThread());
+		SetShaderValue(BatchedParameters, UILuminance, CVarHDRUILuminance.GetValueOnRenderThread());
 
 		if (UITextureWriteMaskRHI != nullptr)
 		{
-			SetTextureParameter(RHICmdList, BoundShader, UIWriteMaskTexture, UITextureWriteMaskRHI);
+			SetTextureParameter(BatchedParameters, UIWriteMaskTexture, UITextureWriteMaskRHI);
 		}
 	}
 
-	template<class TRHIShader>
-	void SetColorDeficiencyParamsBase(FRHICommandList& RHICmdList, TRHIShader* BoundShader, bool bCorrect, EColorVisionDeficiency DeficiencyType, int32 Severity, bool bShowCorrectionWithDeficiency)
+	void SetColorDeficiencyParamsBase(FRHIBatchedShaderParameters& BatchedParameters, bool bCorrect, EColorVisionDeficiency DeficiencyType, int32 Severity, bool bShowCorrectionWithDeficiency)
 	{
-		SetShaderValue(RHICmdList, BoundShader, ColorVisionDeficiencyType, (float)DeficiencyType);
-		SetShaderValue(RHICmdList, BoundShader, ColorVisionDeficiencySeverity, (float)Severity);
-		SetShaderValue(RHICmdList, BoundShader, bCorrectDeficiency, bCorrect ? 1.0f : 0.0f);
-		SetShaderValue(RHICmdList, BoundShader, bSimulateCorrectionWithDeficiency, bShowCorrectionWithDeficiency ? 1.0f : 0.0f);
+		SetShaderValue(BatchedParameters, ColorVisionDeficiencyType, (float)DeficiencyType);
+		SetShaderValue(BatchedParameters, ColorVisionDeficiencySeverity, (float)Severity);
+		SetShaderValue(BatchedParameters, bCorrectDeficiency, bCorrect ? 1.0f : 0.0f);
+		SetShaderValue(BatchedParameters, bSimulateCorrectionWithDeficiency, bShowCorrectionWithDeficiency ? 1.0f : 0.0f);
 	}
 
 	static const TCHAR* GetSourceFilename()
@@ -780,18 +778,23 @@ public:
 		return TEXT("Main");
 	}
 
-	void SetParameters(FRHICommandList& RHICmdList, FRHITexture* UITextureRHI, FRHITexture* UITextureWriteMaskRHI, FRHITexture* SceneTextureRHI, FRHITexture* ColorSpaceLUTRHI)
+	void SetParameters(
+		FRHIBatchedShaderParameters& BatchedParameters,
+		FRHITexture* UITextureRHI,
+		FRHITexture* UITextureWriteMaskRHI,
+		FRHITexture* SceneTextureRHI,
+		FRHITexture* ColorSpaceLUTRHI,
+		bool bCorrect,
+		EColorVisionDeficiency DeficiencyType,
+		int32 Severity,
+		bool bShowCorrectionWithDeficiency
+	)
 	{
-		SetParametersBase(RHICmdList, RHICmdList.GetBoundPixelShader(), UITextureRHI, UITextureWriteMaskRHI, ColorSpaceLUTRHI);
-		SetTextureParameter(RHICmdList, RHICmdList.GetBoundPixelShader(), SceneTexture, SceneSampler, TStaticSamplerState<SF_Point>::GetRHI(), SceneTextureRHI);
+		SetParametersBase(BatchedParameters, UITextureRHI, UITextureWriteMaskRHI, ColorSpaceLUTRHI);
+		SetTextureParameter(BatchedParameters, SceneTexture, SceneSampler, TStaticSamplerState<SF_Point>::GetRHI(), SceneTextureRHI);
+		SetColorDeficiencyParamsBase(BatchedParameters, bCorrect, DeficiencyType, Severity, bShowCorrectionWithDeficiency);
 	}
 
-	void SetColorDeficiencyParams(FRHICommandList& RHICmdList, bool bCorrect, EColorVisionDeficiency DeficiencyType, int32 Severity, bool bShowCorrectionWithDeficiency)
-	{
-		SetColorDeficiencyParamsBase(RHICmdList, RHICmdList.GetBoundPixelShader(), bCorrect, DeficiencyType, Severity, bShowCorrectionWithDeficiency);
-	}
-
-
 private:
 	LAYOUT_FIELD(FShaderResourceParameter, SceneTexture);
 	LAYOUT_FIELD(FShaderResourceParameter, SceneSampler);
@@ -848,16 +851,23 @@ public:
 		return TEXT("CompositeUICS");
 	}
 
-	void SetParameters(FRHICommandList& RHICmdList, FRHITexture* UITextureRHI, FRHITexture* UITextureWriteMaskRHI, FRHIUnorderedAccessView* InRWSceneTexture, FRHITexture* ColorSpaceLUTRHI, const FVector4f& InSceneTextureDimensions)
-	{
-		SetParametersBase(RHICmdList, RHICmdList.GetBoundComputeShader(), UITextureRHI, UITextureWriteMaskRHI, ColorSpaceLUTRHI);
-		SetUAVParameter(RHICmdList, RHICmdList.GetBoundComputeShader(), RWSceneTexture, InRWSceneTexture);
-		SetShaderValue(RHICmdList, RHICmdList.GetBoundComputeShader(), SceneTextureDimensions, InSceneTextureDimensions);
-	}
-
-	void SetColorDeficiencyParams(FRHICommandList& RHICmdList, bool bCorrect, EColorVisionDeficiency DeficiencyType, int32 Severity, bool bShowCorrectionWithDeficiency)
+	void SetParameters(
+		FRHIBatchedShaderParameters& BatchedParameters,
+		FRHITexture* UITextureRHI,
+		FRHITexture* UITextureWriteMaskRHI,
+		FRHIUnorderedAccessView* InRWSceneTexture,
+		FRHITexture* ColorSpaceLUTRHI,
+		const FVector4f& InSceneTextureDimensions,
+		bool bCorrect,
+		EColorVisionDeficiency DeficiencyType,
+		int32 Severity,
+		bool bShowCorrectionWithDeficiency
+	)
 	{
-		SetColorDeficiencyParamsBase(RHICmdList, RHICmdList.GetBoundComputeShader(), bCorrect, DeficiencyType, Severity, bShowCorrectionWithDeficiency);
+		SetParametersBase(BatchedParameters, UITextureRHI, UITextureWriteMaskRHI, ColorSpaceLUTRHI);
+		SetUAVParameter(BatchedParameters, RWSceneTexture, InRWSceneTexture);
+		SetShaderValue(BatchedParameters, SceneTextureDimensions, InSceneTextureDimensions);
+		SetColorDeficiencyParamsBase(BatchedParameters, bCorrect, DeficiencyType, Severity, bShowCorrectionWithDeficiency);
 	}
 
 private:
@@ -1224,14 +1234,17 @@ void FSlateRHIRenderer::DrawWindow_RenderThread(FRHICommandListImmediate& RHICmd
 						GraphicsPSOInit.PrimitiveType = PT_TriangleStrip;
 						SetGraphicsPipelineState(RHICmdList, GraphicsPSOInit, 0);
 
-						VertexShader->SetParameters(RHICmdList, VolumeBounds, FIntVector(VolumeBounds.MaxX - VolumeBounds.MinX));
+						SetAllShaderParametersVS(RHICmdList, VertexShader, VolumeBounds, FIntVector(VolumeBounds.MaxX - VolumeBounds.MinX));
+
 #if PLATFORM_SUPPORTS_GEOMETRY_SHADERS
 						if (GeometryShader.IsValid())
 						{
-							GeometryShader->SetParameters(RHICmdList, VolumeBounds.MinZ);
+							SetAllShaderParametersGS(RHICmdList, GeometryShader, VolumeBounds.MinZ);
 						}
 #endif
-						PixelShader->SetParameters(RHICmdList, ViewportInfo.HDRDisplayOutputFormat, ViewportInfo.HDRDisplayColorGamut, HDRGetDisplayMaximumLuminance());
+						const float DisplayMaxLuminance = HDRGetDisplayMaximumLuminance();
+
+						SetAllShaderParametersPS(RHICmdList, PixelShader, ViewportInfo.HDRDisplayOutputFormat, ViewportInfo.HDRDisplayColorGamut, DisplayMaxLuminance);
 
 						RasterizeToVolumeTexture(RHICmdList, VolumeBounds);
 					}
@@ -1316,8 +1329,9 @@ void FSlateRHIRenderer::DrawWindow_RenderThread(FRHICommandListImmediate& RHICmd
 
 						SetGraphicsPipelineState(RHICmdList, GraphicsPSOInit, 0);
 
-						PixelShader->SetParameters(RHICmdList, ViewportInfo.UITargetRT->GetRHI(), UITargetRTMaskTexture, FinalBufferCopy->GetRHI(), ViewportInfo.ColorSpaceLUT);
-						PixelShader->SetColorDeficiencyParams(RHICmdList, GSlateColorDeficiencyCorrection, GSlateColorDeficiencyType, GSlateColorDeficiencySeverity, GSlateShowColorDeficiencyCorrectionWithDeficiency);
+						SetAllShaderParametersPS(RHICmdList, PixelShader, 
+							ViewportInfo.UITargetRT->GetRHI(), UITargetRTMaskTexture, FinalBufferCopy->GetRHI(), ViewportInfo.ColorSpaceLUT,
+							GSlateColorDeficiencyCorrection, GSlateColorDeficiencyType, GSlateColorDeficiencySeverity, GSlateShowColorDeficiencyCorrectionWithDeficiency);
 
 						RendererModule.DrawRectangle(
 							RHICmdList,
@@ -1364,8 +1378,10 @@ void FSlateRHIRenderer::DrawWindow_RenderThread(FRHICommandListImmediate& RHICmd
 
 					FVector4f SceneTextureDimensions((float)ViewportWidth, (float)ViewportHeight, 1.0f/(float)ViewportWidth, 1.0f/(float)ViewportHeight);
 					SetComputePipelineState(RHICmdList, ComputeShader.GetComputeShader());
-					ComputeShader->SetParameters(RHICmdList, ViewportInfo.UITargetRT->GetRHI(), UITargetRTMaskTexture, BackBufferUAV, ViewportInfo.ColorSpaceLUT, SceneTextureDimensions);
-					ComputeShader->SetColorDeficiencyParams(RHICmdList, GSlateColorDeficiencyCorrection, GSlateColorDeficiencyType, GSlateColorDeficiencySeverity, GSlateShowColorDeficiencyCorrectionWithDeficiency);
+
+					SetAllShaderParametersCS(RHICmdList, ComputeShader,
+						ViewportInfo.UITargetRT->GetRHI(), UITargetRTMaskTexture, BackBufferUAV, ViewportInfo.ColorSpaceLUT, SceneTextureDimensions,
+						GSlateColorDeficiencyCorrection, GSlateColorDeficiencyType, GSlateColorDeficiencySeverity, GSlateShowColorDeficiencyCorrectionWithDeficiency);
 
 					RHICmdList.DispatchComputeShader(FMath::DivideAndRoundUp(ViewportWidth, FCompositeCSBase::NUM_THREADS_PER_GROUP), FMath::DivideAndRoundUp(ViewportHeight, FCompositeCSBase::NUM_THREADS_PER_GROUP), 1);
 
@@ -1403,7 +1419,7 @@ void FSlateRHIRenderer::DrawWindow_RenderThread(FRHICommandListImmediate& RHICmd
 
 				SetGraphicsPipelineState(RHICmdList, GraphicsPSOInit, 0);
 
-				PixelShader->SetParameters(RHICmdList, HDRRenderRT->GetRHI() );
+				SetAllShaderParametersPS(RHICmdList, PixelShader, HDRRenderRT->GetRHI());
 
 				static const FName RendererModuleName("Renderer");
 				IRendererModule& RendererModule = FModuleManager::GetModuleChecked<IRendererModule>(RendererModuleName);
diff --git a/Engine/Source/Runtime/SlateRHIRenderer/Private/SlateShaders.h b/Engine/Source/Runtime/SlateRHIRenderer/Private/SlateShaders.h
--- a/Engine/Source/Runtime/SlateRHIRenderer/Private/SlateShaders.h
+++ b/Engine/Source/Runtime/SlateRHIRenderer/Private/SlateShaders.h
@@ -558,12 +558,12 @@ public:
 	}
 	FHDREditorConvertPS() {}
 
-	void SetParameters(FRHICommandList& RHICmdList, FRHITexture* SceneTextureRHI)
+	void SetParameters(FRHIBatchedShaderParameters& BatchedParameters, FRHITexture* SceneTextureRHI)
 	{
-		SetTextureParameter(RHICmdList, RHICmdList.GetBoundPixelShader(), SceneTexture, SceneSampler, TStaticSamplerState<SF_Point>::GetRHI(), SceneTextureRHI);
+		SetTextureParameter(BatchedParameters, SceneTexture, SceneSampler, TStaticSamplerState<SF_Point>::GetRHI(), SceneTextureRHI);
 		
 		static auto CVarHDRNITLevel = IConsoleManager::Get().FindConsoleVariable(TEXT("Editor.HDRNITLevel"));
-		SetShaderValue(RHICmdList, RHICmdList.GetBoundPixelShader(), UILevel, CVarHDRNITLevel->GetFloat());
+		SetShaderValue(BatchedParameters, UILevel, CVarHDRNITLevel->GetFloat());
 	}
 
 	static bool ShouldCompilePermutation(const FGlobalShaderPermutationParameters& Parameters);
