Reinstated CL 23834590, following the fix to have WP-streamed level's OwningWorld properly set to the destination world upon loading (23991534) :

Stop registering ULandscapeInfo in PostLoad. Technically, this is now possible but not needed.

#rnx

diff --git a/Engine/Source/Runtime/Landscape/Private/Landscape.cpp b/Engine/Source/Runtime/Landscape/Private/Landscape.cpp
--- a/Engine/Source/Runtime/Landscape/Private/Landscape.cpp
+++ b/Engine/Source/Runtime/Landscape/Private/Landscape.cpp
@@ -3342,13 +3342,7 @@ void ALandscapeProxy::PostLoad()
 	{
 		BodyInstance.FixupData(this);
 	}
-
-	if ((GetLinker() && (GetLinker()->UEVer() < VER_UE4_LANDSCAPE_COMPONENT_LAZY_REFERENCES)) ||
-		LandscapeComponents.Num() != CollisionComponents.Num() ||
-		LandscapeComponents.ContainsByPredicate([](ULandscapeComponent* Comp) { return ((Comp != nullptr) && (Comp->GetCollisionComponent() == nullptr)); }))
-	{
-		CreateLandscapeInfo();
-	}
+	
 #if WITH_EDITOR
 	PRAGMA_DISABLE_DEPRECATION_WARNINGS;
 	if (!LandscapeMaterialsOverride_DEPRECATED.IsEmpty())
diff --git a/Engine/Source/Runtime/Landscape/Private/LandscapeEdit.cpp b/Engine/Source/Runtime/Landscape/Private/LandscapeEdit.cpp
--- a/Engine/Source/Runtime/Landscape/Private/LandscapeEdit.cpp
+++ b/Engine/Source/Runtime/Landscape/Private/LandscapeEdit.cpp
@@ -1332,7 +1332,6 @@ void ULandscapeComponent::RecreateCollisionComponent(bool bUseSimpleCollision)
 	TArray<uint8> DominantLayerData;
 	TArray<ULandscapeLayerInfoObject*> LayerInfos;
 	ALandscapeProxy* Proxy = GetLandscapeProxy();
-	ULandscapeInfo* Info = GetLandscapeInfo();
 	const FCollisionSize CollisionSize = FCollisionSize::Create(NumSubsections, SubsectionSizeQuads, CollisionMipLevel);
 	const FCollisionSize SimpleCollisionSize = FCollisionSize::CreateSimple(bUseSimpleCollision, NumSubsections, SubsectionSizeQuads, SimpleCollisionMipLevel);
 	const int32 TotalCollisionSize = CollisionSize.SizeVertsSquare + SimpleCollisionSize.SizeVertsSquare;
@@ -1354,10 +1353,6 @@ void ULandscapeComponent::RecreateCollisionComponent(bool bUseSimpleCollision)
 			LayerInfos = CollisionComp->ComponentLayerInfos;
 		}
 
-		if (Info)
-		{
-			Info->Modify();
-		}
 		Proxy->Modify();
 		CollisionComp->DestroyComponent();
 		CollisionComp = nullptr;
