Workaround the D3D12-only bug that randomly generates invalid landscape tiles on global merge.
The problem is that there seems to be a missing barrier when copying individual edit layers heightmap textures into the atlas RT (destination state = COPY_DEST) and then using it as a SRV afterwards in the merge PS. Using a PS-based copy texture avoids the issue (CVar landscape.UsePSCopyWorkaround to enable/disable it).
Also, this CL changes the landscape.OutputDiffBitmap and turns it into landscape.DumpHeightmapDiff and landscape.DumpWeightmapDiff (0 = no dump, 1 = mip 0 diff, 2 = all mips diff). Both of these now output to the Saved/LandscapeLayers folder, with a sub-folder named after the current frame's date/time, under a sub-folder named after the landscape actor being updated (so that all diffs done in the same frame can be easily identified by their date, landscape actor name and type : heightmap or weightmap)

diff --git a/Engine/Shaders/Private/LandscapeLayersPS.usf b/Engine/Shaders/Private/LandscapeLayersPS.usf
--- a/Engine/Shaders/Private/LandscapeLayersPS.usf
+++ b/Engine/Shaders/Private/LandscapeLayersPS.usf
@@ -14,6 +14,7 @@ float3 LandscapeGridScale;						// x == LS Actor DrawScale.X, y == LS Actor Draw
 float CurrentMipComponentVertexCount;
 float2 CurrentMipTextureSize;
 float2 ParentMipTextureSize;
+float4 SourceOffsetAndSizeUV;
 
 float ExtractAndScaleHeight(float2 InPackedHeight)
 {
@@ -381,7 +382,9 @@ void CopyTextureVS(in uint VertexId : SV_VertexID, out float4 OutPosition : SV_P
 	OutPosition.zw = float2(0, 1);
 }
 
-void CopyTexturePS(in float4 InPosition : SV_POSITION, in noperspective float2 InTexCoord : TEXCOORD0, out float4 OutColor : SV_Target0)
+void CopyTexturePS(in float4 InPosition : SV_POSITION, in noperspective float2 InUV : TEXCOORD0, out float4 OutColor : SV_Target0)
 {
-	OutColor = ReadTexture1.SampleLevel(ReadTexture1Sampler, InTexCoord, 0);
+	// InUV is in the [0,1] range but in the viewport space, whose size is CopySize, so we need to adjust it to sample the corresponding range in the source texture : 
+	float2 SampleUV = (InUV * SourceOffsetAndSizeUV.zw) + SourceOffsetAndSizeUV.xy;
+	OutColor = ReadTexture1.SampleLevel(ReadTexture1Sampler, SampleUV, 0);
 }
diff --git a/Engine/Source/Runtime/Landscape/Classes/Landscape.h b/Engine/Source/Runtime/Landscape/Classes/Landscape.h
--- a/Engine/Source/Runtime/Landscape/Classes/Landscape.h
+++ b/Engine/Source/Runtime/Landscape/Classes/Landscape.h
@@ -437,9 +437,9 @@ private:
 	void PrintLayersDebugWeightData(const FString& InContext, const TArray<FColor>& InWeightmapData, const FIntPoint& InDataSize, uint8 InMipRender) const;
 
 	void UpdateWeightDirtyData(ULandscapeComponent* InLandscapeComponent, UTexture2D const* InWeightmap, FColor const* InOldData, FColor const* InNewData, uint8 InChannel);
-	void OnDirtyWeightmap(FTextureToComponentHelper const& MapHelper, UTexture2D const* InWeightmap, FColor const* InOldData, FColor const* InNewData);
+	void OnDirtyWeightmap(FTextureToComponentHelper const& MapHelper, UTexture2D const* InWeightmap, FColor const* InOldData, FColor const* InNewData, int32 InMipLevel);
 	void UpdateHeightDirtyData(ULandscapeComponent* InLandscapeComponent, UTexture2D const* InHeightmap, FColor const* InOldData, FColor const* InNewData);
-	void OnDirtyHeightmap(FTextureToComponentHelper const& MapHelper, UTexture2D const* InWeightmap, FColor const* InOldData, FColor const* InNewData);
+	void OnDirtyHeightmap(FTextureToComponentHelper const& MapHelper, UTexture2D const* InWeightmap, FColor const* InOldData, FColor const* InNewData, int32 InMipLevel);
 
 	static bool IsTextureReady(UTexture2D* InTexture, bool bInWaitForStreaming);
 	static bool IsMaterialResourceCompiled(FMaterialResource* InMaterialResource, bool bInWaitForCompilation);
diff --git a/Engine/Source/Runtime/Landscape/Private/LandscapeEditLayers.cpp b/Engine/Source/Runtime/Landscape/Private/LandscapeEditLayers.cpp
--- a/Engine/Source/Runtime/Landscape/Private/LandscapeEditLayers.cpp
+++ b/Engine/Source/Runtime/Landscape/Private/LandscapeEditLayers.cpp
@@ -111,7 +111,6 @@ static TAutoConsoleVariable<int32> CVarOutputLayersWeightmapsRTContent(
 	0,
 	TEXT("This will output the content of render target used for weightmap. This is used for debugging only."));
 
-
 static TAutoConsoleVariable<int32> CVarLandscapeSimulatePhysics(
 	TEXT("landscape.SimulatePhysics"),
 	0,
@@ -127,10 +126,15 @@ static TAutoConsoleVariable<int32> CVarLandscapeLayerBrushOptim(
 	0,
 	TEXT("This will enable landscape layers optim."));
 
-static TAutoConsoleVariable<int32> CVarLandscapeOutputDiffBitmap(
-	TEXT("landscape.OutputDiffBitmap"),
+static TAutoConsoleVariable<int32> CVarLandscapeDumpHeightmapDiff(
+	TEXT("landscape.DumpHeightmapDiff"),
+	0,
+	TEXT("This will save images for readback heightmap textures that have changed in the last edit layer blend phase. (= 0 No Diff, 1 = Mip 0 Diff, 2 = All Mips Diff"));
+
+static TAutoConsoleVariable<int32> CVarLandscapeDumpWeightmapDiff(
+	TEXT("landscape.DumpWeightmapDiff"),
 	0,
-	TEXT("This will save images for readback textures that have changed in the last layer blend phase. (= 1 Heightmap Diff, = 2 Weightmap Diff, = 3 All Diffs"));
+	TEXT("This will save images for readback weightmap textures that have changed in the last edit layer blend phase. (= 0 No Diff, 1 = Mip 0 Diff, 2 = All Mips Diff"));
 
 TAutoConsoleVariable<int32> CVarLandscapeShowDirty(
 	TEXT("landscape.ShowDirty"),
@@ -157,6 +161,13 @@ TAutoConsoleVariable<int32> CVarLandscapeRemoveEmptyPaintLayersOnEdit(
 	1,
 	TEXT("This will analyze weightmaps on readback and remove unneeded allocations (for unpainted layers)."));
 
+// COMMENT [jonathan.bard] : this is a workaround for a random D3D12 issue that misses some cache flush when copying edit layers individual heightmaps to the global merge atlas and then using 
+//  the atlas as SRV right after, which should be ensured by the manually inserted barriers but is sometimes not). Using PS-based copies doesn't exhibit the behavior so we'll stick with that for now : 
+TAutoConsoleVariable<bool> CVarLandscapeUsePSCopyWorkaround(
+	TEXT("landscape.UsePSCopyWorkaround"),
+	1,
+	TEXT("This is a workaround for a random D3D12-only issue when using a COPY_DEST RT as a SRV."));
+
 void OnLandscapeEditLayersLocalMergeChanged(IConsoleVariable* CVar)
 {
 	for (TObjectIterator<UWorld> It; It; ++It)
@@ -759,6 +770,7 @@ public:
 	{
 		ReadTexture1Param.Bind(Initializer.ParameterMap, TEXT("ReadTexture1"));
 		ReadTexture1SamplerParam.Bind(Initializer.ParameterMap, TEXT("ReadTexture1Sampler"));
+		SourceOffsetAndSizeUVParam.Bind(Initializer.ParameterMap, TEXT("SourceOffsetAndSizeUV"));
 	}
 
 	FLandscapeCopyTexturePS()
@@ -766,12 +778,20 @@ public:
 
 	void SetParameters(FRHIBatchedShaderParameters& BatchedParameters, FRHITexture* TextureRHI)
 	{
-		SetTextureParameter(BatchedParameters, ReadTexture1Param, ReadTexture1SamplerParam, TStaticSamplerState<SF_Point, AM_Clamp, AM_Clamp, AM_Clamp>::GetRHI(), TextureRHI);
+		FVector2f SourceSize(InSourceTextureRHI->GetSizeXY());
+		FVector2f SourceOffsetUV = FVector2f(InSourcePosition) / SourceSize;
+		FVector2f FinalCopySizePixels;
+		FinalCopySizePixels.X = (InCopySizePixels.X > 0) ? InCopySizePixels.X : SourceSize.X;
+		FinalCopySizePixels.Y = (InCopySizePixels.Y > 0) ? InCopySizePixels.Y : SourceSize.Y;
+		FVector2f CopySizeUV = FinalCopySizePixels / SourceSize;
+		SetTextureParameter(BatchedParameters, ReadTexture1Param, ReadTexture1SamplerParam, TStaticSamplerState<SF_Point, AM_Clamp, AM_Clamp, AM_Clamp>::GetRHI(), InSourceTextureRHI);
+		SetShaderValue(BatchedParameters, SourceOffsetAndSizeUVParam, FVector4f(SourceOffsetUV.X, SourceOffsetUV.Y, CopySizeUV.X, CopySizeUV.Y));
 	}
 
 private:
 	LAYOUT_FIELD(FShaderResourceParameter, ReadTexture1Param);
 	LAYOUT_FIELD(FShaderResourceParameter, ReadTexture1SamplerParam);
+	LAYOUT_FIELD(FShaderParameter, SourceOffsetAndSizeUVParam);
 };
 
 IMPLEMENT_GLOBAL_SHADER(FLandscapeCopyTextureVS, "/Engine/Private/LandscapeLayersPS.usf", "CopyTextureVS", SF_Vertex);
@@ -1452,7 +1472,7 @@ struct FLandscapeLayersCopyTextureParams
 		}
 		if (InDestTexture != nullptr)
 		{
-			DestResourceDebugName = InDestTexture->GetName();
+			SourceResourceDebugName = InDestTexture->GetName();
 			DestResource = InDestTexture->GetResource();
 		}
 	}
@@ -1481,6 +1501,7 @@ struct FLandscapeLayersCopyTextureParams
 	uint32 DestArrayIndex = 0;
 	ERHIAccess SourceAccess = ERHIAccess::SRVMask;
 	ERHIAccess DestAccess = ERHIAccess::SRVMask;
+	bool bUsePS = false;
 };
 
 class FLandscapeLayersCopyTexture_RenderThread
@@ -1491,8 +1512,20 @@ public:
 	{}
 
 	const FLandscapeLayersCopyTextureParams& GetParams() const { return Params; }
-
 	void Copy(FRHICommandListImmediate& InRHICmdList)
+	{
+		if (Params.bUsePS)
+		{
+			CopyInternalPS(InRHICmdList);
+		}
+		else
+		{
+			CopyInternal(InRHICmdList);
+		}
+	}
+
+private:
+	void CopyInternal(FRHICommandListImmediate& InRHICmdList)
 	{
 		SCOPED_GPU_STAT(InRHICmdList, LandscapeLayers_CopyTexture);
 		SCOPED_DRAW_EVENTF(InRHICmdList, LandscapeLayers, TEXT("LandscapeLayers_Copy %s -> %s, Mip (%d -> %d), Array Index (%d -> %d)"), *Params.SourceResourceDebugName, *Params.DestResourceDebugName, Params.SourceMip, Params.DestMip, Params.SourceArrayIndex, Params.DestArrayIndex);
@@ -1528,6 +1561,52 @@ public:
 		InRHICmdList.Transition(FRHITransitionInfo(Params.DestResource->TextureRHI, ERHIAccess::CopyDest, Params.DestAccess));
 	}
 
+	void CopyInternalPS(FRHICommandListImmediate& InRHICmdList)
+	{
+		SCOPED_GPU_STAT(InRHICmdList, LandscapeLayers_CopyTexturePS);
+		SCOPED_DRAW_EVENTF(InRHICmdList, LandscapeLayers, TEXT("LandscapeLayers_CopyPS %s -> %s, Mip (%d -> %d), Array Index (%d -> %d)"), *Params.SourceResourceDebugName, *Params.DestResourceDebugName, Params.SourceMip, Params.DestMip, Params.SourceArrayIndex, Params.DestArrayIndex);
+
+		FIntPoint SourceSize(Params.SourceResource->GetSizeX() >> Params.SourceMip, Params.SourceResource->GetSizeY() >> Params.SourceMip);
+		FIntPoint DestSize(Params.DestResource->GetSizeX() >> Params.DestMip, Params.DestResource->GetSizeY() >> Params.DestMip);
+
+		// If CopySize is passed, used that as the size (and don't adjust with the mip level : consider that the user has computed it properly) : 
+		FIntPoint Size;
+		Size.X = (Params.CopySize.X > 0) ? Params.CopySize.X : SourceSize.X;
+		Size.Y = (Params.CopySize.Y > 0) ? Params.CopySize.Y : SourceSize.Y;
+		check((Params.SourceArrayIndex == 0) && (Params.DestArrayIndex == 0) && (Params.SourceMip == 0) && (Params.DestMip == 0)); // The PS version of copy is not supported on texture arrays and mips for now
+
+		InRHICmdList.Transition(FRHITransitionInfo(Params.SourceResource->TextureRHI, Params.SourceAccess, ERHIAccess::SRVGraphics));
+		InRHICmdList.Transition(FRHITransitionInfo(Params.DestResource->TextureRHI, Params.DestAccess, ERHIAccess::RTV));
+
+		FRHIRenderPassInfo RPInfo(Params.DestResource->TextureRHI, ERenderTargetActions::DontLoad_Store);
+		InRHICmdList.BeginRenderPass(RPInfo, TEXT("CopyTexture"));
+
+		FGlobalShaderMap* GlobalShaderMap = GetGlobalShaderMap(GMaxRHIFeatureLevel);
+		TShaderMapRef< FLandscapeCopyTextureVS > VertexShader(GlobalShaderMap);
+		TShaderMapRef< FLandscapeCopyTexturePS > PixelShader(GlobalShaderMap);
+
+		FGraphicsPipelineStateInitializer GraphicsPSOInit;
+		InRHICmdList.ApplyCachedRenderTargets(GraphicsPSOInit);
+		GraphicsPSOInit.BlendState = TStaticBlendState<>::GetRHI();
+		GraphicsPSOInit.RasterizerState = TStaticRasterizerState<>::GetRHI();
+		GraphicsPSOInit.DepthStencilState = TStaticDepthStencilState<false, CF_Always>::GetRHI();
+		GraphicsPSOInit.PrimitiveType = PT_TriangleList;
+		GraphicsPSOInit.BoundShaderState.VertexDeclarationRHI = GetVertexDeclarationFVector4();
+		GraphicsPSOInit.BoundShaderState.VertexShaderRHI = VertexShader.GetVertexShader();
+		GraphicsPSOInit.BoundShaderState.PixelShaderRHI = PixelShader.GetPixelShader();
+		SetGraphicsPipelineState(InRHICmdList, GraphicsPSOInit, 0);
+
+		PixelShader->SetParameters(InRHICmdList, Params.SourceResource->TextureRHI, Params.SourcePosition, Size);
+
+		InRHICmdList.SetViewport((float)Params.DestPosition.X, (float)Params.DestPosition.Y, 0.0f, (float)(Params.DestPosition.X + Size.X), (float)(Params.DestPosition.Y + Size.Y), 1.0f);
+		InRHICmdList.DrawIndexedPrimitive(GTwoTrianglesIndexBuffer.IndexBufferRHI, 0, 0, 4, 0, 2, 1);
+
+		InRHICmdList.EndRenderPass();
+
+		InRHICmdList.Transition(FRHITransitionInfo(Params.SourceResource->TextureRHI, ERHIAccess::SRVGraphics, Params.SourceAccess));
+		InRHICmdList.Transition(FRHITransitionInfo(Params.DestResource->TextureRHI, ERHIAccess::RTV, Params.DestAccess));
+	}
+
 private:
 	FLandscapeLayersCopyTextureParams Params;
 };
@@ -4753,6 +4832,8 @@ int32 ALandscape::PerformLayersHeightmapsGlobalMerge(const FUpdateLayersContentC
 	UTextureRenderTarget2D* LandscapeScratchRT2 = HeightmapRTList[(int32)EHeightmapRTType::HeightmapRT_Scratch2];
 	UTextureRenderTarget2D* LandscapeScratchRT3 = HeightmapRTList[(int32)EHeightmapRTType::HeightmapRT_Scratch3];
 
+	const bool bUsePSCopyWorkaround = CVarLandscapeUsePSCopyWorkaround.GetValueOnGameThread();
+
 	for (FLandscapeLayer& Layer : LandscapeLayers)
 	{
 		//Draw Layer heightmap to Combined RT Atlas
@@ -4785,6 +4866,7 @@ int32 ALandscape::PerformLayersHeightmapsGlobalMerge(const FUpdateLayersContentC
 				CopyTextureParams.CopySize = LayerHeightmap.SectionRect.Size();
 				// Copy from the heightmap's top-left corner to the composited texture's position :
 				CopyTextureParams.DestPosition = LayerHeightmap.SectionRect.Min;
+				CopyTextureParams.bUsePS = bUsePSCopyWorkaround;
 			}
 			ExecuteCopyLayersTexture(MoveTemp(DeferredCopyTextures));
 		}
@@ -4873,9 +4955,24 @@ int32 ALandscape::PerformLayersHeightmapsGlobalMerge(const FUpdateLayersContentC
 
 			const FIntPoint Mip0CopySize = Heightmap.SectionRect.Size();
 			const FIntPoint Mip0SourcePosition = Heightmap.SectionRect.Min;
+
+			FString SourceTextureName = CombinedHeightmapAtlasRT->GetName();
+			FString DestinationTextureName = Heightmap.Texture->GetName();
+			if (const TArray<ULandscapeComponent*>* Components = InUpdateLayersContentContext.MapHelper.HeightmapToComponents.Find(Heightmap.Texture))
+			{
+				FString ComponentName;
+				FString ActorName;
+				for (ULandscapeComponent* Component : *Components)
+				{
+					ComponentName += Component->GetName() + TEXT(", ");
+					ActorName = Component->GetLandscapeProxy()->GetActorLabel();
+				}
+				DestinationTextureName = FString::Format(TEXT("{0}-{1}-{2}"), { ActorName, ComponentName, Heightmap.Texture->GetName() });
+			}
+
 			// Mip 0
 			{
-				FLandscapeLayersCopyTextureParams& CopyTextureParams = DeferredCopyTextures.Add_GetRef(FLandscapeLayersCopyTextureParams(CombinedHeightmapAtlasRT, Heightmap.Texture));
+				FLandscapeLayersCopyTextureParams& CopyTextureParams = DeferredCopyTextures.Add_GetRef(FLandscapeLayersCopyTextureParams(SourceTextureName, CombinedHeightmapAtlasRT->GetResource(), DestinationTextureName, Heightmap.Texture->GetResource()));
 				// Only copy the size that's actually needed : 
 				CopyTextureParams.CopySize = Mip0CopySize;
 				// Copy from the composited texture's position to the top-left corner of the heightmap
@@ -4889,7 +4986,8 @@ int32 ALandscape::PerformLayersHeightmapsGlobalMerge(const FUpdateLayersContentC
 				UTextureRenderTarget2D* RenderTargetMip = HeightmapRTList[MipRTIndex];
 				if (RenderTargetMip != nullptr)
 				{
-					FLandscapeLayersCopyTextureParams& CopyTextureParams = DeferredCopyTextures.Add_GetRef(FLandscapeLayersCopyTextureParams(RenderTargetMip, Heightmap.Texture));
+					SourceTextureName = RenderTargetMip->GetName();
+					FLandscapeLayersCopyTextureParams& CopyTextureParams = DeferredCopyTextures.Add_GetRef(FLandscapeLayersCopyTextureParams(SourceTextureName, RenderTargetMip->GetResource(), DestinationTextureName, Heightmap.Texture->GetResource()));
 					CopyTextureParams.CopySize.X = Mip0CopySize.X >> MipIndex;
 					CopyTextureParams.CopySize.Y = Mip0CopySize.Y >> MipIndex;
 					CopyTextureParams.SourcePosition.X = Mip0SourcePosition.X >> MipIndex;
@@ -5108,11 +5206,19 @@ void ALandscape::UpdateWeightDirtyData(ULandscapeComponent* InLandscapeComponent
 	LandscapeEdit.SetDirtyData(X1, Y1, X2, Y2, DirtyData.Get(), 0);
 }
 
-void ALandscape::OnDirtyWeightmap(FTextureToComponentHelper const& MapHelper, UTexture2D const* InWeightmap, FColor const* InOldData, FColor const* InNewData)
+void ALandscape::OnDirtyWeightmap(FTextureToComponentHelper const& MapHelper, UTexture2D const* InWeightmap, FColor const* InOldData, FColor const* InNewData, int32 InMipLevel)
 {
-	const bool bWriteDiff = (CVarLandscapeOutputDiffBitmap.GetValueOnAnyThread() & 2) != 0;
+	int32 DumpWeightmapDiff = CVarLandscapeDumpWeightmapDiff.GetValueOnAnyThread();
+	const bool bDumpDiff = (DumpWeightmapDiff > 0);
+	const bool bDumpDiffAllMips = (DumpWeightmapDiff > 1);
 	const bool bTrackDirty = CVarLandscapeTrackDirty.GetValueOnAnyThread() != 0;
-	if (!bWriteDiff && !bTrackDirty)
+	ULandscapeSubsystem* LandscapeSubsystem = GetWorld()->GetSubsystem<ULandscapeSubsystem>();
+	check(LandscapeSubsystem != nullptr);
+	const FDateTime CurrentTime = LandscapeSubsystem->GetAppCurrentDateTime();
+
+	if ((!bDumpDiff && !bTrackDirty)
+		|| (bDumpDiff && !bDumpDiffAllMips && (InMipLevel > 0))
+		|| (bTrackDirty && (InMipLevel > 0)))
 	{
 		return;
 	}
@@ -5135,13 +5241,21 @@ void ALandscape::OnDirtyWeightmap(FTextureToComponentHelper const& MapHelper, UT
 						UpdateWeightDirtyData(Component, InWeightmap, InOldData, InNewData, AllocInfo.WeightmapTextureChannel);
 					}
 
-					if (bWriteDiff)
+					if (bDumpDiff)
 					{
+						const int32 SizeU = InWeightmap->Source.GetSizeX() >> InMipLevel;
+						const int32 SizeV = InWeightmap->Source.GetSizeY() >> InMipLevel;
+
+						FString WorldName = GetWorld()->GetName();
+						FString ParentLandscapeActorName = GetActorLabel();
+						ALandscapeProxy* Proxy = Cast<ALandscapeProxy>(Component->GetOwner());
+						check(Proxy);
+						FString ActorName = Proxy->GetActorLabel();
+						FString FilePattern = FString::Format(TEXT("{0}/LandscapeLayers/{1}/{2}/{3}/Weightmaps/{4}-{5}-{6}[mip{7}]"), { FPaths::ProjectSavedDir(), CurrentTime.ToString(), WorldName, ParentLandscapeActorName, ActorName, Component->GetName(), InWeightmap->GetName(), InMipLevel });
+
 						FFileHelper::EColorChannel ColorChannel = GetWeightmapColorChannel(AllocInfo);
-						FString LevelName = FPackageName::GetShortName(Component->GetOutermost());
-						FString FilePattern = FString::Format(TEXT("LandscapeLayers/{0}-{1}-{2}-WM"), { LevelName, Component->GetName(), AllocInfo.GetLayerName().ToString() });
-						FFileHelper::CreateBitmap(*(FilePattern + "-Pre.bmp"), InWeightmap->Source.GetSizeX(), InWeightmap->Source.GetSizeY(), InOldData, nullptr, nullptr, nullptr, true, ColorChannel);
-						FFileHelper::CreateBitmap(*(FilePattern + "-Post.bmp"), InWeightmap->Source.GetSizeX(), InWeightmap->Source.GetSizeY(), InNewData, nullptr, nullptr, nullptr, true, ColorChannel);
+						FFileHelper::CreateBitmap(*(FilePattern + "_a(pre).bmp"), SizeU, SizeV, InOldData, /*SubRectangle = */nullptr, &IFileManager::Get(), /*OutFilename = */nullptr, /*bool bInWriteAlpha = */true, ColorChannel);
+						FFileHelper::CreateBitmap(*(FilePattern + "_b(post).bmp"), SizeU, SizeV, InNewData, /*SubRectangle = */nullptr, &IFileManager::Get(), /*OutFilename = */nullptr, /*bool bInWriteAlpha = */true, ColorChannel);
 					}
 				}
 			}
@@ -5186,12 +5300,19 @@ void ALandscape::UpdateHeightDirtyData(ULandscapeComponent* InLandscapeComponent
 	LandscapeEdit.SetDirtyData(X1, Y1, X2, Y2, DirtyData.Get(), 0);
 }
 
-void ALandscape::OnDirtyHeightmap(FTextureToComponentHelper const& MapHelper, UTexture2D const* InHeightmap, FColor const* InOldData, FColor const* InNewData)
+void ALandscape::OnDirtyHeightmap(FTextureToComponentHelper const& MapHelper, UTexture2D const* InHeightmap, FColor const* InOldData, FColor const* InNewData, int32 InMipLevel)
 {
-	const bool bHeightmapDiff = (CVarLandscapeOutputDiffBitmap.GetValueOnAnyThread() & 1) != 0;
+	int32 DumpHeightmapDiff = CVarLandscapeDumpHeightmapDiff.GetValueOnAnyThread();
+	const bool bDumpDiff = (DumpHeightmapDiff > 0);
+	const bool bDumpDiffAllMips = (DumpHeightmapDiff > 1);
 	const bool bTrackDirty = CVarLandscapeTrackDirty.GetValueOnAnyThread() != 0;
+	ULandscapeSubsystem* LandscapeSubsystem = GetWorld()->GetSubsystem<ULandscapeSubsystem>();
+	check(LandscapeSubsystem != nullptr);
+	const FDateTime CurrentTime = LandscapeSubsystem->GetAppCurrentDateTime();
 
-	if (!bHeightmapDiff && !bTrackDirty)
+	if ((!bDumpDiff && !bTrackDirty)
+		|| (bDumpDiff && !bDumpDiffAllMips && (InMipLevel > 0))
+		|| (bTrackDirty && (InMipLevel > 0)))
 	{
 		return;
 	}
@@ -5206,20 +5327,24 @@ void ALandscape::OnDirtyHeightmap(FTextureToComponentHelper const& MapHelper, UT
 				UpdateHeightDirtyData(Component, InHeightmap, InOldData, InNewData);
 			}
 
-			if (bHeightmapDiff)
+			if (bDumpDiff)
 			{
-				FString LevelName = FPackageName::GetShortName(Component->GetOutermost());
-				FString FilePattern = FString::Format(TEXT("LandscapeLayers/{0}-{1}-HM"), { LevelName, Component->GetName() });
+				FString WorldName = GetWorld()->GetName();
+				FString ParentLandscapeActorName = GetActorLabel();
+				ALandscapeProxy* Proxy = Cast<ALandscapeProxy>(Component->GetOwner());
+				check(Proxy);
+				FString ActorName = Proxy->GetActorLabel();
+				FString FilePattern = FString::Format(TEXT("{0}/LandscapeLayers/{1}/{2}/{3}/Heightmaps/{4}-{5}-{6}[mip{7}]"), { FPaths::ProjectSavedDir(), CurrentTime.ToString(), WorldName, ParentLandscapeActorName, ActorName, Component->GetName(), InHeightmap->GetName(), InMipLevel });
 
-				const int32 SizeU = InHeightmap->Source.GetSizeX();
-				const int32 SizeV = InHeightmap->Source.GetSizeY();
+				const int32 SizeU = InHeightmap->Source.GetSizeX() >> InMipLevel;
+				const int32 SizeV = InHeightmap->Source.GetSizeY() >> InMipLevel;
 				const int32 HeightmapOffsetX = static_cast<int32>(Component->HeightmapScaleBias.Z * SizeU);
 				const int32 HeightmapOffsetY = static_cast<int32>(Component->HeightmapScaleBias.W * SizeV);
-				const int32 ComponentWidth = (SubsectionSizeQuads + 1) * NumSubsections;
+				const int32 ComponentWidth = ((SubsectionSizeQuads + 1) * NumSubsections) >> InMipLevel;
 				FIntRect SubRegion(HeightmapOffsetX, HeightmapOffsetY, HeightmapOffsetX + ComponentWidth, HeightmapOffsetY + ComponentWidth);
 
-				FFileHelper::CreateBitmap(*(FilePattern + "-Pre.bmp"), InHeightmap->Source.GetSizeX(), InHeightmap->Source.GetSizeY(), InOldData, &SubRegion, &IFileManager::Get(), nullptr, true);
-				FFileHelper::CreateBitmap(*(FilePattern + "-Post.bmp"), InHeightmap->Source.GetSizeX(), InHeightmap->Source.GetSizeY(), InNewData, &SubRegion, &IFileManager::Get(), nullptr, true);
+				FFileHelper::CreateBitmap(*(FilePattern + "_a(pre).bmp"), SizeU, SizeV, InOldData, &SubRegion, &IFileManager::Get(), /*OutFilename = */nullptr, /*bool bInWriteAlpha = */true);
+				FFileHelper::CreateBitmap(*(FilePattern + "_b(post).bmp"), SizeU, SizeV, InNewData, &SubRegion, &IFileManager::Get(), /*OutFilename = */nullptr, /*bool bInWriteAlpha = */true);
 			}
 		}
 	}
@@ -5285,17 +5410,7 @@ bool ALandscape::ResolveLayersTexture(
 						TRACE_CPUPROFILER_EVENT_SCOPE(LandscapeLayers_CalculateHash);
 						Hash = FLandscapeEditLayerReadback::CalculateHash((uint8*)OutMipsData[MipIndex].GetData(), OutMipsData[MipIndex].Num() * sizeof(FColor));
 					}
-					if (InCPUReadback->SetHash(Hash))
-					{
-						if (bIsWeightmap)
-						{
-							OnDirtyWeightmap(MapHelper, InOutputTexture, (FColor*)TextureData, OutMipsData[MipIndex].GetData());
-						}
-						else
-						{
-							OnDirtyHeightmap(MapHelper, InOutputTexture, (FColor*)TextureData, OutMipsData[MipIndex].GetData());
-						}
-						bChanged = true;
+					bChanged |= InCPUReadback->SetHash(Hash);
 
 						// We're about to modify the texture's source data, the texture needs to know so that it can handle properly update cached platform data (additionally, the package needs to be dirtied) :
 						if (GetDefault<ULandscapeSettings>()->LandscapeDirtyingMode == ELandscapeDirtyingMode::InLandscapeModeAndUserTriggeredChanges)
@@ -5307,7 +5422,17 @@ bool ALandscape::ResolveLayersTexture(
 						{
 							GetLandscapeInfo()->ModifyObject(InOutputTexture);
 						}
+				}
 							
+				if (bChanged)
+				{
+					if (bIsWeightmap)
+					{
+						OnDirtyWeightmap(MapHelper, InOutputTexture, (FColor*)TextureData, OutMipsData[MipIndex].GetData(), MipIndex);
+					}
+					else
+					{
+						OnDirtyHeightmap(MapHelper, InOutputTexture, (FColor*)TextureData, OutMipsData[MipIndex].GetData(), MipIndex);
 					}
 				}
 
diff --git a/Engine/Source/Runtime/Landscape/Private/LandscapeSubsystem.cpp b/Engine/Source/Runtime/Landscape/Private/LandscapeSubsystem.cpp
--- a/Engine/Source/Runtime/Landscape/Private/LandscapeSubsystem.cpp
+++ b/Engine/Source/Runtime/Landscape/Private/LandscapeSubsystem.cpp
@@ -28,6 +28,9 @@
 #include "Algo/Transform.h"
 #include "Algo/RemoveIf.h"
 #include "Algo/Unique.h"
+#include "Misc/App.h"
+#include "Misc/DateTime.h"
+
 
 #include UE_INLINE_GENERATED_CPP_BY_NAME(LandscapeSubsystem)
 
@@ -185,6 +188,11 @@ void ULandscapeSubsystem::Tick(float DeltaTime)
 	Super::Tick(DeltaTime);
 
 #if WITH_EDITOR
+	{
+		TRACE_CPUPROFILER_EVENT_SCOPE(YOUPI);
+		AppCurrentDateTime = FDateTime::Now();
+	}
+
 	//Check if we need to start or stop creating Collision SceneProxies
 	ILandscapeModule& LandscapeModule = FModuleManager::GetModuleChecked<ILandscapeModule>("Landscape");
 	int32 NumViewsWithShowCollision = LandscapeModule.GetLandscapeSceneViewExtension()->GetNumViewsWithShowCollision();
@@ -559,6 +567,11 @@ void ULandscapeSubsystem::DisplayMessages(FCanvas* Canvas, float& XPos, float& Y
 	}
 }
 
-#endif
+FDateTime ULandscapeSubsystem::GetAppCurrentDateTime()
+{
+	return AppCurrentDateTime;
+}
+
+#endif // WITH_EDITOR
 
 #undef LOCTEXT_NAMESPACE
diff --git a/Engine/Source/Runtime/Landscape/Public/LandscapeSubsystem.h b/Engine/Source/Runtime/Landscape/Public/LandscapeSubsystem.h
--- a/Engine/Source/Runtime/Landscape/Public/LandscapeSubsystem.h
+++ b/Engine/Source/Runtime/Landscape/Public/LandscapeSubsystem.h
@@ -10,6 +10,7 @@
 class ALandscapeProxy;
 class ULandscapeInfo;
 class FLandscapeNotificationManager;
+struct FDateTime;
 
 #if WITH_EDITOR
 struct FOnHeightmapStreamedContext
@@ -86,6 +87,7 @@ public:
 	FLandscapeNotificationManager* GetNotificationManager() { return NotificationManager; }
 	FOnHeightmapStreamedDelegate& GetOnHeightmapStreamedDelegate() { return OnHeightmapStreamed; }
 	LANDSCAPE_API bool AnyViewShowCollisions() const { return bAnyViewShowCollisions; }  //! Returns true if any view has view collisions enabled.
+	FDateTime GetAppCurrentDateTime();
 #endif // WITH_EDITOR
 
 private:
@@ -109,5 +111,6 @@ private:
 	FLandscapeNotificationManager* NotificationManager;
 	FOnHeightmapStreamedDelegate OnHeightmapStreamed;
 	bool bAnyViewShowCollisions = false;
+	FDateTime AppCurrentDateTime; // Represents FDateTime::Now(), at the beginning of the frame (useful to get a human-readable date/time that is fixed during the frame)
 #endif // WITH_EDITOR
 };
