Build fix

diff --git a/Engine/Source/Runtime/Landscape/Private/LandscapeEditLayers.cpp b/Engine/Source/Runtime/Landscape/Private/LandscapeEditLayers.cpp
--- a/Engine/Source/Runtime/Landscape/Private/LandscapeEditLayers.cpp
+++ b/Engine/Source/Runtime/Landscape/Private/LandscapeEditLayers.cpp
@@ -776,7 +776,7 @@ public:
 	FLandscapeCopyTexturePS()
 	{}
 
-	void SetParameters(FRHIBatchedShaderParameters& BatchedParameters, FRHITexture* TextureRHI)
+	void SetParameters(FRHIBatchedShaderParameters& BatchedParameters, FRHITexture* InSourceTextureRHI, const FIntPoint& InSourcePosition = FIntPoint(0, 0), const FIntPoint& InCopySizePixels = FIntPoint(0, 0))
 	{
 		FVector2f SourceSize(InSourceTextureRHI->GetSizeXY());
 		FVector2f SourceOffsetUV = FVector2f(InSourcePosition) / SourceSize;
@@ -1596,7 +1596,7 @@ private:
 		GraphicsPSOInit.BoundShaderState.PixelShaderRHI = PixelShader.GetPixelShader();
 		SetGraphicsPipelineState(InRHICmdList, GraphicsPSOInit, 0);
 
-		PixelShader->SetParameters(InRHICmdList, Params.SourceResource->TextureRHI, Params.SourcePosition, Size);
+		SetAllShaderParametersPS(InRHICmdList, PixelShader, Params.SourceResource->TextureRHI, Params.SourcePosition, Size);
 
 		InRHICmdList.SetViewport((float)Params.DestPosition.X, (float)Params.DestPosition.Y, 0.0f, (float)(Params.DestPosition.X + Size.X), (float)(Params.DestPosition.Y + Size.Y), 1.0f);
 		InRHICmdList.DrawIndexedPrimitive(GTwoTrianglesIndexBuffer.IndexBufferRHI, 0, 0, 4, 0, 2, 1);
