Adds FBatchTextureCopy helper class to request and batch texture copies.
This class aims to minimze the number of lock/unlock mip for a texture source by aggregating same type copy operations (when sources/destinations are the same, channels are the only differences).
Use FBatchTextureCopy when generating mobile weightmaps.

diff --git a/Engine/Source/Runtime/Landscape/Private/LandscapeEdit.cpp b/Engine/Source/Runtime/Landscape/Private/LandscapeEdit.cpp
--- a/Engine/Source/Runtime/Landscape/Private/LandscapeEdit.cpp
+++ b/Engine/Source/Runtime/Landscape/Private/LandscapeEdit.cpp
@@ -6911,10 +6911,12 @@ void ULandscapeComponent::GenerateMobileWeightmapLayerAllocations()
 void ULandscapeComponent::GenerateMobilePlatformPixelData(bool bIsCooking, const ITargetPlatform* TargetPlatform)
 {
 	check(!IsTemplate());
+	TRACE_CPUPROFILER_EVENT_SCOPE(ULandscapeComponent::GenerateMobilePlatformPixelData);
 
 	GenerateMobileWeightmapLayerAllocations();
 
 	int32 WeightmapSize = (SubsectionSizeQuads + 1) * NumSubsections;
+	UE::Landscape::FBatchTextureCopy CopyRequests;
 
 	MobileWeightmapTextures.Empty();
 	{
@@ -6937,7 +6939,8 @@ void ULandscapeComponent::GenerateMobilePlatformPixelData(bool bIsCooking, const
 					MobileWeightmapTextures.Add(CurrentWeightmapTexture);
 				}
 
-				LandscapeData.CopyTextureFromWeightmap(CurrentWeightmapTexture, CurrentChannel, this, Allocation.LayerInfo);
+				CopyRequests.AddWeightmapCopy(CurrentWeightmapTexture, CurrentChannel, this, Allocation.LayerInfo);
+
 				// update Allocation
 				Allocation.WeightmapTextureIndex = MobileWeightmapTextures.Num() - 1;
 				Allocation.WeightmapTextureChannel = CurrentChannel;
@@ -6947,6 +6950,8 @@ void ULandscapeComponent::GenerateMobilePlatformPixelData(bool bIsCooking, const
 		}
 	}
 
+	CopyRequests.ProcessTextureCopies();
+
 	GDisableAutomaticTextureMaterialUpdateDependencies = true;
 	for (int TextureIdx = 0; TextureIdx < MobileWeightmapTextures.Num(); TextureIdx++)
 	{
diff --git a/Engine/Source/Runtime/Landscape/Private/LandscapeUtils.cpp b/Engine/Source/Runtime/Landscape/Private/LandscapeUtils.cpp
--- a/Engine/Source/Runtime/Landscape/Private/LandscapeUtils.cpp
+++ b/Engine/Source/Runtime/Landscape/Private/LandscapeUtils.cpp
@@ -7,9 +7,15 @@
 
 #if WITH_EDITOR
 #include "EditorDirectories.h"
+#include "Engine/Texture2D.h"
+#include "LandscapeComponent.h"
+#include "LandscapeLayerInfoObject.h"
 #include "ObjectTools.h"
 #endif
 
+// Channel remapping
+extern const size_t ChannelOffsets[4];
+
 namespace UE::Landscape
 {
 
@@ -64,6 +70,175 @@ FString GetLayerInfoObjectPackageName(const ULevel* InLevel, const FName& InLaye
 	return PackageName;
 }
 
+uint32 GetTypeHash(const FTextureCopyRequest& InKey)
+{
+	uint32 Hash = GetTypeHash(InKey.Source);
+	return HashCombine(Hash, GetTypeHash(InKey.Destination));
+}
+
+bool operator==(const FTextureCopyRequest& InEntryA, const FTextureCopyRequest& InEntryB)
+{
+	return (InEntryA.Source == InEntryB.Source) && (InEntryA.Destination == InEntryB.Destination);
+}
+
+bool FBatchTextureCopy::AddWeightmapCopy(UTexture2D* InDestination, int8 InDestinationChannel, const ULandscapeComponent* InComponent, ULandscapeLayerInfoObject* InLayerInfo)
+{
+	FTextureCopyRequest CopyRequest;
+	const TArray<UTexture2D*>& ComponentWeightmapTextures = InComponent->GetWeightmapTextures();
+	const TArray<FWeightmapLayerAllocationInfo>& ComponentWeightmapLayerAllocations = InComponent->GetWeightmapLayerAllocations();
+	int8 SourceChannel = INDEX_NONE;
+
+	CopyRequest.Destination = InDestination;
+
+	// Find the proper Source Texture and channel from Layer Allocations
+	for (const FWeightmapLayerAllocationInfo& ComponentWeightmapLayerAllocation : ComponentWeightmapLayerAllocations)
+	{
+		if (ComponentWeightmapLayerAllocation.LayerInfo == InLayerInfo)
+		{
+			CopyRequest.Source = ComponentWeightmapTextures[ComponentWeightmapLayerAllocation.WeightmapTextureIndex];
+			SourceChannel = ComponentWeightmapLayerAllocation.WeightmapTextureChannel;
+			break;
+		}
+	}
+
+	// Check if we found a proper allocation for this LayerInfo
+	if (SourceChannel != INDEX_NONE)
+	{
+		check((InDestinationChannel < 4) && (SourceChannel < 4));
+		FTextureCopyChannelMapping& ChannelMapping = CopyRequests.FindOrAdd(MoveTemp(CopyRequest));
+		ChannelMapping[ChannelOffsets[InDestinationChannel]] = ChannelOffsets[SourceChannel];
+		return true;
+	}
+
+	return false;
+}
+
+struct FSourceDataMipNumber
+{
+	TArray<const uint8*> SourceDataPtr;
+	int32 MipNumber = 0;
+};
+
+struct FDestinationDataMipNumber
+{
+	TArray<uint8*> DestinationDataPtr;
+	int32 MipNumber = 0;
+};
+
+bool FBatchTextureCopy::ProcessTextureCopies()
+{
+	TRACE_CPUPROFILER_EVENT_SCOPE(FBatchTextureCopy::ProcessTextureCopyRequest);
+	TMap<UTexture2D*, FSourceDataMipNumber> Sources;
+	TMap<UTexture2D*, FDestinationDataMipNumber> Destinations;
+
+	if (CopyRequests.Num() == 0)
+	{
+		return false;
+	}
+
+	// Populate source/destination maps to filter unique occurences
+	for (const TPair<FTextureCopyRequest, FTextureCopyChannelMapping>& CopyRequest : CopyRequests)
+	{
+		FSourceDataMipNumber& SourceData = Sources.Add(CopyRequest.Key.Source);
+		SourceData.MipNumber = CopyRequest.Key.Source->Source.GetNumMips();
+
+		FDestinationDataMipNumber& DestinationData = Destinations.Add(CopyRequest.Key.Destination);
+		DestinationData.MipNumber = CopyRequest.Key.Destination->Source.GetNumMips();
+	}
+
+	// Lock all sources mips
+	for (TPair<UTexture2D*, FSourceDataMipNumber>& Source : Sources)
+	{
+		int32 MipNumber = Source.Value.MipNumber;
+		TArray<const uint8*>& SourceDataPtr = Source.Value.SourceDataPtr;
+		
+		SourceDataPtr.Reserve(MipNumber);
+
+		for (int32 MipLevel = 0; MipLevel < MipNumber; ++MipLevel)
+		{
+			SourceDataPtr.Add(Source.Key->Source.LockMipReadOnly(MipLevel));
+		}
+	}
+
+	// Lock all destinations mips
+	for (TPair<UTexture2D*, FDestinationDataMipNumber>& Destination : Destinations)
+	{
+		int32 MipNumber = Destination.Value.MipNumber;
+		TArray<uint8*>& DestinationDataPtr = Destination.Value.DestinationDataPtr;
+
+		for (int32 MipLevel = 0; MipLevel < MipNumber; ++MipLevel)
+		{
+			DestinationDataPtr.Add(Destination.Key->Source.LockMip(MipLevel));
+		}
+	}
+
+	for (const TPair<FTextureCopyRequest, FTextureCopyChannelMapping>& CopyRequest : CopyRequests)
+	{
+		const FSourceDataMipNumber* SourceDataMipNumber = Sources.Find(CopyRequest.Key.Source);
+		const FDestinationDataMipNumber* DestinationDataMipNumber = Destinations.Find(CopyRequest.Key.Destination);
+
+		check((SourceDataMipNumber != nullptr) && (DestinationDataMipNumber != nullptr));
+		check(SourceDataMipNumber->MipNumber == DestinationDataMipNumber->MipNumber);
+
+		const int32 MipNumber = SourceDataMipNumber->MipNumber;
+
+		for (int32 MipLevel = 0; MipLevel < MipNumber; ++MipLevel)
+		{
+			const int32 MipSize = CopyRequest.Key.Destination->Source.GetSizeX() >> MipLevel;
+			check(MipSize == (CopyRequest.Key.Destination->Source.GetSizeY() >> MipLevel));
+
+			int32 MipSizeSquare = FMath::Square(MipSize);
+			const uint8* SourceTextureData = SourceDataMipNumber->SourceDataPtr[MipLevel];
+			uint8* DestTextureData = DestinationDataMipNumber->DestinationDataPtr[MipLevel];
+
+			check((SourceTextureData != nullptr) && (DestTextureData != nullptr));
+
+			const FTextureCopyChannelMapping& ChannelMapping = CopyRequest.Value;
+
+			// Perform the copy, redirecting channels using mappings
+			for (int32 Index = 0; Index < MipSizeSquare; ++Index)
+			{
+				int32 Base = Index * 4;
+
+				for (int32 Channel = 0; Channel < 4; ++Channel)
+				{
+					if (ChannelMapping[Channel] == INDEX_NONE)
+					{
+						continue;
+					}
+
+					DestTextureData[Base + Channel] = SourceTextureData[Base + ChannelMapping[Channel]];
+				}
+			}
+		}
+	}
+
+	// Unlock all sources mips
+	for (TPair<UTexture2D*, FSourceDataMipNumber>& Source : Sources)
+	{
+		int32 MipNumber = Source.Value.MipNumber;
+		
+		for (int32 MipLevel = 0; MipLevel < MipNumber; ++MipLevel)
+		{
+			Source.Key->Source.UnlockMip(MipLevel);
+		}
+	}
+	
+	// Unlock all destination mips
+	for (TPair<UTexture2D*, FDestinationDataMipNumber>& Destination : Destinations)
+	{
+		int32 MipNumber = Destination.Value.MipNumber;
+		
+		for (int32 MipLevel = 0; MipLevel < MipNumber; ++MipLevel)
+		{
+			Destination.Key->Source.UnlockMip(MipLevel);
+		}
+	}
+
+	return true;
+}
+
+
 #endif //!WITH_EDITOR
 
 } // end namespace UE::Landscape
\ No newline at end of file
diff --git a/Engine/Source/Runtime/Landscape/Public/LandscapeUtils.h b/Engine/Source/Runtime/Landscape/Public/LandscapeUtils.h
--- a/Engine/Source/Runtime/Landscape/Public/LandscapeUtils.h
+++ b/Engine/Source/Runtime/Landscape/Public/LandscapeUtils.h
@@ -11,6 +11,9 @@
 
 class ULevel;
 enum EShaderPlatform : uint16;
+class ULandscapeComponent;
+class ULandscapeLayerInfoObject;
+class UTexture2D;
 
 namespace UE::Landscape
 {
@@ -23,6 +26,52 @@ LANDSCAPE_API bool DoesPlatformSupportEditLayers(EShaderPlatform InShaderPlatfor
 
 #if WITH_EDITOR
 
+struct FTextureCopyRequest
+{
+	UTexture2D* Source = nullptr;
+	UTexture2D* Destination = nullptr;
+};
+
+uint32 GetTypeHash(const FTextureCopyRequest& InKey);
+bool operator==(const FTextureCopyRequest& InEntryA, const FTextureCopyRequest& InEntryB);
+
+/** Represents the DestinationChannel->SourceChannel binding.DestinationChannel is used as index.
+ *  For example if the source channel is 1 and the destination channel is 2, then Mappings[2] == 1.
+ */
+struct FTextureCopyChannelMapping
+{
+	FTextureCopyChannelMapping()
+		: Mappings{ INDEX_NONE, INDEX_NONE, INDEX_NONE, INDEX_NONE }
+	{}
+
+	int8& operator[](int32 Index) { return Mappings[Index]; }
+	const int8 operator[](int32 Index) const { return Mappings[Index]; }
+
+	int8 Mappings[4];
+};
+
+class LANDSCAPE_API FBatchTextureCopy
+{
+public:
+	/**
+	* Uses the provided arguments to add proper source/destination entries to internal copy requests.
+	* @param	InDestination	The texture used as a destination for the copy.
+	* @param	InDestinationChannel	The channel used as a destination for the copy.
+	* @param	InComponent		The component containing the wanted source weightmap.
+	* @param	InLayerInfo		The layer info used to retrieve the proper source weightmap and channel.
+	* @return True if the copy has been successfully added.
+	*/
+	bool AddWeightmapCopy(UTexture2D* InDestination, int8 InDestinationChannel, const ULandscapeComponent* InComponent, ULandscapeLayerInfoObject* InLayerInfo);
+
+	/** Process pending internal copy requests. */
+	bool ProcessTextureCopies();
+
+private:
+	using FTextureCopyChannelMappingMap = TMap<FTextureCopyRequest, FTextureCopyChannelMapping>;
+
+	FTextureCopyChannelMappingMap CopyRequests;
+};
+
 /**
  * Returns a generated path used for Landscape Shared Assets
  * @param	InPath	Path used as a basis to generate shared assets path. If /Temp/, it will be replaced by the last valid path used for level.
