Allow multiple copies of a Landscape via Level Instancing.

The main issue for runtime is that the landscape GUIDs collide when you pull in multiple copies of the same landscape.  This can cause errors to be spammed and other interesting side effects.  Here we detect an instanced situation on PostLoad, and deterministically generate a replacement landscape GUID.

There are other editor issues when combining Level Instancing with Landscapes, but I will address those in a separate change.  This change addresses the runtime issues.

Tested with PIE as well as standalone cooked.

diff --git a/Engine/Source/Runtime/Landscape/Private/Landscape.cpp b/Engine/Source/Runtime/Landscape/Private/Landscape.cpp
--- a/Engine/Source/Runtime/Landscape/Private/Landscape.cpp
+++ b/Engine/Source/Runtime/Landscape/Private/Landscape.cpp
@@ -72,6 +72,7 @@ Landscape.cpp: Terrain rendering
 #include "RenderCaptureInterface.h"
 #include "VisualLogger/VisualLogger.h"
 #include "LandscapeSettings.h"
+#include "Misc/ArchiveMD5.h"
 
 #if WITH_EDITOR
 #include "LandscapeEdit.h"
@@ -3322,10 +3323,50 @@ bool ULandscapeInfo::UpdateLayerInfoMap(ALandscapeProxy* Proxy /*= nullptr*/, bo
 
 #endif // WITH_EDITOR
 
+
+void FixupLandscapeGuidsIfInstanced(ALandscapeProxy* LandscapeProxy)
+{
+	// if the outer world is instanced, we need to change our landscape guid (in a deterministic way)
+	// this avoids guid collisions when you instance a world (and it's landscapes) multiple times,
+	// while maintaining the same GUID between landscape proxy objects within an instance
+	UWorld* OuterWorld = LandscapeProxy->GetTypedOuter<UWorld>();
+	UPackage* OuterPackage = OuterWorld->GetPackage();
+	if (OuterPackage->GetFName() != OuterPackage->GetLoadedPath().GetPackageFName())	// OuterWorld->IsInstanced()) is broken when PIE uses a memory built package
+	{
+		FArchiveMD5 Ar;
+		FGuid OldLandscapeGuid = LandscapeProxy->GetLandscapeGuid();
+		Ar << OldLandscapeGuid;
+
+#if WITH_EDITOR
+		// to work around PIE issues, we use the world partition to find the package
+		UWorldPartition* WorldPartition = FWorldPartitionHelpers::GetWorldPartition(LandscapeProxy);
+		if (WorldPartition)
+		{
+			OuterPackage = WorldPartition->GetPackage();
+			FName PackageName = OuterPackage->GetFName();
+			Ar << PackageName;
+		}
+		else
+#endif // WITH_EDITOR
+		{
+			FObjectKey ContainerKey(OuterPackage);
+			Ar << ContainerKey;
+		}
+
+		FMD5Hash MD5Hash;
+		Ar.GetHash(MD5Hash);
+
+		FGuid NewLandscapeGuid = MD5HashToGuid(MD5Hash);
+		LandscapeProxy->SetLandscapeGuid(NewLandscapeGuid);
+	}
+}
+
 void ALandscapeProxy::PostLoad()
 {
 	Super::PostLoad();
 
+	FixupLandscapeGuidsIfInstanced(this);
+
 	// Temporary
 	if (ComponentSizeQuads == 0 && LandscapeComponents.Num() > 0)
 	{
