Fix for regression introduced by 24097029:
All landscape heightmaps are modified when any heightmap is modified -- this restores previous behavior.

diff --git a/Engine/Source/Runtime/Landscape/Private/LandscapeEditLayers.cpp b/Engine/Source/Runtime/Landscape/Private/LandscapeEditLayers.cpp
--- a/Engine/Source/Runtime/Landscape/Private/LandscapeEditLayers.cpp
+++ b/Engine/Source/Runtime/Landscape/Private/LandscapeEditLayers.cpp
@@ -5412,18 +5412,21 @@ bool ALandscape::ResolveLayersTexture(
 					}
 					bChanged |= InCPUReadback->SetHash(Hash);
 
+					if (bChanged)
+					{
 						// We're about to modify the texture's source data, the texture needs to know so that it can handle properly update cached platform data (additionally, the package needs to be dirtied) :
 						if (GetDefault<ULandscapeSettings>()->LandscapeDirtyingMode == ELandscapeDirtyingMode::InLandscapeModeAndUserTriggeredChanges)
 						{
 							FLandscapeDirtyOnlyInModeScope Scope(GetLandscapeInfo(), !bUserTriggered);
-							GetLandscapeInfo()->ModifyObject(InOutputTexture);	
+							GetLandscapeInfo()->ModifyObject(InOutputTexture);
 						}
 						else
 						{
 							GetLandscapeInfo()->ModifyObject(InOutputTexture);
 						}
+					}
 				}
-							
+
 				if (bChanged)
 				{
 					if (bIsWeightmap)
