Adding LandscapeRayTracingStates getter/setter to FSceneViewStateInterface so that it's actually usable inside Engine after Renderer/Private was removed from Engine's include paths.

diff --git a/Engine/Source/Runtime/Engine/Public/SceneManagement.h b/Engine/Source/Runtime/Engine/Public/SceneManagement.h
--- a/Engine/Source/Runtime/Engine/Public/SceneManagement.h
+++ b/Engine/Source/Runtime/Engine/Public/SceneManagement.h
@@ -70,6 +70,7 @@ class FRayTracingGeometry;
 struct FViewMatrices;
 struct FEngineShowFlags;
 class FViewport;
+class FLandscapeRayTracingStateList;
 
 namespace UE { namespace Color { class FColorSpace; } }
 
@@ -285,6 +286,9 @@ PRAGMA_ENABLE_DEPRECATION_WARNINGS
 	* returns the path tracer sample count
 	*/
 	virtual uint32 GetPathTracingSampleCount() const = 0;
+
+	virtual void SetLandscapeRayTracingStates(TPimplPtr<FLandscapeRayTracingStateList>&& InLandscapeRayTracingStates) = 0;
+	virtual FLandscapeRayTracingStateList* GetLandscapeRayTracingStates() const = 0;
 #endif
 
 	/**
diff --git a/Engine/Source/Runtime/Landscape/Private/LandscapeRender.cpp b/Engine/Source/Runtime/Landscape/Private/LandscapeRender.cpp
--- a/Engine/Source/Runtime/Landscape/Private/LandscapeRender.cpp
+++ b/Engine/Source/Runtime/Landscape/Private/LandscapeRender.cpp
@@ -39,7 +39,6 @@ LandscapeRender.cpp: New terrain rendering
 #include "RHIStaticStates.h"
 #include "PrimitiveSceneInfo.h"
 #include "SceneView.h"
-#include "ScenePrivate.h"
 #include "LandscapeProxy.h"
 #include "HAL/LowLevelMemTracker.h"
 #include "MeshMaterialShader.h"
@@ -725,10 +724,10 @@ void FLandscapeSceneViewExtension::PreRenderView_RenderThread(FRDGBuilder& Graph
 	if (InView.State)
 	{
 		// Create the ray tracing state list class if necessary
-		FSceneViewState* ViewState = InView.State->GetConcreteViewState();
-		if (!ViewState->LandscapeRayTracingStates.IsValid())
+		FSceneViewStateInterface* ViewState = InView.State;
+		if (ViewState->GetLandscapeRayTracingStates() == nullptr)
 		{
-			ViewState->LandscapeRayTracingStates = MakePimpl<FLandscapeRayTracingStateList>();
+			ViewState->SetLandscapeRayTracingStates( MakePimpl<FLandscapeRayTracingStateList>() );
 		}
 	}
 #endif	// RHI_RAYTRACING
@@ -1297,11 +1296,11 @@ FLandscapeRayTracingState* FLandscapeRayTracingImpl::FindOrCreateRayTracingState
 {
 	// Default view key of zero if there's no ViewStateInterface provided.  View keys start at 1, so 0 wouldn't be a valid key on an actual view.
 	uint32 ViewKey = 0;
-	FSceneViewState* ViewState = nullptr;
+	FSceneViewStateInterface* ViewState = nullptr;
 	if (ViewStateInterface)
 	{
-		ViewState = ViewStateInterface->GetConcreteViewState();
-		ViewKey = ViewState->UniqueID;
+		ViewState = ViewStateInterface;
+		ViewKey = ViewState->GetViewKey();
 	}
 
 	// Check for existing state for this view.  We're just doing a linear search of the array, because practical applications won't have
@@ -1326,7 +1325,7 @@ FLandscapeRayTracingState* FLandscapeRayTracingImpl::FindOrCreateRayTracingState
 	if (ViewState)
 	{
 		// Link into the view state's linked list, so it can be cleaned up if the view gets deleted
-		FLandscapeRayTracingStateList* StateList = ViewState->LandscapeRayTracingStates.Get();
+		FLandscapeRayTracingStateList* StateList = ViewState->GetLandscapeRayTracingStates();
 		check(StateList);
 
 		if (StateList->ListHead)
diff --git a/Engine/Source/Runtime/Renderer/Private/ScenePrivate.h b/Engine/Source/Runtime/Renderer/Private/ScenePrivate.h
--- a/Engine/Source/Runtime/Renderer/Private/ScenePrivate.h
+++ b/Engine/Source/Runtime/Renderer/Private/ScenePrivate.h
@@ -84,7 +84,6 @@ class FDistanceFieldObjectBuffers;
 struct FHairStrandsInstance;
 struct FPathTracingState;
 class FSparseVolumeTextureViewerSceneProxy;
-class FLandscapeRayTracingStateList;
 class FExponentialHeightFogSceneInfo;
 class FStaticMeshBatch;
 class FShadowScene;
@@ -1013,6 +1012,9 @@ public:
 
 	// List of landscape ray tracing state associated with this view, so it can be cleaned up if the view gets deleted.
 	TPimplPtr<FLandscapeRayTracingStateList> LandscapeRayTracingStates;
+
+	virtual void SetLandscapeRayTracingStates(TPimplPtr<FLandscapeRayTracingStateList>&& InLandscapeRayTracingStates) final { LandscapeRayTracingStates = MoveTemp(InLandscapeRayTracingStates); }
+	virtual FLandscapeRayTracingStateList* GetLandscapeRayTracingStates() const final { return LandscapeRayTracingStates.Get(); }
 #endif
 
 	TUniquePtr<FForwardLightingViewResources> ForwardLightingResources;
