Allow landscape edits to use a local explicit edit layer, instead of the shared edit layer state on ALandscape.

The core changes are in FLandscapeEditDataInterface (LandscapeEdit.h, LandscapeEditInterface.cpp)
That class was changed so that it can go between using a local state and the shared state on ALandscape.
Any existing code using that class will end up using the shared state, leading to the same behavior as before (yay back compat).
A function GetEditLayer() was added to FLandscapeEditDataInterface, to return the appropriate edit layer.

Some variants of Component functions were created that take as input an explicit layer GUID, rather than automatically looking up the shared state, so that
FLandscapeEditDataInterface could use them to work on specific edit layers.

These functions are:
SetWeightmapTexturesInternal()
SetWeightmapTexturesUsageInternal()
ReallocateWeightmapsInternal()
ReplaceLayerInternal()

The actual tool updates to make them use local explicit edit layers are in the second change.

diff --git a/Engine/Source/Runtime/Landscape/Classes/LandscapeComponent.h b/Engine/Source/Runtime/Landscape/Classes/LandscapeComponent.h
--- a/Engine/Source/Runtime/Landscape/Classes/LandscapeComponent.h
+++ b/Engine/Source/Runtime/Landscape/Classes/LandscapeComponent.h
@@ -713,12 +713,15 @@ public:
 
 	LANDSCAPE_API void SetHeightmap(UTexture2D* NewHeightmap);
 	LANDSCAPE_API void SetWeightmapTextures(const TArray<UTexture2D*>& InNewWeightmapTextures, bool InApplyToEditingWeightmap = false);
+	void SetWeightmapTexturesInternal(const TArray<UTexture2D*>& InNewWeightmapTextures, const FGuid& InEditLayerGuid);
 
 #if WITH_EDITOR
 	LANDSCAPE_API void SetWeightmapLayerAllocations(const TArray<FWeightmapLayerAllocationInfo>& InNewWeightmapLayerAllocations);
 	LANDSCAPE_API uint32 ComputeLayerHash(bool InReturnEditingHash = true) const;
 
 	LANDSCAPE_API void SetWeightmapTexturesUsage(const TArray<ULandscapeWeightmapUsage*>& InNewWeightmapTexturesUsage, bool InApplyToEditingWeightmap = false);
+	void SetWeightmapTexturesUsageInternal(const TArray<ULandscapeWeightmapUsage*>& InNewWeightmapTexturesUsage, const FGuid& InEditLayerGuid);
+
 	LANDSCAPE_API TArray<ULandscapeWeightmapUsage*>& GetWeightmapTexturesUsage(bool InReturnEditingWeightmap = false);
 	LANDSCAPE_API const TArray<ULandscapeWeightmapUsage*>& GetWeightmapTexturesUsage(bool InReturnEditingWeightmap = false) const;
 	LANDSCAPE_API TArray<ULandscapeWeightmapUsage*>& GetWeightmapTexturesUsage(const FGuid& InLayerGuid);
@@ -749,7 +752,7 @@ public:
 	bool GetPendingCollisionDataUpdate() const { return bPendingCollisionDataUpdate; }
 	void SetPendingLayerCollisionDataUpdate(bool bInPendingLayerCollisionDataUpdate) { bPendingLayerCollisionDataUpdate = bInPendingLayerCollisionDataUpdate; }
 	bool GetPendingLayerCollisionDataUpdate() const { return bPendingLayerCollisionDataUpdate; }
-#endif 
+#endif // WITH_EDITOR
 
 	virtual bool IsShown(const FEngineShowFlags& ShowFlags) const override;
 
@@ -788,16 +791,17 @@ public:
 	LANDSCAPE_API void DeleteLayer(ULandscapeLayerInfoObject* LayerInfo, FLandscapeEditDataInterface& LandscapeEdit);
 	
 	/** Deletes a material layer from the specified edit layer on this component, removing all its data, adjusting other layer's weightmaps if necessary, etc. */
-	void DeleteLayerInternal(ULandscapeLayerInfoObject* LayerInfo, FLandscapeEditDataInterface& LandscapeEdit, const FGuid& EditLayerGuid);
+	void DeleteLayerInternal(ULandscapeLayerInfoObject* LayerInfo, FLandscapeEditDataInterface& LandscapeEdit, const FGuid& InEditLayerGuid);
 
 	/** Deletes a layer from this component, but doesn't do anything else (assumes the user knows what he's doing, use DeleteLayer otherwise) */
 	void DeleteLayerAllocation(const FGuid& InEditLayerGuid, int32 InLayerAllocationIdx, bool bInShouldDirtyPackage);
 
-	/** Fills a layer to 100% on this component, adding it if needed and removing other layers that get painted away */
+	/** Fills a layer to 100% on this component, adding it if needed and removing other layers that get painted away.  Uses the edit layer specified by LandscapeEdit. */
 	LANDSCAPE_API void FillLayer(ULandscapeLayerInfoObject* LayerInfo, FLandscapeEditDataInterface& LandscapeEdit);
 
 	/** Replaces one layerinfo on this component with another */
 	LANDSCAPE_API void ReplaceLayer(ULandscapeLayerInfoObject* FromLayerInfo, ULandscapeLayerInfoObject* ToLayerInfo, FLandscapeEditDataInterface& LandscapeEdit);
+	void ReplaceLayerInternal(ULandscapeLayerInfoObject* FromLayerInfo, ULandscapeLayerInfoObject* ToLayerInfo, FLandscapeEditDataInterface& LandscapeEdit, const FGuid& InEditLayerGUID);
 
 	// true if the component's landscape material supports grass
 	bool MaterialHasGrass() const;
@@ -1013,10 +1017,15 @@ public:
 	void UpdateCollisionPhysicalMaterialData(TArray<UPhysicalMaterial*> const& InPhysicalMaterials, TArray<uint8> const& InMaterialIds);
 
 	/**
-	 * Create weightmaps for this component for the layers specified in the WeightmapLayerAllocations array
+	 * Create weightmaps for this component for the layers specified in the WeightmapLayerAllocations array, works in the landscape current edit layer when InCanUseEditingWeightmap is true
 	 */
 	LANDSCAPE_API void ReallocateWeightmaps(FLandscapeEditDataInterface* DataInterface = nullptr, bool InCanUseEditingWeightmap = true, bool InSaveToTransactionBuffer = true, bool InForceReallocate = false, ALandscapeProxy* InTargetProxy = nullptr, TArray<UTexture*>* OutNewCreatedTextures = nullptr);
 
+	/**
+	 * Create weightmaps for this component for the layers specified in the WeightmapLayerAllocations array, works in the specified edit layer
+	 */
+	void ReallocateWeightmapsInternal(FLandscapeEditDataInterface* DataInterface = nullptr, const FGuid& InEditLayerGuid = FGuid(), bool InSaveToTransactionBuffer = true, bool InForceReallocate = false, ALandscapeProxy* InTargetProxy = nullptr, TArray<UTexture*>* OutNewCreatedTextures = nullptr);
+
 	/** Returns the component's LandscapeMaterial, or the Component's OverrideLandscapeMaterial if set */
 	LANDSCAPE_API UMaterialInterface* GetLandscapeMaterial(int8 InLODIndex = INDEX_NONE) const;
 
diff --git a/Engine/Source/Runtime/Landscape/Private/Landscape.cpp b/Engine/Source/Runtime/Landscape/Private/Landscape.cpp
--- a/Engine/Source/Runtime/Landscape/Private/Landscape.cpp
+++ b/Engine/Source/Runtime/Landscape/Private/Landscape.cpp
@@ -1948,13 +1948,13 @@ TArray<FWeightmapLayerAllocationInfo>& ULandscapeComponent::GetWeightmapLayerAll
 
 FLandscapeLayerComponentData* ULandscapeComponent::GetEditingLayer()
 {
-	const FGuid& EditingLayerGuid = GetLandscapeActor()->EditingLayer;
+	const FGuid& EditingLayerGuid = GetLandscapeActor()->GetEditingLayer();
 	return EditingLayerGuid.IsValid() ? LayersData.Find(EditingLayerGuid) : nullptr;
 }
 
 const FLandscapeLayerComponentData* ULandscapeComponent::GetEditingLayer() const
 {
-	const FGuid& EditingLayerGuid = GetLandscapeActor()->EditingLayer;
+	const FGuid& EditingLayerGuid = GetLandscapeActor()->GetEditingLayer();
 	return EditingLayerGuid.IsValid() ? LayersData.Find(EditingLayerGuid) : nullptr;
 }
 
@@ -2141,6 +2141,27 @@ void ULandscapeComponent::SetWeightmapTextures(const TArray<UTexture2D*>& InNewW
 	}
 }
 
+// Note that there is a slight difference in behavior with the Internal function:
+// unlike SetWeightmapTextures, this function will never set the runtime WeightmapTextures when you intended to set an edit layer's WeightmapData.Textures
+void ULandscapeComponent::SetWeightmapTexturesInternal(const TArray<UTexture2D*>& InNewWeightmapTextures, const FGuid& InEditLayerGuid)
+{
+	if (InEditLayerGuid.IsValid())
+	{
+#if WITH_EDITOR
+		FLandscapeLayerComponentData* EditingLayer = GetLayerData(InEditLayerGuid);
+		if (ensure(EditingLayer))
+		{
+			EditingLayer->WeightmapData.Textures.Reset(InNewWeightmapTextures.Num());
+			EditingLayer->WeightmapData.Textures.Append(InNewWeightmapTextures);
+		}
+#endif // WITH_EDITOR
+	}
+	else
+	{
+		WeightmapTextures = InNewWeightmapTextures;
+	}
+}
+
 
 #if WITH_EDITOR
 void ULandscapeComponent::SetWeightmapLayerAllocations(const TArray<FWeightmapLayerAllocationInfo>& InNewWeightmapLayerAllocations)
@@ -2215,6 +2236,25 @@ void ULandscapeComponent::SetWeightmapTexturesUsage(const TArray<ULandscapeWeigh
 	}
 }
 
+void ULandscapeComponent::SetWeightmapTexturesUsageInternal(const TArray<ULandscapeWeightmapUsage*>& InNewWeightmapTexturesUsage, const FGuid& InEditLayerGuid)
+{
+	if (InEditLayerGuid.IsValid())
+	{
+#if WITH_EDITOR
+		FLandscapeLayerComponentData* EditingLayer = GetLayerData(InEditLayerGuid);
+		if (ensure(EditingLayer))
+		{
+			EditingLayer->WeightmapData.TextureUsages.Reset(InNewWeightmapTexturesUsage.Num());
+			EditingLayer->WeightmapData.TextureUsages.Append(InNewWeightmapTexturesUsage);
+		}
+#endif // WITH_EDITOR
+	}
+	else
+	{
+		WeightmapTexturesUsage = InNewWeightmapTexturesUsage;
+	}
+}
+
 void ULandscapeComponent::DeleteLayerAllocation(const FGuid& InEditLayerGuid, int32 InLayerAllocationIdx, bool bInShouldDirtyPackage)
 {
 	TArray<FWeightmapLayerAllocationInfo>& ComponentWeightmapLayerAllocations = GetWeightmapLayerAllocations(InEditLayerGuid);
@@ -5517,6 +5557,6 @@ void FLandscapeProxyComponentDataChangedParams::ForEachComponent(TFunctionRef<vo
 		Func(Component);
 	}
 }
-#endif
+#endif // WITH_EDITOR
 
 #undef LOCTEXT_NAMESPACE
diff --git a/Engine/Source/Runtime/Landscape/Private/LandscapeEdit.cpp b/Engine/Source/Runtime/Landscape/Private/LandscapeEdit.cpp
--- a/Engine/Source/Runtime/Landscape/Private/LandscapeEdit.cpp
+++ b/Engine/Source/Runtime/Landscape/Private/LandscapeEdit.cpp
@@ -6278,17 +6278,24 @@ void ULandscapeInfo::ClearSelectedRegion(bool bIsComponentwise /*= true*/)
 }
 
 void ULandscapeComponent::ReallocateWeightmaps(FLandscapeEditDataInterface* DataInterface, bool InCanUseEditingWeightmap, bool InSaveToTransactionBuffer, bool InForceReallocate, ALandscapeProxy* InTargetProxy, TArray<UTexture*>* OutNewCreatedTextures)
+{
+	FGuid TargetLayerGuid = InCanUseEditingWeightmap ? GetEditingLayerGUID() : FGuid();
+	ReallocateWeightmapsInternal(DataInterface, TargetLayerGuid, InSaveToTransactionBuffer, InForceReallocate, InTargetProxy, OutNewCreatedTextures);
+}
+
+void ULandscapeComponent::ReallocateWeightmapsInternal(FLandscapeEditDataInterface* DataInterface, const FGuid& InEditLayerGuid, bool InSaveToTransactionBuffer, bool InForceReallocate, ALandscapeProxy* InTargetProxy, TArray<UTexture*>* OutNewCreatedTextures)
 {
 	int32 NeededNewChannels = 0;
 	ALandscapeProxy* TargetProxy = InTargetProxy ? InTargetProxy : GetLandscapeProxy();
+	
+	if (InEditLayerGuid.IsValid())
+	{
+		check(TargetProxy->HasLayersContent());
+	}
 
-	FGuid EditingLayerGUID = GetEditingLayerGUID();
-	check(!TargetProxy->HasLayersContent() || !InCanUseEditingWeightmap || EditingLayerGUID.IsValid());
-	FGuid TargetLayerGuid = InCanUseEditingWeightmap ? EditingLayerGUID : FGuid();
-
-	TArray<FWeightmapLayerAllocationInfo>& ComponentWeightmapLayerAllocations = GetWeightmapLayerAllocations(InCanUseEditingWeightmap);
-	TArray<UTexture2D*>& ComponentWeightmapTextures = GetWeightmapTextures(InCanUseEditingWeightmap);
-	TArray<ULandscapeWeightmapUsage*>& ComponentWeightmapTexturesUsage = GetWeightmapTexturesUsage(InCanUseEditingWeightmap);
+	TArray<FWeightmapLayerAllocationInfo>& ComponentWeightmapLayerAllocations = GetWeightmapLayerAllocations(InEditLayerGuid);
+	TArray<UTexture2D*>& ComponentWeightmapTextures = GetWeightmapTextures(InEditLayerGuid);
+	TArray<ULandscapeWeightmapUsage*>& ComponentWeightmapTexturesUsage = GetWeightmapTexturesUsage(InEditLayerGuid);
 
 	// When force reallocating, skip tests to see if allocations are necessary based on Component's WeightmapLayeAllocInfo
 	if (!InForceReallocate)
@@ -6328,7 +6335,7 @@ void ULandscapeComponent::ReallocateWeightmaps(FLandscapeEditDataInterface* Data
 		{
 			ULandscapeWeightmapUsage* Usage = ComponentWeightmapTexturesUsage[TexIdx];
 			check(Usage);
-			check(Usage->LayerGuid == TargetLayerGuid);
+			check(Usage->LayerGuid == InEditLayerGuid);
 			ExistingTexAvailableChannels += Usage->FreeChannelCount();
 
 			if (ExistingTexAvailableChannels >= NeededNewChannels)
@@ -6415,7 +6422,7 @@ void ULandscapeComponent::ReallocateWeightmaps(FLandscapeEditDataInterface* Data
 			{
 				ULandscapeWeightmapUsage* TryWeightmapUsage = ItPair.Value;
 				//
-				if (TryWeightmapUsage->FreeChannelCount() >= TotalNeededChannels && TryWeightmapUsage->LayerGuid == TargetLayerGuid)
+				if (TryWeightmapUsage->FreeChannelCount() >= TotalNeededChannels && TryWeightmapUsage->LayerGuid == InEditLayerGuid)
 				{
 					if (TryWeightmapUsage->IsEmpty())
 					{
@@ -6468,7 +6475,7 @@ void ULandscapeComponent::ReallocateWeightmaps(FLandscapeEditDataInterface* Data
 
 			// Store it in the usage map
 			CurrentWeightmapUsage = TargetProxy->WeightmapUsageMap.Add(CurrentWeightmapTexture, TargetProxy->CreateWeightmapUsage());
-			CurrentWeightmapUsage->LayerGuid = TargetLayerGuid;
+			CurrentWeightmapUsage->LayerGuid = InEditLayerGuid;
 			// UE_LOG(LogLandscape, Log, TEXT("Making a new texture %s"), *CurrentWeightmapTexture->GetName());
 		}
 
@@ -6540,8 +6547,8 @@ void ULandscapeComponent::ReallocateWeightmaps(FLandscapeEditDataInterface* Data
 	}
 
 	// Replace the weightmap textures
-	SetWeightmapTextures(MoveTemp(NewWeightmapTextures), InCanUseEditingWeightmap);
-	SetWeightmapTexturesUsage(MoveTemp(NewComponentWeightmapTexturesUsage), InCanUseEditingWeightmap);	
+	SetWeightmapTexturesInternal(MoveTemp(NewWeightmapTextures), InEditLayerGuid);
+	SetWeightmapTexturesUsageInternal(MoveTemp(NewComponentWeightmapTexturesUsage), InEditLayerGuid);
 
 	TargetProxy->ValidateProxyLayersWeightmapUsage();
 }
diff --git a/Engine/Source/Runtime/Landscape/Private/LandscapeEditInterface.cpp b/Engine/Source/Runtime/Landscape/Private/LandscapeEditInterface.cpp
--- a/Engine/Source/Runtime/Landscape/Private/LandscapeEditInterface.cpp
+++ b/Engine/Source/Runtime/Landscape/Private/LandscapeEditInterface.cpp
@@ -32,6 +32,8 @@ extern const size_t ChannelOffsets[4] = {STRUCT_OFFSET(FColor,R), STRUCT_OFFSET(
 //
 FLandscapeEditDataInterface::FLandscapeEditDataInterface(ULandscapeInfo* InLandscapeInfo, bool bInUploadTextureChangesToGPU)
 	: FLandscapeTextureDataInterface(bInUploadTextureChangesToGPU)
+	, bUseSharedLandscapeEditLayer(true)
+	, LocalEditLayerGUID()
 {
 	if (InLandscapeInfo)
 	{
@@ -43,6 +45,43 @@ FLandscapeEditDataInterface::FLandscapeEditDataInterface(ULandscapeInfo* InLands
 	}
 }
 
+
+FLandscapeEditDataInterface::FLandscapeEditDataInterface(ULandscapeInfo* InLandscapeInfo, const FGuid& InEditLayerGUID, bool bInUploadTextureChangesToGPU)
+	: FLandscapeTextureDataInterface(bInUploadTextureChangesToGPU),
+	bUseSharedLandscapeEditLayer(false),
+	LocalEditLayerGUID(InEditLayerGUID)
+
+{
+	if (InLandscapeInfo)
+	{
+		LandscapeInfo = InLandscapeInfo;
+		ComponentSizeQuads = InLandscapeInfo->ComponentSizeQuads;
+		SubsectionSizeQuads = InLandscapeInfo->SubsectionSizeQuads;
+		ComponentNumSubsections = InLandscapeInfo->ComponentNumSubsections;
+		DrawScale = InLandscapeInfo->DrawScale;
+	}
+}
+
+void FLandscapeEditDataInterface::SetEditLayer(const FGuid& InEditLayerGUID)
+{
+	bUseSharedLandscapeEditLayer = false;
+	LocalEditLayerGUID = InEditLayerGUID;
+}
+
+FGuid FLandscapeEditDataInterface::GetEditLayer() const
+{
+	if (bUseSharedLandscapeEditLayer)
+	{
+		FGuid LandscapeEditLayerGuid;
+		if ((LandscapeInfo != nullptr) && LandscapeInfo->LandscapeActor.IsValid())
+		{
+			LandscapeEditLayerGuid = LandscapeInfo->LandscapeActor->GetEditingLayer();
+		}
+		return LandscapeEditLayerGuid;
+	}
+	return LocalEditLayerGUID;
+}
+
 FLandscapeTextureDataInterface::FLandscapeTextureDataInterface(bool bInUploadTextureChangesToGPU)
 	: bUploadTextureChangesToGPU(bInUploadTextureChangesToGPU), bShouldDirtyPackage(true)
 {}
@@ -295,7 +334,7 @@ void FLandscapeEditDataInterface::SetHeightData(int32 X1, int32 Y1, int32 X2, in
 				}
 			}
 
-			UTexture2D* Heightmap = InHeightmap != nullptr ? InHeightmap : Component->GetHeightmap(true);
+			UTexture2D* Heightmap = InHeightmap != nullptr ? InHeightmap : Component->GetHeightmap(GetEditLayer());
 			UTexture2D* XYOffsetmapTexture = InXYOffsetmapTexture != nullptr ? InXYOffsetmapTexture : ToRawPtr(Component->XYOffsetmapTexture);
 
 			Component->Modify(GetShouldDirtyPackage());
@@ -487,6 +526,11 @@ void FLandscapeEditDataInterface::SetHeightData(int32 X1, int32 Y1, int32 X2, in
 void FLandscapeEditDataInterface::RecalculateNormals()
 {
 	if (!LandscapeInfo) return;
+
+	// we should only ever be calculating normals on the final runtime layer.
+	// if this gets hit, then we are wasting time calculating them on edit layers.
+	ensure(!GetEditLayer().IsValid());
+
 	// Recalculate normals for each component in turn
 	for( auto It = LandscapeInfo->XYtoComponentMap.CreateIterator(); It; ++It )
 	{
@@ -533,12 +577,13 @@ void FLandscapeEditDataInterface::RecalculateNormals()
 		}
 
 		// Find the texture data corresponding to this vertex
-		int32 SizeU = Component->GetHeightmap(true)->Source.GetSizeX();
-		int32 SizeV = Component->GetHeightmap(true)->Source.GetSizeY();
+		UTexture2D* EditHeightmap = Component->GetHeightmap(GetEditLayer());
+		int32 SizeU = EditHeightmap->Source.GetSizeX();
+		int32 SizeV = EditHeightmap->Source.GetSizeY();
 		int32 HeightmapOffsetX = static_cast<int32>(Component->HeightmapScaleBias.Z * SizeU);
 		int32 HeightmapOffsetY = static_cast<int32>(Component->HeightmapScaleBias.W * SizeV);
 
-		FLandscapeTextureDataInfo* TexDataInfo = GetTextureDataInfo(Component->GetHeightmap(true));
+		FLandscapeTextureDataInfo* TexDataInfo = GetTextureDataInfo(EditHeightmap);
 		FColor* HeightmapTextureData = (FColor*)TexDataInfo->GetMipData(0);
 
 		// Apply vertex normals to the component
@@ -615,7 +660,7 @@ void FLandscapeEditDataInterface::GetHeightDataTemplFast(const int32 X1, const i
 				continue;
 			}
 
-			UTexture2D* Heightmap = InHeightmap != nullptr ? InHeightmap : Component->GetHeightmap(true);
+			UTexture2D* Heightmap = InHeightmap != nullptr ? InHeightmap : Component->GetHeightmap(GetEditLayer());
 
 			FLandscapeTextureDataInfo* TexDataInfo = NULL;
 			FColor* HeightmapTextureData = NULL;
@@ -959,20 +1004,22 @@ void FLandscapeEditDataInterface::CalcMissingValues(const int32& X1, const int32
 
 FColor& FLandscapeEditDataInterface::GetHeightMapColor(const ULandscapeComponent* Component, int32 TexU, int32 TexV, FColor* TextureData)
 {
+	UTexture2D* Heightmap = Component->GetHeightmap(GetEditLayer());
 	check(Component);
 	if (!TextureData)
 	{
-		FLandscapeTextureDataInfo* TexDataInfo = GetTextureDataInfo(Component->GetHeightmap(true));
+		FLandscapeTextureDataInfo* TexDataInfo = GetTextureDataInfo(Heightmap);
 		TextureData = (FColor*)TexDataInfo->GetMipData(0);	
 	}
 
 	// All Heightmaps of component have the same texture size
-	const FTextureSource& HeightmapTextureSource = Component->GetHeightmap()->Source;
+	const FTextureSource& HeightmapTextureSource = Heightmap->Source;
 	const int32 SizeU = HeightmapTextureSource.GetSizeX();
 	const int32 SizeV = HeightmapTextureSource.GetSizeY();
 	const int32 HeightmapOffsetX = static_cast<int32>(Component->HeightmapScaleBias.Z * SizeU);
 	const int32 HeightmapOffsetY = static_cast<int32>(Component->HeightmapScaleBias.W * SizeV);
 
+
 	const int32 TexX = HeightmapOffsetX + TexU;
 	const int32 TexY = HeightmapOffsetY + TexV;
 	FColor& TexData = TextureData[ TexX + TexY * SizeU ];
@@ -1049,7 +1096,7 @@ void FLandscapeEditDataInterface::GetHeightDataInternal(int32& ValidX1, int32& V
 
 			if( Component )
 			{
-				TexDataInfo = GetTextureDataInfo(Component->GetHeightmap(true));
+				TexDataInfo = GetTextureDataInfo(Component->GetHeightmap(GetEditLayer()));
 				HeightmapTextureData = (FColor*)TexDataInfo->GetMipData(0);
 				ComponentDataExist[ComponentIndexXY] = true;
 				// Update valid region
@@ -1086,7 +1133,7 @@ void FLandscapeEditDataInterface::GetHeightDataInternal(int32& ValidX1, int32& V
 						if (BorderComponent[0])
 						{
 							NoBorderX1 = false;
-							NeighborTexDataInfo[0] = GetTextureDataInfo(BorderComponent[0]->GetHeightmap(true));
+							NeighborTexDataInfo[0] = GetTextureDataInfo(BorderComponent[0]->GetHeightmap(GetEditLayer()));
 							NeighborHeightmapTextureData[0] = (FColor*)NeighborTexDataInfo[0]->GetMipData(0);
 							break;
 						}
@@ -1102,7 +1149,7 @@ void FLandscapeEditDataInterface::GetHeightDataInternal(int32& ValidX1, int32& V
 						if (BorderComponent[1])
 						{
 							NoBorderX2 = false;
-							NeighborTexDataInfo[1] = GetTextureDataInfo(BorderComponent[1]->GetHeightmap(true));
+							NeighborTexDataInfo[1] = GetTextureDataInfo(BorderComponent[1]->GetHeightmap(GetEditLayer()));
 							NeighborHeightmapTextureData[1] = (FColor*)NeighborTexDataInfo[1]->GetMipData(0);
 							break;
 						}
@@ -1118,7 +1165,7 @@ void FLandscapeEditDataInterface::GetHeightDataInternal(int32& ValidX1, int32& V
 						if (BorderComponent[2])
 						{
 							NoBorderY1[ComponentIndexXX] = false;
-							NeighborTexDataInfo[2] = GetTextureDataInfo(BorderComponent[2]->GetHeightmap(true));
+							NeighborTexDataInfo[2] = GetTextureDataInfo(BorderComponent[2]->GetHeightmap(GetEditLayer()));
 							NeighborHeightmapTextureData[2] = (FColor*)NeighborTexDataInfo[2]->GetMipData(0);
 							break;
 						}
@@ -1129,7 +1176,7 @@ void FLandscapeEditDataInterface::GetHeightDataInternal(int32& ValidX1, int32& V
 					BorderComponent[2] = BorderComponentY1[ComponentIndexXX];
 					if (BorderComponent[2])
 					{
-						NeighborTexDataInfo[2] = GetTextureDataInfo(BorderComponent[2]->GetHeightmap(true));
+						NeighborTexDataInfo[2] = GetTextureDataInfo(BorderComponent[2]->GetHeightmap(GetEditLayer()));
 						NeighborHeightmapTextureData[2] = (FColor*)NeighborTexDataInfo[2]->GetMipData(0);
 					}
 				}
@@ -1143,7 +1190,7 @@ void FLandscapeEditDataInterface::GetHeightDataInternal(int32& ValidX1, int32& V
 						if (BorderComponent[3])
 						{
 							NoBorderY2[ComponentIndexXX] = false;
-							NeighborTexDataInfo[3] = GetTextureDataInfo(BorderComponent[3]->GetHeightmap(true));
+							NeighborTexDataInfo[3] = GetTextureDataInfo(BorderComponent[3]->GetHeightmap(GetEditLayer()));
 							NeighborHeightmapTextureData[3] = (FColor*)NeighborTexDataInfo[3]->GetMipData(0);
 							break;
 						}
@@ -1154,7 +1201,7 @@ void FLandscapeEditDataInterface::GetHeightDataInternal(int32& ValidX1, int32& V
 					BorderComponent[3] = BorderComponentY2[ComponentIndexXX];
 					if (BorderComponent[3])
 					{
-						NeighborTexDataInfo[3] = GetTextureDataInfo(BorderComponent[3]->GetHeightmap(true));
+						NeighborTexDataInfo[3] = GetTextureDataInfo(BorderComponent[3]->GetHeightmap(GetEditLayer()));
 						NeighborHeightmapTextureData[3] = (FColor*)NeighborTexDataInfo[3]->GetMipData(0);
 					}
 				}
@@ -1587,15 +1634,15 @@ void ULandscapeComponent::DeleteLayer(ULandscapeLayerInfoObject* LayerInfo, FLan
 	DeleteLayerInternal(LayerInfo, LandscapeEdit, EditLayerGuid);
 }
 
-void ULandscapeComponent::DeleteLayerInternal(ULandscapeLayerInfoObject* LayerInfo, FLandscapeEditDataInterface& LandscapeEdit, const FGuid &EditLayerGuid)
+void ULandscapeComponent::DeleteLayerInternal(ULandscapeLayerInfoObject* LayerInfo, FLandscapeEditDataInterface& LandscapeEdit, const FGuid& InEditLayerGuid)
 {
 	TRACE_CPUPROFILER_EVENT_SCOPE(LandscapeComponent_DeleteLayer);
 
 	ULandscapeComponent* Component = this;
 	
 	// This can be called during WeightmapFixup, so skip checking for updated weightmap allocations
-	const TArray<FWeightmapLayerAllocationInfo>& ComponentWeightmapLayerAllocations = Component->GetWeightmapLayerAllocations(EditLayerGuid);
-	const TArray<UTexture2D*>& ComponentWeightmapTextures = Component->GetWeightmapTextures(EditLayerGuid);
+	const TArray<FWeightmapLayerAllocationInfo>& ComponentWeightmapLayerAllocations = Component->GetWeightmapLayerAllocations(InEditLayerGuid);
+	const TArray<UTexture2D*>& ComponentWeightmapTextures = Component->GetWeightmapTextures(InEditLayerGuid);
 
 	// Delete the first LayerAllocation with a matching LayerInfo
 	const int32 DeleteAllocIdx = ComponentWeightmapLayerAllocations.IndexOfByPredicate(
@@ -1606,7 +1653,7 @@ void ULandscapeComponent::DeleteLayerInternal(ULandscapeLayerInfoObject* LayerIn
 		return;
 	}
 
-	Component->DeleteLayerAllocation(EditLayerGuid, DeleteAllocIdx, LandscapeEdit.GetShouldDirtyPackage());
+	Component->DeleteLayerAllocation(InEditLayerGuid, DeleteAllocIdx, LandscapeEdit.GetShouldDirtyPackage());
 
 	// See if the deleted layer is a NoWeightBlend layer - if not, we don't have to worry about normalization
 	// If the layer doesn't exist, assume it is a blended layer (so renormalization will run whether needed or not)
@@ -1761,7 +1808,7 @@ void FLandscapeEditDataInterface::DeleteLayer(ULandscapeLayerInfoObject* LayerIn
 	for (auto& XYComponentPair : LandscapeInfo->XYtoComponentMap)
 	{
 		ULandscapeComponent* Component = XYComponentPair.Value;
-		Component->DeleteLayer(LayerInfo, *this);
+		Component->DeleteLayerInternal(LayerInfo, *this, GetEditLayer());
 	}
 
 	if (LandscapeInfo->LandscapeActor.IsValid() && LandscapeInfo->LandscapeActor->HasLayersContent())
@@ -1782,6 +1829,8 @@ void ULandscapeComponent::FillLayer(ULandscapeLayerInfoObject* LayerInfo, FLands
 
 	ULandscapeComponent* Component = this;
 
+	FGuid EditLayerGuid = LandscapeEdit.GetEditLayer();
+
 	ALandscapeProxy* Proxy = Component->GetLandscapeProxy();
 	Component->Modify(LandscapeEdit.GetShouldDirtyPackage());
 	Proxy->Modify(LandscapeEdit.GetShouldDirtyPackage());
@@ -1789,9 +1838,9 @@ void ULandscapeComponent::FillLayer(ULandscapeLayerInfoObject* LayerInfo, FLands
 	const bool bFillLayerIsNoWeightBlend = LayerInfo->bNoWeightBlend;
 	bool bClearOtherWeightBlendLayers = !bFillLayerIsNoWeightBlend;
 
-	TArray<FWeightmapLayerAllocationInfo>& ComponentWeightmapLayerAllocations = Component->GetWeightmapLayerAllocations(true);
-	TArray<UTexture2D*>& ComponentWeightmapTextures = Component->GetWeightmapTextures(true);
-	TArray<ULandscapeWeightmapUsage*>& ComponentWeightmapTexturesUsage = Component->GetWeightmapTexturesUsage(true);
+	TArray<FWeightmapLayerAllocationInfo>& ComponentWeightmapLayerAllocations = Component->GetWeightmapLayerAllocations(EditLayerGuid);
+	TArray<UTexture2D*>& ComponentWeightmapTextures = Component->GetWeightmapTextures(EditLayerGuid);
+	TArray<ULandscapeWeightmapUsage*>& ComponentWeightmapTexturesUsage = Component->GetWeightmapTexturesUsage(EditLayerGuid);
 
 	// Find the index for this layer in this component.
 	int32 FillLayerIdx = ComponentWeightmapLayerAllocations.IndexOfByPredicate(
@@ -1819,7 +1868,7 @@ void ULandscapeComponent::FillLayer(ULandscapeLayerInfoObject* LayerInfo, FLands
 	{
 		FillLayerIdx = ComponentWeightmapLayerAllocations.Num();
 		ComponentWeightmapLayerAllocations.Add(FWeightmapLayerAllocationInfo(LayerInfo));
-		Component->ReallocateWeightmaps(&LandscapeEdit);
+		Component->ReallocateWeightmapsInternal(&LandscapeEdit, EditLayerGuid);
 	}
 
 	check(FillLayerIdx != INDEX_NONE);
@@ -1988,7 +2037,7 @@ void FLandscapeEditDataInterface::FillLayer(ULandscapeLayerInfoObject* LayerInfo
 	for (auto& XYComponentPair : LandscapeInfo->XYtoComponentMap)
 	{
 		ULandscapeComponent* Component = XYComponentPair.Value;
-		const TArray<FWeightmapLayerAllocationInfo>& ComponentWeightmapLayerAllocations = Component->GetWeightmapLayerAllocations(true);
+		const TArray<FWeightmapLayerAllocationInfo>& ComponentWeightmapLayerAllocations = Component->GetWeightmapLayerAllocations(GetEditLayer());
 
 		bool LayerNeedFilling = true;
 
@@ -2023,6 +2072,11 @@ void FLandscapeEditDataInterface::FillLayer(ULandscapeLayerInfoObject* LayerInfo
 }
 
 void ULandscapeComponent::ReplaceLayer(ULandscapeLayerInfoObject* FromLayerInfo, ULandscapeLayerInfoObject* ToLayerInfo, FLandscapeEditDataInterface& LandscapeEdit)
+{
+	ReplaceLayerInternal(FromLayerInfo, ToLayerInfo, LandscapeEdit, GetEditingLayerGUID());
+}
+
+void ULandscapeComponent::ReplaceLayerInternal(ULandscapeLayerInfoObject* FromLayerInfo, ULandscapeLayerInfoObject* ToLayerInfo, FLandscapeEditDataInterface& LandscapeEdit, const FGuid& InEditLayerGUID)
 {
 	check(FromLayerInfo && ToLayerInfo);
 	if (FromLayerInfo == ToLayerInfo)
@@ -2030,9 +2084,9 @@ void ULandscapeComponent::ReplaceLayer(ULandscapeLayerInfoObject* FromLayerInfo,
 		return;
 	}
 
-	TArray<FWeightmapLayerAllocationInfo>& ComponentWeightmapLayerAllocations = GetWeightmapLayerAllocations(true);
-	TArray<UTexture2D*>& ComponentWeightmapTextures = GetWeightmapTextures(true);
-	TArray<ULandscapeWeightmapUsage*>& ComponentWeightmapTexturesUsage = GetWeightmapTexturesUsage(true);
+	TArray<FWeightmapLayerAllocationInfo>& ComponentWeightmapLayerAllocations = GetWeightmapLayerAllocations(InEditLayerGUID);
+	TArray<UTexture2D*>& ComponentWeightmapTextures = GetWeightmapTextures(InEditLayerGUID);
+	TArray<ULandscapeWeightmapUsage*>& ComponentWeightmapTexturesUsage = GetWeightmapTexturesUsage(InEditLayerGUID);
 
 	// Find the index for this layer in this component.
 	int32 FromLayerIdx = INDEX_NONE;
@@ -2249,12 +2303,12 @@ void FLandscapeEditDataInterface::ReplaceLayer(ULandscapeLayerInfoObject* FromLa
 {
 	if (!LandscapeInfo) return;
 
-	auto DoReplace = [&]()
+	auto DoReplace = [&](const FGuid& EditLayerGuid)
 	{
 		for (auto It = LandscapeInfo->XYtoComponentMap.CreateIterator(); It; ++It)
 		{
 			ULandscapeComponent* Component = It.Value();
-			Component->ReplaceLayer(FromLayerInfo, ToLayerInfo, *this);
+			Component->ReplaceLayerInternal(FromLayerInfo, ToLayerInfo, *this, EditLayerGuid);
 		}
 	};
 
@@ -2268,14 +2322,15 @@ void FLandscapeEditDataInterface::ReplaceLayer(ULandscapeLayerInfoObject* FromLa
 			{
 				CurrentLayer.WeightmapLayerAllocationBlend.Add(ToLayerInfo, OutValue);
 			}
-			FScopedSetLandscapeEditingLayer Scope(LandscapeInfo->LandscapeActor.Get(), CurrentLayer.Guid);
-			DoReplace();
+
+			DoReplace(CurrentLayer.Guid);
 		});
 		LandscapeInfo->LandscapeActor->RequestLayersContentUpdateForceAll(ELandscapeLayerUpdateMode::Update_Weightmap_All);
 	}
 	else
 	{
-		DoReplace();
+		FGuid BaseLayerGuid;
+		DoReplace(BaseLayerGuid);
 	}
 }
 
@@ -2424,7 +2479,7 @@ namespace
 	}
 };
 
-bool DeleteLayerIfAllZero(ULandscapeComponent* const Component, const uint8* const TexDataPtr, int32 TexSize, int32 LayerIdx, bool bShouldDirtyPackage)
+bool DeleteLayerIfAllZero(ULandscapeComponent* const Component, const FGuid& InEditLayerGuid, const uint8* const TexDataPtr, int32 TexSize, int32 LayerIdx, bool bShouldDirtyPackage)
 {
 	// Check the data for the entire component and to see if it's all zero
 	for (int32 TexY = 0; TexY < TexSize; TexY++)
@@ -2442,7 +2497,7 @@ bool DeleteLayerIfAllZero(ULandscapeComponent* const Component, const uint8* con
 		}
 	}
 
-	Component->DeleteLayerAllocation(Component->GetEditingLayerGUID(), LayerIdx, bShouldDirtyPackage);
+	Component->DeleteLayerAllocation(InEditLayerGuid, LayerIdx, bShouldDirtyPackage);
 
 	return true;
 }
@@ -2530,8 +2585,8 @@ inline TMap<const ULandscapeLayerInfoObject*, uint32> FLandscapeEditDataInterfac
 	TMap<const ULandscapeLayerInfoObject*, uint32> LayerInfluenceMap;
 
 	ULandscapeComponent* Component = LandscapeInfo->XYtoComponentMap.FindChecked(FIntPoint(ComponentIndexX,ComponentIndexY));
-	const TArray<FWeightmapLayerAllocationInfo>& ComponentWeightmapLayerAllocations = Component->GetWeightmapLayerAllocations(true);
-	const TArray<UTexture2D*>& ComponentWeightmapTextures = Component->GetWeightmapTextures(true);
+	const TArray<FWeightmapLayerAllocationInfo>& ComponentWeightmapLayerAllocations = Component->GetWeightmapLayerAllocations(GetEditLayer());
+	const TArray<UTexture2D*>& ComponentWeightmapTextures = Component->GetWeightmapTextures(GetEditLayer());
 
 	// used if InOptionalLayerDataPtrs is null
 	TArray<FLandscapeTextureDataInfo*, TInlineAllocator<2>> InternalTexDataInfos;
@@ -2733,7 +2788,7 @@ void FLandscapeEditDataInterface::SetAlphaData(ULandscapeLayerInfoObject* const
 
 			Component->Modify(GetShouldDirtyPackage());
 
-			TArray<FWeightmapLayerAllocationInfo>& ComponentWeightmapLayerAllocations = Component->GetWeightmapLayerAllocations(true);
+			TArray<FWeightmapLayerAllocationInfo>& ComponentWeightmapLayerAllocations = Component->GetWeightmapLayerAllocations(GetEditLayer());
 			int32 UpdateLayerIdx = ComponentWeightmapLayerAllocations.IndexOfByPredicate([LayerInfo](const FWeightmapLayerAllocationInfo& Allocation){ return Allocation.LayerInfo == LayerInfo; });
 
 			// Need allocation for weightmap
@@ -2804,7 +2859,7 @@ void FLandscapeEditDataInterface::SetAlphaData(ULandscapeLayerInfoObject* const
 
 				UpdateLayerIdx = ComponentWeightmapLayerAllocations.Num();
 				new (ComponentWeightmapLayerAllocations) FWeightmapLayerAllocationInfo(LayerInfo);
-				Component->ReallocateWeightmaps(this);
+				Component->ReallocateWeightmapsInternal(this, GetEditLayer());
 
 				if (!Component->GetLandscapeProxy()->HasLayersContent())
 				{
@@ -2819,7 +2874,7 @@ void FLandscapeEditDataInterface::SetAlphaData(ULandscapeLayerInfoObject* const
 
 			// Lock data for all the weightmaps
 			TexDataInfos.Reset();
-			const TArray<UTexture2D*>& ComponentWeightmapTextures = Component->GetWeightmapTextures(true);
+			const TArray<UTexture2D*>& ComponentWeightmapTextures = Component->GetWeightmapTextures(GetEditLayer());
 			TexDataInfos.AddUninitialized(ComponentWeightmapTextures.Num());
 
 			for (int32 WeightmapIdx = 0; WeightmapIdx < ComponentWeightmapTextures.Num(); ++WeightmapIdx)
@@ -3133,7 +3188,7 @@ void FLandscapeEditDataInterface::SetAlphaData(ULandscapeLayerInfoObject* const
 			{
 				if (LayerEditDataAllZero[LayerIdx])
 				{
-					bool bLayerDeleted = DeleteLayerIfAllZero(Component, LayerDataPtrs[LayerIdx], TexSize, LayerIdx, GetShouldDirtyPackage());
+					bool bLayerDeleted = DeleteLayerIfAllZero(Component, GetEditLayer(), LayerDataPtrs[LayerIdx], TexSize, LayerIdx, GetShouldDirtyPackage());
 
 					if (bLayerDeleted)
 					{
@@ -3222,7 +3277,7 @@ void FLandscapeEditDataInterface::SetAlphaData(const TSet<ULandscapeLayerInfoObj
 
 			NeedAllocationInfos.Reset();
 
-			TArray<FWeightmapLayerAllocationInfo>& ComponentWeightmapLayerAllocations = Component->GetWeightmapLayerAllocations(true);
+			TArray<FWeightmapLayerAllocationInfo>& ComponentWeightmapLayerAllocations = Component->GetWeightmapLayerAllocations(GetEditLayer());
 
 			for (ULandscapeLayerInfoObject* LayerInfo : DirtyLayerInfos)
 			{
@@ -3258,7 +3313,7 @@ void FLandscapeEditDataInterface::SetAlphaData(const TSet<ULandscapeLayerInfoObj
 						}
 						ComponentWeightmapLayerAllocations.Emplace(LayerInfoNeedingAllocation);
 					}
-					Component->ReallocateWeightmaps(this);
+					Component->ReallocateWeightmapsInternal(this, GetEditLayer());
 					
 					if (!Component->GetLandscapeProxy()->HasLayersContent())
 					{
@@ -3272,7 +3327,7 @@ void FLandscapeEditDataInterface::SetAlphaData(const TSet<ULandscapeLayerInfoObj
 
 			// Lock data for all the weightmaps
 			TexDataInfos.Reset();
-			const TArray<UTexture2D*>& ComponentWeightmapTextures = Component->GetWeightmapTextures(true);
+			const TArray<UTexture2D*>& ComponentWeightmapTextures = Component->GetWeightmapTextures(GetEditLayer());
 			TexDataInfos.AddUninitialized(ComponentWeightmapTextures.Num());
 
 			for (int32 WeightmapIdx = 0; WeightmapIdx < ComponentWeightmapTextures.Num(); ++WeightmapIdx)
@@ -3422,7 +3477,7 @@ void FLandscapeEditDataInterface::SetAlphaData(const TSet<ULandscapeLayerInfoObj
 			{
 				if (LayerEditDataAllZero[LayerIdx])
 				{
-					bool bLayerDeleted = DeleteLayerIfAllZero(Component, LayerDataInfos[LayerIdx].TexDataPtr, TexSize, LayerIdx, GetShouldDirtyPackage());
+					bool bLayerDeleted = DeleteLayerIfAllZero(Component, GetEditLayer(), LayerDataInfos[LayerIdx].TexDataPtr, TexSize, LayerIdx, GetShouldDirtyPackage());
 
 					if (bLayerDeleted)
 					{
@@ -3474,8 +3529,8 @@ void FLandscapeEditDataInterface::GetWeightDataTemplFast(ULandscapeLayerInfoObje
 			uint8 WeightmapChannelOffset = 0;
 			TArray<FLandscapeTextureDataInfo*> TexDataInfos; // added for whole weight case...
 
-			const TArray<FWeightmapLayerAllocationInfo>& ComponentWeightmapLayerAllocations = Component->GetWeightmapLayerAllocations(true);
-			const TArray<UTexture2D*>& ComponentWeightmapTextures = Component->GetWeightmapTextures(true);
+			const TArray<FWeightmapLayerAllocationInfo>& ComponentWeightmapLayerAllocations = Component->GetWeightmapLayerAllocations(GetEditLayer());
+			const TArray<UTexture2D*>& ComponentWeightmapTextures = Component->GetWeightmapTextures(GetEditLayer());
 
 			if (LayerInfo != NULL)
 			{
@@ -3600,8 +3655,8 @@ uint8 FLandscapeEditDataInterface::GetWeightMapData(const ULandscapeComponent* C
 	{
 		if (LayerInfo != NULL)
 		{
-			const TArray<FWeightmapLayerAllocationInfo>& ComponentWeightmapLayerAllocations = Component->GetWeightmapLayerAllocations(true);
-			const TArray<UTexture2D*>& ComponentWeightmapTextures = Component->GetWeightmapTextures(true);
+			const TArray<FWeightmapLayerAllocationInfo>& ComponentWeightmapLayerAllocations = Component->GetWeightmapLayerAllocations(GetEditLayer());
+			const TArray<UTexture2D*>& ComponentWeightmapTextures = Component->GetWeightmapTextures(GetEditLayer());
 
 			for( int32 LayerIdx=0;LayerIdx<ComponentWeightmapLayerAllocations.Num();LayerIdx++ )
 			{
@@ -3688,11 +3743,11 @@ void FLandscapeEditDataInterface::GetWeightDataTempl(ULandscapeLayerInfoObject*
 		
 			if( Component )
 			{
-				const TArray<UTexture2D*>& ComponentWeightmapTextures = Component->GetWeightmapTextures(true);
+				const TArray<UTexture2D*>& ComponentWeightmapTextures = Component->GetWeightmapTextures(GetEditLayer());
 
 				if (LayerInfo != NULL)
 				{
-					const TArray<FWeightmapLayerAllocationInfo>& ComponentWeightmapLayerAllocations = Component->GetWeightmapLayerAllocations(true);
+					const TArray<FWeightmapLayerAllocationInfo>& ComponentWeightmapLayerAllocations = Component->GetWeightmapLayerAllocations(GetEditLayer());
 
 					for( int32 LayerIdx=0;LayerIdx<ComponentWeightmapLayerAllocations.Num();LayerIdx++ )
 					{
@@ -3759,8 +3814,8 @@ void FLandscapeEditDataInterface::GetWeightDataTempl(ULandscapeLayerInfoObject*
 							NoBorderX1 = false;
 							if (LayerInfo != NULL)
 							{
-								const TArray<FWeightmapLayerAllocationInfo>& BorderWeightmapLayerAllocations = BorderComponent[0]->GetWeightmapLayerAllocations(true);
-								const TArray<UTexture2D*>& BorderWeightmapTextures = BorderComponent[0]->GetWeightmapTextures(true);
+								const TArray<FWeightmapLayerAllocationInfo>& BorderWeightmapLayerAllocations = BorderComponent[0]->GetWeightmapLayerAllocations(GetEditLayer());
+								const TArray<UTexture2D*>& BorderWeightmapTextures = BorderComponent[0]->GetWeightmapTextures(GetEditLayer());
 
 								for( int32 LayerIdx=0;LayerIdx<BorderWeightmapLayerAllocations.Num();LayerIdx++ )
 								{
@@ -3790,8 +3845,8 @@ void FLandscapeEditDataInterface::GetWeightDataTempl(ULandscapeLayerInfoObject*
 							NoBorderX2 = false;
 							if (LayerInfo != NULL)
 							{
-								const TArray<FWeightmapLayerAllocationInfo>& BorderWeightmapLayerAllocations = BorderComponent[1]->GetWeightmapLayerAllocations(true);
-								const TArray<UTexture2D*>& BorderWeightmapTextures = BorderComponent[1]->GetWeightmapTextures(true);
+								const TArray<FWeightmapLayerAllocationInfo>& BorderWeightmapLayerAllocations = BorderComponent[1]->GetWeightmapLayerAllocations(GetEditLayer());
+								const TArray<UTexture2D*>& BorderWeightmapTextures = BorderComponent[1]->GetWeightmapTextures(GetEditLayer());
 
 								for( int32 LayerIdx=0;LayerIdx<BorderWeightmapLayerAllocations.Num();LayerIdx++ )
 								{
@@ -3820,8 +3875,8 @@ void FLandscapeEditDataInterface::GetWeightDataTempl(ULandscapeLayerInfoObject*
 							NoBorderY1[ComponentIndexXX] = false;
 							if (LayerInfo != NULL)
 							{
-								const TArray<FWeightmapLayerAllocationInfo>& BorderWeightmapLayerAllocations = BorderComponent[2]->GetWeightmapLayerAllocations(true);
-								const TArray<UTexture2D*>& BorderWeightmapTextures = BorderComponent[2]->GetWeightmapTextures(true);
+								const TArray<FWeightmapLayerAllocationInfo>& BorderWeightmapLayerAllocations = BorderComponent[2]->GetWeightmapLayerAllocations(GetEditLayer());
+								const TArray<UTexture2D*>& BorderWeightmapTextures = BorderComponent[2]->GetWeightmapTextures(GetEditLayer());
 
 								for( int32 LayerIdx=0;LayerIdx<BorderWeightmapLayerAllocations.Num();LayerIdx++ )
 								{
@@ -3845,8 +3900,8 @@ void FLandscapeEditDataInterface::GetWeightDataTempl(ULandscapeLayerInfoObject*
 					{
 						if (LayerInfo != NULL)
 						{
-							const TArray<FWeightmapLayerAllocationInfo>& BorderWeightmapLayerAllocations = BorderComponent[2]->GetWeightmapLayerAllocations(true);
-							const TArray<UTexture2D*>& BorderWeightmapTextures = BorderComponent[2]->GetWeightmapTextures(true);
+							const TArray<FWeightmapLayerAllocationInfo>& BorderWeightmapLayerAllocations = BorderComponent[2]->GetWeightmapLayerAllocations(GetEditLayer());
+							const TArray<UTexture2D*>& BorderWeightmapTextures = BorderComponent[2]->GetWeightmapTextures(GetEditLayer());
 
 							for( int32 LayerIdx=0;LayerIdx<BorderWeightmapLayerAllocations.Num();LayerIdx++ )
 							{
@@ -3874,8 +3929,8 @@ void FLandscapeEditDataInterface::GetWeightDataTempl(ULandscapeLayerInfoObject*
 							NoBorderY2[ComponentIndexXX] = false;
 							if (LayerInfo != NULL)
 							{
-								const TArray<FWeightmapLayerAllocationInfo>& BorderWeightmapLayerAllocations = BorderComponent[3]->GetWeightmapLayerAllocations(true);
-								const TArray<UTexture2D*>& BorderWeightmapTextures = BorderComponent[3]->GetWeightmapTextures(true);
+								const TArray<FWeightmapLayerAllocationInfo>& BorderWeightmapLayerAllocations = BorderComponent[3]->GetWeightmapLayerAllocations(GetEditLayer());
+								const TArray<UTexture2D*>& BorderWeightmapTextures = BorderComponent[3]->GetWeightmapTextures(GetEditLayer());
 
 								for( int32 LayerIdx=0;LayerIdx<BorderWeightmapLayerAllocations.Num();LayerIdx++ )
 								{
@@ -3900,8 +3955,8 @@ void FLandscapeEditDataInterface::GetWeightDataTempl(ULandscapeLayerInfoObject*
 					{
 						if (LayerInfo != NULL)
 						{
-							const TArray<FWeightmapLayerAllocationInfo>& BorderWeightmapLayerAllocations = BorderComponent[3]->GetWeightmapLayerAllocations(true);
-							const TArray<UTexture2D*>& BorderWeightmapTextures = BorderComponent[3]->GetWeightmapTextures(true);
+							const TArray<FWeightmapLayerAllocationInfo>& BorderWeightmapLayerAllocations = BorderComponent[3]->GetWeightmapLayerAllocations(GetEditLayer());
+							const TArray<UTexture2D*>& BorderWeightmapTextures = BorderComponent[3]->GetWeightmapTextures(GetEditLayer());
 
 							for( int32 LayerIdx=0;LayerIdx<BorderWeightmapLayerAllocations.Num();LayerIdx++ )
 							{
@@ -4173,8 +4228,8 @@ void FLandscapeEditDataInterface::GetWeightDataTempl(ULandscapeLayerInfoObject*
 							{
 								StoreData.PreInit(LandscapeInfo->Layers.Num());
 
-								const TArray<FWeightmapLayerAllocationInfo>& ComponentWeightmapLayerAllocations = Component->GetWeightmapLayerAllocations(true);
-								const TArray<UTexture2D*>& ComponentWeightmapTextures = Component->GetWeightmapTextures(true);
+								const TArray<FWeightmapLayerAllocationInfo>& ComponentWeightmapLayerAllocations = Component->GetWeightmapLayerAllocations(GetEditLayer());
+								const TArray<UTexture2D*>& ComponentWeightmapTextures = Component->GetWeightmapTextures(GetEditLayer());
 
 								for( int32 LayerIdx=0;LayerIdx<ComponentWeightmapLayerAllocations.Num();LayerIdx++ )
 								{
diff --git a/Engine/Source/Runtime/Landscape/Private/LandscapeEditLayers.cpp b/Engine/Source/Runtime/Landscape/Private/LandscapeEditLayers.cpp
--- a/Engine/Source/Runtime/Landscape/Private/LandscapeEditLayers.cpp
+++ b/Engine/Source/Runtime/Landscape/Private/LandscapeEditLayers.cpp
@@ -4796,13 +4796,11 @@ int32 ALandscape::PerformLayersHeightmapsGlobalMerge(const FUpdateLayersContentC
 	// Calculate Top Left Lambda
 	auto GetUniqueHeightmaps = [&](const TArray<ULandscapeComponent*>& InLandscapeComponents, TArray<FHeightmapCopyInfo>& OutHeightmaps, const FIntPoint& LandscapeBaseQuads, const FGuid& InLayerGuid = FGuid())
 	{
-		FScopedSetLandscapeEditingLayer Scope(this, InLayerGuid);
-
 		const int32 ComponentSizeQuad = SubsectionSizeQuads * NumSubsections;
 		const int32 ComponentSizeVerts = (SubsectionSizeQuads + 1) * NumSubsections;
 		for (ULandscapeComponent* Component : InLandscapeComponents)
 		{
-			UTexture2D* ComponentHeightmap = Component->GetHeightmap(true);
+			UTexture2D* ComponentHeightmap = Component->GetHeightmap(InLayerGuid);
 
 			int32 Index = OutHeightmaps.IndexOfByPredicate([=](const FHeightmapCopyInfo& LayerHeightmap) { return LayerHeightmap.Texture == ComponentHeightmap; });
 
@@ -9342,7 +9340,7 @@ bool ALandscape::IsEditingLayerReservedForSplines() const
 	if (CanHaveLayersContent())
 	{
 		const FLandscapeLayer* SplinesReservedLayer = GetLandscapeSplinesReservedLayer();
-		return SplinesReservedLayer && SplinesReservedLayer->Guid == EditingLayer;
+		return SplinesReservedLayer && SplinesReservedLayer->Guid == GetEditingLayer();
 	}
 	return false;
 }
diff --git a/Engine/Source/Runtime/Landscape/Public/LandscapeEdit.h b/Engine/Source/Runtime/Landscape/Public/LandscapeEdit.h
--- a/Engine/Source/Runtime/Landscape/Public/LandscapeEdit.h
+++ b/Engine/Source/Runtime/Landscape/Public/LandscapeEdit.h
@@ -160,9 +160,18 @@ private:
 
 struct LANDSCAPE_API FLandscapeEditDataInterface : public FLandscapeTextureDataInterface
 {
-	// tor
+	// this constructor will build an interface that works in the current edit layer (uses the ALandscape PrivateEditingLayer state)
 	FLandscapeEditDataInterface(ULandscapeInfo* InLandscape, bool bInUploadTextureChangesToGPU = true);
 
+	// this constructor will build an interface that works in the specified edit layer : InEditLayerGUID
+	FLandscapeEditDataInterface(ULandscapeInfo* InLandscape, const FGuid& InEditLayerGUID, bool bInUploadTextureChangesToGPU = true);
+
+	// switch the layer this interface works in
+	void SetEditLayer(const FGuid& InEditLayerGUID);
+
+	// returns the current layer this interface is working in
+	FGuid GetEditLayer() const;
+
 	// Misc
 	bool GetComponentsInRegion(int32 X1, int32 Y1, int32 X2, int32 Y2, TSet<ULandscapeComponent*>* OutComponents = NULL);
 
@@ -230,7 +239,8 @@ struct LANDSCAPE_API FLandscapeEditDataInterface : public FLandscapeTextureDataI
 	void FillLayer(ULandscapeLayerInfoObject* LayerInfo);
 	// Fill all empty layers and re-normalize layers
 	void FillEmptyLayers(ULandscapeLayerInfoObject* LayerInfo);
-	// Replace/merge a layer
+	
+	// Replace/merge a layer (across all edit layers and the main runtime layer)
 	void ReplaceLayer(ULandscapeLayerInfoObject* FromLayerInfo, ULandscapeLayerInfoObject* ToLayerInfo);
 
 	template<typename TStoreData>
@@ -296,6 +306,10 @@ private:
 	FVector DrawScale;
 
 	ULandscapeInfo* LandscapeInfo;
+	
+	// if true, we use the LandscapeActor->PrivateEditingLayer, otherwise we use the LocalEditLayerGUID specified below
+	bool bUseSharedLandscapeEditLayer;
+	FGuid LocalEditLayerGUID;
 
 	void FillLayer(ULandscapeLayerInfoObject* LayerInfo, bool bEmptyLayersOnly);
 
@@ -389,6 +403,16 @@ struct FHeightmapAccessor
 		LandscapeEdit = new FLandscapeEditDataInterface(InLandscapeInfo);
 	}
 
+	void SetEditLayer(const FGuid& InEditLayerGUID)
+	{
+		LandscapeEdit->SetEditLayer(InEditLayerGUID);
+	}
+
+	FGuid GetEditLayer() const
+	{
+		return LandscapeEdit->GetEditLayer();
+	}
+
 	// accessors
 	void GetData(int32& X1, int32& Y1, int32& X2, int32& Y2, TMap<FIntPoint, uint16>& Data)
 	{
@@ -606,6 +630,16 @@ struct FAlphamapAccessor
 		LandscapeEdit.Flush();
 	}
 
+	void SetEditLayer(const FGuid& InEditLayerGUID)
+	{
+		LandscapeEdit.SetEditLayer(InEditLayerGUID);
+	}
+	
+	FGuid GetEditLayer() const
+	{
+		return LandscapeEdit.GetEditLayer();
+	}
+
 private:
 	ULandscapeInfo* LandscapeInfo;
 	FLandscapeEditDataInterface LandscapeEdit;
