Landscape : Fix undo / redo landscape sculpt

diff --git a/Engine/Source/Runtime/Landscape/Private/LandscapeEdit.cpp b/Engine/Source/Runtime/Landscape/Private/LandscapeEdit.cpp
--- a/Engine/Source/Runtime/Landscape/Private/LandscapeEdit.cpp
+++ b/Engine/Source/Runtime/Landscape/Private/LandscapeEdit.cpp
@@ -677,6 +677,13 @@ void ULandscapeComponent::PostEditUndo()
 {
 	ALandscapeProxy* Proxy = GetLandscapeProxy();
 	
+	// PreEditUndo of the LandscapeStreamingProxy resets it's pointer to the top level landscape so this ensures it's fixed up prior
+	// to requesting HeightMap or WeightMap updates
+	if (ULandscapeInfo* LandscapeInfo = Proxy->GetLandscapeInfo(); !LandscapeInfo->IsRegistered(Proxy))
+	{
+		LandscapeInfo->RegisterActor(Proxy, true);
+	}
+
 	// On undo, request a recompute weightmap usages, as they are not transacted (they combine information between multiple components)
 	Proxy->RequestProxyLayersWeightmapUsageUpdate();
 
