Adds Landscape BP Brush visibility layer editing.
* Deprecate Render method for BP Brushes, RenderLayer should now be used
* Adds FLandscapeBrushParameters struct that should be used to pass arguments when rendering a layer
* Modifies ELandscapeToolTargetType to be an enum class, usable by blueprints
* Adds BP Brush visibility layer toggle in Landscape Editor
* Adds visibility layer editing for LandscapeCircleHeightPatches
* Adds visibility layer editing for LandscapeTexturePatches
* Adds helper method to determine if a layer info object name is a visibility layer

diff --git a/Engine/Plugins/Experimental/LandscapePatch/Shaders/Private/LandscapeCircleHeightPatchPS.usf b/Engine/Plugins/Experimental/LandscapePatch/Shaders/Private/LandscapeCircleHeightPatchPS.usf
--- a/Engine/Plugins/Experimental/LandscapePatch/Shaders/Private/LandscapeCircleHeightPatchPS.usf
+++ b/Engine/Plugins/Experimental/LandscapePatch/Shaders/Private/LandscapeCircleHeightPatchPS.usf
@@ -8,7 +8,9 @@
 float3 InCenter;
 float InRadius;
 float InFalloff;
-Texture2D<float4> InSourceHeightmap;
+Texture2D<float4> InSourceTexture;
+
+#if CIRCLE_HEIGHT_PATCH
 
 void ApplyLandscapeCircleHeightPatch(in float4 SVPos : SV_POSITION, out float2 OutColor : SV_Target0)
 {
@@ -16,7 +18,7 @@ void ApplyLandscapeCircleHeightPatch(in float4 SVPos : SV_POSITION, out float2 O
 	const float ClampedFalloff = max(InFalloff, KINDA_SMALL_NUMBER);
 	
 	int2 TextureCoordinates = floor(SVPos.xy);
-	float4 CurrentPackedHeight = InSourceHeightmap.Load(int3(TextureCoordinates, 0));
+	float4 CurrentPackedHeight = InSourceTexture.Load(int3(TextureCoordinates, 0));
 	float CurrentHeight = UnpackHeight(CurrentPackedHeight.xy);
 	
 	float PatchHeight = InCenter.z;
@@ -32,4 +34,27 @@ void ApplyLandscapeCircleHeightPatch(in float4 SVPos : SV_POSITION, out float2 O
 	float NewHeight = lerp(CurrentHeight, PatchHeight, Alpha);
 	
 	OutColor = PackHeight(NewHeight);
-}
\ No newline at end of file
+}
+
+#endif // CIRCLE_HEIGHT_PATCH
+
+#if CIRCLE_VISIBILITY_PATCH
+
+void ApplyLandscapeCircleVisibilityPatch(in float4 SVPos : SV_POSITION, out float OutColor : SV_Target0)
+{
+	int2 TextureCoordinates = floor(SVPos.xy);
+	float Distance = distance(TextureCoordinates, InCenter.xy);
+	
+	float4 CurrentWeight = InSourceTexture.Load(int3(TextureCoordinates, 0)).x;
+
+	if (Distance <= InRadius)
+	{
+		OutColor = 255.f;
+	}
+	else
+	{
+		OutColor = CurrentWeight.x;
+	}
+}
+
+#endif // CIRCLE_VISIBILITY_PATCH
\ No newline at end of file
diff --git a/Engine/Plugins/Experimental/LandscapePatch/Source/LandscapePatch/Private/LandscapeCircleHeightPatch.cpp b/Engine/Plugins/Experimental/LandscapePatch/Source/LandscapePatch/Private/LandscapeCircleHeightPatch.cpp
--- a/Engine/Plugins/Experimental/LandscapePatch/Source/LandscapePatch/Private/LandscapeCircleHeightPatch.cpp
+++ b/Engine/Plugins/Experimental/LandscapePatch/Source/LandscapePatch/Private/LandscapeCircleHeightPatch.cpp
@@ -4,7 +4,10 @@
 
 #include "Engine/TextureRenderTarget2D.h"
 #include "LandscapeCircleHeightPatchPS.h"
+#include "LandscapeEditTypes.h"
+#include "LandscapeLayerInfoObject.h"
 #include "LandscapePatchManager.h"
+#include "LandscapeProxy.h"
 #include "RenderGraphUtils.h"
 #include "TextureResource.h"
 #include "RenderingThread.h"
@@ -43,72 +46,68 @@ void ULandscapeCircleHeightPatch::Initialize_Native(const FTransform & InLandsca
 {
 }
 
-UTextureRenderTarget2D* ULandscapeCircleHeightPatch::Render_Native(bool bIsHeightmap,
-	UTextureRenderTarget2D* InCombinedResult,
-	const FName& InWeightmapLayerName)
+UTextureRenderTarget2D* ULandscapeCircleHeightPatch::RenderLayer_Native(const FLandscapeBrushParameters& InParameters)
 {
 	if (!ensure(PatchManager.IsValid()))
 	{
-		return InCombinedResult;
+		return InParameters.CombinedResult;
 	}
 
-	// Circle height patch doesn't affect weight maps
-	if (!bIsHeightmap)
+	// Circle height patch doesn't affect regular weightmap layers.
+	if ((bEditVisibility && (InParameters.LayerType != ELandscapeToolTargetType::Visibility)) || (!bEditVisibility && (InParameters.LayerType != ELandscapeToolTargetType::Heightmap)))
 	{
-		return InCombinedResult;
+		return InParameters.CombinedResult;
 	}
 
 	FTransform HeightmapCoordsToWorld = PatchManager->GetHeightmapCoordsToWorld();
-
 	double ToHeightmapRadiusScale = GetComponentTransform().GetScale3D().X / HeightmapCoordsToWorld.GetScale3D().X;
-
 	FVector3d CircleCenterWorld = GetComponentTransform().GetTranslation();
 	FVector3d CenterInHeightmapCoordinates = HeightmapCoordsToWorld.InverseTransformPosition(CircleCenterWorld);
-
 	float RadiusAdjustment = bExclusiveRadius ? 0 : 1;
-
 	float HeightmapRadius = Radius * ToHeightmapRadiusScale + RadiusAdjustment;
 	float HeightmapFalloff = Falloff * ToHeightmapRadiusScale + RadiusAdjustment;
 
-	// Go ahead an pack everything into a copy of the param struct so we don't have to capture everything
-	// individually in the lambda below.
-	FLandscapeCircleHeightPatchPS::FParameters ShaderParamsToCopy;
-	ShaderParamsToCopy.InCenter = (FVector3f)CenterInHeightmapCoordinates;
-	ShaderParamsToCopy.InRadius = HeightmapRadius;
-	ShaderParamsToCopy.InFalloff = HeightmapFalloff;
-
-	ENQUEUE_RENDER_COMMAND(LandscapeCircleHeightPatch)([InCombinedResult, ShaderParamsToCopy](FRHICommandListImmediate& RHICmdList)
+	ENQUEUE_RENDER_COMMAND(LandscapeCircleHeightPatch)([CombinedResult = InParameters.CombinedResult, CenterInHeightmapCoordinates, HeightmapRadius, HeightmapFalloff, bEditVisibility = bEditVisibility](FRHICommandListImmediate& RHICmdList)
 	{
 		TRACE_CPUPROFILER_EVENT_SCOPE(LandscapeCircleHeightPatch);
 
-		FRDGBuilder GraphBuilder(RHICmdList, RDG_EVENT_NAME("ApplyLandmassBoxHeightPatch"));
+		const TCHAR* OutputName = bEditVisibility ? TEXT("LandscapeCircleVisibilityPatchOutput") : TEXT("LandscapeCircleHeightPatchOutput");
+		const TCHAR* InputCopyName = bEditVisibility ? TEXT("LandscapeCircleVisibilityPatchInputCopy") : TEXT("LandscapeCircleHeightPatchInputCopy");
+
+		FRDGBuilder GraphBuilder(RHICmdList, RDG_EVENT_NAME("ApplyLandscapeCirclePatch"));
 		
-		TRefCountPtr<IPooledRenderTarget> RenderTarget = CreateRenderTarget(InCombinedResult->GetResource()->GetTexture2DRHI(), TEXT("LandscapeCircleHeightPatchOutput"));
+		TRefCountPtr<IPooledRenderTarget> RenderTarget = CreateRenderTarget(CombinedResult->GetResource()->GetTexture2DRHI(), OutputName);
 		FRDGTextureRef DestinationTexture = GraphBuilder.RegisterExternalTexture(RenderTarget);
 
 		// Make a copy of our heightmap input
-		FRDGTextureRef InputCopy = GraphBuilder.CreateTexture(DestinationTexture->Desc, TEXT("LandscapeCircleHeightPatchInputCopy"));
+		FRDGTextureRef InputCopy = GraphBuilder.CreateTexture(DestinationTexture->Desc, InputCopyName);
 
 		FRHICopyTextureInfo CopyTextureInfo;
 		CopyTextureInfo.NumMips = 1;
 		CopyTextureInfo.Size = FIntVector(DestinationTexture->Desc.GetSize().X, DestinationTexture->Desc.GetSize().Y, 0);
 		AddCopyTexturePass(GraphBuilder, DestinationTexture, InputCopy, CopyTextureInfo);
-
-		FLandscapeCircleHeightPatchPS::FParameters* ShaderParams =
-			GraphBuilder.AllocParameters<FLandscapeCircleHeightPatchPS::FParameters>();
-		*ShaderParams = ShaderParamsToCopy;
-
 		FRDGTextureSRVRef InputCopySRV = GraphBuilder.CreateSRV(FRDGTextureSRVDesc::CreateForMipLevel(InputCopy, 0));
-		ShaderParams->InSourceHeightmap = InputCopySRV;
 
+		FLandscapeCircleHeightPatchPS::FParameters* ShaderParams = GraphBuilder.AllocParameters<FLandscapeCircleHeightPatchPS::FParameters>();
+		ShaderParams->InCenter = (FVector3f)CenterInHeightmapCoordinates;
+		ShaderParams->InRadius = HeightmapRadius;
+		ShaderParams->InFalloff = HeightmapFalloff;
+		ShaderParams->InSourceTexture = InputCopySRV;
 		ShaderParams->RenderTargets[0] = FRenderTargetBinding(DestinationTexture, ERenderTargetLoadAction::ENoAction, /*InMipIndex = */0);
 
-		FLandscapeCircleHeightPatchPS::AddToRenderGraph(GraphBuilder, ShaderParams);
+		if (bEditVisibility)
+		{
+			FLandscapeCircleVisibilityPatchPS::AddToRenderGraph(GraphBuilder, ShaderParams);
+		}
+		else
+		{
+			FLandscapeCircleHeightPatchPS::AddToRenderGraph(GraphBuilder, ShaderParams);
+		}
 
 		GraphBuilder.Execute();
 	});
 
-	return InCombinedResult;
+	return InParameters.CombinedResult;
 }
 
 #undef LOCTEXT_NAMESPACE
diff --git a/Engine/Plugins/Experimental/LandscapePatch/Source/LandscapePatch/Private/LandscapeCircleHeightPatchPS.cpp b/Engine/Plugins/Experimental/LandscapePatch/Source/LandscapePatch/Private/LandscapeCircleHeightPatchPS.cpp
--- a/Engine/Plugins/Experimental/LandscapePatch/Source/LandscapePatch/Private/LandscapeCircleHeightPatchPS.cpp
+++ b/Engine/Plugins/Experimental/LandscapePatch/Source/LandscapePatch/Private/LandscapeCircleHeightPatchPS.cpp
@@ -15,8 +15,7 @@ bool FLandscapeCircleHeightPatchPS::ShouldCompilePermutation(const FGlobalShader
 
 void FLandscapeCircleHeightPatchPS::ModifyCompilationEnvironment(const FGlobalShaderPermutationParameters& Parameters, FShaderCompilerEnvironment& OutEnvironment)
 {
-	// If we end up using an ifdef
-	//OutEnvironment.SetDefine(TEXT("CIRCLE_HEIGHT_PATCH"), 1);
+	OutEnvironment.SetDefine(TEXT("CIRCLE_HEIGHT_PATCH"), 1);
 }
 
 void FLandscapeCircleHeightPatchPS::AddToRenderGraph(FRDGBuilder& GraphBuilder, FParameters* InParameters)
@@ -24,7 +23,7 @@ void FLandscapeCircleHeightPatchPS::AddToRenderGraph(FRDGBuilder& GraphBuilder,
 	FGlobalShaderMap* ShaderMap = GetGlobalShaderMap(GMaxRHIFeatureLevel);
 	TShaderMapRef<FLandscapeCircleHeightPatchPS> PixelShader(ShaderMap);
 
-	FIntVector TextureSize = InParameters->InSourceHeightmap->Desc.Texture->Desc.GetSize();
+	FIntVector TextureSize = InParameters->InSourceTexture->Desc.Texture->Desc.GetSize();
 
 	FPixelShaderUtils::AddFullscreenPass(
 		GraphBuilder,
@@ -35,4 +34,26 @@ void FLandscapeCircleHeightPatchPS::AddToRenderGraph(FRDGBuilder& GraphBuilder,
 		FIntRect(0, 0, TextureSize.X, TextureSize.Y));
 }
 
+void FLandscapeCircleVisibilityPatchPS::ModifyCompilationEnvironment(const FGlobalShaderPermutationParameters& Parameters, FShaderCompilerEnvironment& OutEnvironment)
+{
+	OutEnvironment.SetDefine(TEXT("CIRCLE_VISIBILITY_PATCH"), 1);
+}
+
+void FLandscapeCircleVisibilityPatchPS::AddToRenderGraph(FRDGBuilder& GraphBuilder, FParameters* InParameters)
+{
+	FGlobalShaderMap* ShaderMap = GetGlobalShaderMap(GMaxRHIFeatureLevel);
+	TShaderMapRef<FLandscapeCircleVisibilityPatchPS> PixelShader(ShaderMap);
+
+	FIntVector TextureSize = InParameters->InSourceTexture->Desc.Texture->Desc.GetSize();
+
+	FPixelShaderUtils::AddFullscreenPass(
+		GraphBuilder,
+		ShaderMap,
+		RDG_EVENT_NAME("LandmassCircleVisibilityPatch"),
+		PixelShader,
+		InParameters,
+		FIntRect(0, 0, TextureSize.X, TextureSize.Y));
+}
+
 IMPLEMENT_GLOBAL_SHADER(FLandscapeCircleHeightPatchPS, "/Plugin/LandscapePatch/Private/LandscapeCircleHeightPatchPS.usf", "ApplyLandscapeCircleHeightPatch", SF_Pixel);
+IMPLEMENT_GLOBAL_SHADER(FLandscapeCircleVisibilityPatchPS, "/Plugin/LandscapePatch/Private/LandscapeCircleHeightPatchPS.usf", "ApplyLandscapeCircleVisibilityPatch", SF_Pixel);
diff --git a/Engine/Plugins/Experimental/LandscapePatch/Source/LandscapePatch/Private/LandscapePatchManager.cpp b/Engine/Plugins/Experimental/LandscapePatch/Source/LandscapePatch/Private/LandscapePatchManager.cpp
--- a/Engine/Plugins/Experimental/LandscapePatch/Source/LandscapePatch/Private/LandscapePatchManager.cpp
+++ b/Engine/Plugins/Experimental/LandscapePatch/Source/LandscapePatch/Private/LandscapePatchManager.cpp
@@ -6,6 +6,7 @@
 #include "Engine/Level.h"
 #include "Engine/World.h"
 #include "Landscape.h"
+#include "LandscapeEditTypes.h"
 #include "LandscapeDataAccess.h"
 #include "LandscapePatchComponent.h"
 #include "LandscapePatchLogging.h"
@@ -23,6 +24,7 @@ ALandscapePatchManager::ALandscapePatchManager(const FObjectInitializer& ObjectI
 #if WITH_EDITOR
 	SetAffectsHeightmap(true);
 	SetAffectsWeightmap(true);
+	SetAffectsVisibilityLayer(true);
 #endif
 }
 
@@ -70,12 +72,11 @@ void ALandscapePatchManager::Initialize_Native(const FTransform & InLandscapeTra
 		+ InLandscapeTransform.GetTranslation());
 }
 
-UTextureRenderTarget2D* ALandscapePatchManager::Render_Native(bool InIsHeightmap,
-	UTextureRenderTarget2D* InCombinedResult,
-	const FName& InWeightmapLayerName)
+UTextureRenderTarget2D* ALandscapePatchManager::RenderLayer_Native(const FLandscapeBrushParameters& InParameters)
 {
 	// Used to determine whether we need to remove any invalid brushes
 	bool bHaveInvalidPatches = false;
+	FLandscapeBrushParameters BrushParameters = InParameters;
 
 	// TODO: There are many uncertainties in how we iterate across the height patches and have them
 	// apply themselves. For one thing we may want to pass around a render graph, in which case this
@@ -95,7 +96,7 @@ UTextureRenderTarget2D* ALandscapePatchManager::Render_Native(bool InIsHeightmap
 		{
 			if (Component->IsEnabled())
 			{
-				InCombinedResult = Component->Render_Native(InIsHeightmap, InCombinedResult, InWeightmapLayerName);
+				BrushParameters.CombinedResult = Component->RenderLayer_Native(BrushParameters);
 			}
 		}
 		else if (Component.IsNull())
@@ -124,10 +125,10 @@ UTextureRenderTarget2D* ALandscapePatchManager::Render_Native(bool InIsHeightmap
 	{
 		PatchComponents.RemoveAll([](TSoftObjectPtr<ULandscapePatchComponent> Component) {
 			return Component.IsNull();
-		});
+			});
 	}
 
-	return InCombinedResult;
+	return BrushParameters.CombinedResult;
 }
 
 void ALandscapePatchManager::SetTargetLandscape(ALandscape* InTargetLandscape)
diff --git a/Engine/Plugins/Experimental/LandscapePatch/Source/LandscapePatch/Private/LandscapeTexturePatch.cpp b/Engine/Plugins/Experimental/LandscapePatch/Source/LandscapePatch/Private/LandscapeTexturePatch.cpp
--- a/Engine/Plugins/Experimental/LandscapePatch/Source/LandscapePatch/Private/LandscapeTexturePatch.cpp
+++ b/Engine/Plugins/Experimental/LandscapePatch/Source/LandscapePatch/Private/LandscapeTexturePatch.cpp
@@ -75,23 +75,24 @@ namespace LandscapeTexturePatchLocals
 }
 
 #if WITH_EDITOR
-UTextureRenderTarget2D* ULandscapeTexturePatch::Render_Native(bool bIsHeightmap,
-	UTextureRenderTarget2D* InCombinedResult,
-	const FName& InWeightmapLayerName)
+UTextureRenderTarget2D* ULandscapeTexturePatch::RenderLayer_Native(const FLandscapeBrushParameters& InParameters)
 {
 	using namespace UE::Landscape;
+	const bool bIsHeightmapTarget = InParameters.LayerType == ELandscapeToolTargetType::Heightmap;
+	const bool bIsWeightmapTarget = InParameters.LayerType == ELandscapeToolTargetType::Weightmap;
+	const bool bIsVisibilityLayerTarget = InParameters.LayerType == ELandscapeToolTargetType::Visibility;
 
-	if (bIsHeightmap)
+	if (bIsHeightmapTarget)
 	{
 		if (bReinitializeHeightOnNextRender)
 		{
 			bReinitializeHeightOnNextRender = false;
-			ReinitializeHeight(InCombinedResult);
-			return InCombinedResult;
+			ReinitializeHeight(InParameters.CombinedResult);
+			return InParameters.CombinedResult;
 		}
 		else
 		{
-			return ApplyToHeightmap(InCombinedResult);
+			return ApplyToHeightmap(InParameters.CombinedResult);
 		}
 	}
 	else
@@ -101,7 +102,8 @@ UTextureRenderTarget2D* ULandscapeTexturePatch::Render_Native(bool bIsHeightmap,
 
 		for (TObjectPtr<ULandscapeWeightPatchTextureInfo> WeightPatchEntry : WeightPatches)
 		{
-			if (WeightPatchEntry->WeightmapLayerName == InWeightmapLayerName)
+			if ((bIsWeightmapTarget && (WeightPatchEntry->WeightmapLayerName == InParameters.WeightmapLayerName)) ||
+				(bIsVisibilityLayerTarget && WeightPatchEntry->bEditVisibilityLayer))
 			{
 				WeightPatchInfo = WeightPatchEntry;
 				break;
@@ -110,18 +112,18 @@ UTextureRenderTarget2D* ULandscapeTexturePatch::Render_Native(bool bIsHeightmap,
 
 		if (!WeightPatchInfo)
 		{
-			return InCombinedResult;
+			return InParameters.CombinedResult;
 		}
 
 		if (WeightPatchInfo->bReinitializeOnNextRender)
 		{
 			WeightPatchInfo->bReinitializeOnNextRender = false;
-			ReinitializeWeightPatch(WeightPatchInfo, InCombinedResult);
-			return InCombinedResult;
+			ReinitializeWeightPatch(WeightPatchInfo, InParameters.CombinedResult);
+			return InParameters.CombinedResult;
 		}
 		else
 		{
-			return ApplyToWeightmap(WeightPatchInfo, InCombinedResult);
+			return ApplyToWeightmap(WeightPatchInfo, InParameters.CombinedResult);
 		}
 	}
 }
diff --git a/Engine/Plugins/Experimental/LandscapePatch/Source/LandscapePatch/Public/LandscapeCircleHeightPatch.h b/Engine/Plugins/Experimental/LandscapePatch/Source/LandscapePatch/Public/LandscapeCircleHeightPatch.h
--- a/Engine/Plugins/Experimental/LandscapePatch/Source/LandscapePatch/Public/LandscapeCircleHeightPatch.h
+++ b/Engine/Plugins/Experimental/LandscapePatch/Source/LandscapePatch/Public/LandscapeCircleHeightPatch.h
@@ -4,6 +4,7 @@
 
 
 #include "LandscapePatchComponent.h"
+#include "LandscapeEditTypes.h"
 
 #include "LandscapeCircleHeightPatch.generated.h"
 
@@ -24,9 +25,8 @@ public:
 	virtual void Initialize_Native(const FTransform& InLandscapeTransform, 
 		const FIntPoint& InLandscapeSize, 
 		const FIntPoint& InLandscapeRenderTargetSize) override;
-	virtual UTextureRenderTarget2D* Render_Native(bool InIsHeightmap,
-		UTextureRenderTarget2D* InCombinedResult,
-		const FName& InWeightmapLayerName) override;
+
+	virtual UTextureRenderTarget2D* RenderLayer_Native(const FLandscapeBrushParameters& InParameters) override;
 
 	// UActorComponent
 	virtual void OnComponentCreated() override;
@@ -40,6 +40,10 @@ protected:
 	UPROPERTY(EditAnywhere, Category = Settings)
 	float Falloff = 500;
 
+	/** Specifies if this patch edits the visibility layer. */
+	UPROPERTY(EditAnywhere, Category = Settings)
+	bool bEditVisibility = false;
+
 	/** When true, only the vertices in the circle have alpha 1. If false, the radius is expanded slightly so that neighboring 
 	  vertices are also included and the whole circle is able to lie flat. */
 	UPROPERTY(EditAnywhere, Category = Settings, AdvancedDisplay)
diff --git a/Engine/Plugins/Experimental/LandscapePatch/Source/LandscapePatch/Public/LandscapeCircleHeightPatchPS.h b/Engine/Plugins/Experimental/LandscapePatch/Source/LandscapePatch/Public/LandscapeCircleHeightPatchPS.h
--- a/Engine/Plugins/Experimental/LandscapePatch/Source/LandscapePatch/Public/LandscapeCircleHeightPatchPS.h
+++ b/Engine/Plugins/Experimental/LandscapePatch/Source/LandscapePatch/Public/LandscapeCircleHeightPatchPS.h
@@ -16,7 +16,7 @@ class LANDSCAPEPATCH_API FLandscapeCircleHeightPatchPS : public FGlobalShader
 
 public:
 	BEGIN_SHADER_PARAMETER_STRUCT(FParameters, )
-		SHADER_PARAMETER_RDG_TEXTURE_SRV(Texture2D<float4>, InSourceHeightmap) // Our input texture
+		SHADER_PARAMETER_RDG_TEXTURE_SRV(Texture2D<float4>, InSourceTexture) // Our input texture
 		SHADER_PARAMETER(FVector3f, InCenter)
 		SHADER_PARAMETER(float, InRadius)
 		SHADER_PARAMETER(float, InFalloff)
@@ -30,6 +30,23 @@ public:
 	static void AddToRenderGraph(FRDGBuilder& GraphBuilder, FParameters* InParameters);
 };
 
+/**
+ * Shader that applies a circle patch to a landscape visibility layer.
+ */
+class LANDSCAPEPATCH_API FLandscapeCircleVisibilityPatchPS : public FLandscapeCircleHeightPatchPS
+{
+	DECLARE_GLOBAL_SHADER(FLandscapeCircleVisibilityPatchPS);
+
+public:
+	FLandscapeCircleVisibilityPatchPS() = default;
+	FLandscapeCircleVisibilityPatchPS(const ShaderMetaType::CompiledShaderInitializerType& Initializer)
+		: FLandscapeCircleHeightPatchPS(Initializer)
+	{}
+
+	static void ModifyCompilationEnvironment(const FGlobalShaderPermutationParameters& Parameters, FShaderCompilerEnvironment& OutEnvironment);
+	static void AddToRenderGraph(FRDGBuilder& GraphBuilder, FParameters* InParameters);
+};
+
 #if UE_ENABLE_INCLUDE_ORDER_DEPRECATED_IN_5_2
 #include "CoreMinimal.h"
 #endif
diff --git a/Engine/Plugins/Experimental/LandscapePatch/Source/LandscapePatch/Public/LandscapePatchComponent.h b/Engine/Plugins/Experimental/LandscapePatch/Source/LandscapePatch/Public/LandscapePatchComponent.h
--- a/Engine/Plugins/Experimental/LandscapePatch/Source/LandscapePatch/Public/LandscapePatchComponent.h
+++ b/Engine/Plugins/Experimental/LandscapePatch/Source/LandscapePatch/Public/LandscapePatchComponent.h
@@ -4,6 +4,8 @@
 
 
 #include "Components/SceneComponent.h"
+#include "LandscapeBlueprintBrushBase.h"
+#include "LandscapeEditTypes.h"
 
 #include "LandscapePatchComponent.generated.h"
 
@@ -31,9 +33,8 @@ public:
 	virtual void Initialize_Native(const FTransform& InLandscapeTransform,
 		const FIntPoint& InLandscapeSize,
 		const FIntPoint& InLandscapeRenderTargetSize) {}
-	virtual UTextureRenderTarget2D* Render_Native(bool InIsHeightmap,
-		UTextureRenderTarget2D* InCombinedResult,
-		const FName& InWeightmapLayerName) { return InCombinedResult; }
+
+	virtual UTextureRenderTarget2D* RenderLayer_Native(const FLandscapeBrushParameters& InParameters) { return InParameters.CombinedResult; }
 
 	virtual bool IsAffectingWeightmapLayer(const FName& InLayerName) const { return false; }
 
diff --git a/Engine/Plugins/Experimental/LandscapePatch/Source/LandscapePatch/Public/LandscapePatchManager.h b/Engine/Plugins/Experimental/LandscapePatch/Source/LandscapePatch/Public/LandscapePatchManager.h
--- a/Engine/Plugins/Experimental/LandscapePatch/Source/LandscapePatch/Public/LandscapePatchManager.h
+++ b/Engine/Plugins/Experimental/LandscapePatch/Source/LandscapePatch/Public/LandscapePatchManager.h
@@ -4,6 +4,7 @@
 
 
 #include "LandscapeBlueprintBrushBase.h"
+#include "LandscapeEditTypes.h"
 
 #include "LandscapePatchManager.generated.h"
 
@@ -30,9 +31,8 @@ public:
 	virtual void Initialize_Native(const FTransform& InLandscapeTransform,
 		const FIntPoint& InLandscapeSize,
 		const FIntPoint& InLandscapeRenderTargetSize) override;
-	virtual UTextureRenderTarget2D* Render_Native(bool InIsHeightmap,
-		UTextureRenderTarget2D* InCombinedResult,
-		const FName& InWeightmapLayerName) override;
+
+	virtual UTextureRenderTarget2D* RenderLayer_Native(const FLandscapeBrushParameters& InParameters) override;
 
 	// Adds the brush to the given landscape, removing it from any previous one. This differs from SetOwningLandscape
 	// in that SetOwningLandscape is called by the landscape itself from AddBrushToLayer to update the manager.
diff --git a/Engine/Plugins/Experimental/LandscapePatch/Source/LandscapePatch/Public/LandscapeTexturePatch.h b/Engine/Plugins/Experimental/LandscapePatch/Source/LandscapePatch/Public/LandscapeTexturePatch.h
--- a/Engine/Plugins/Experimental/LandscapePatch/Source/LandscapePatch/Public/LandscapeTexturePatch.h
+++ b/Engine/Plugins/Experimental/LandscapePatch/Source/LandscapePatch/Public/LandscapeTexturePatch.h
@@ -161,6 +161,10 @@ public:
 protected:
 	UPROPERTY(EditAnywhere, Category = WeightPatch)
 	FName WeightmapLayerName;
+
+	/** Specifies if this patch edits the visibility layer. */
+	UPROPERTY(EditAnywhere, Category = WeightPatch)
+	bool bEditVisibilityLayer = false;
 	
 	/** Texture to use when source mode is set to texture asset. */
 	UPROPERTY(EditAnywhere, Category = WeightPatch, meta = (EditConditionHides,
@@ -216,9 +220,7 @@ class LANDSCAPEPATCH_API ULandscapeTexturePatch : public ULandscapePatchComponen
 public:
 
 #if WITH_EDITOR
-	virtual UTextureRenderTarget2D* Render_Native(bool InIsHeightmap,
-		UTextureRenderTarget2D* InCombinedResult,
-		const FName& InWeightmapLayerName) override;
+	UTextureRenderTarget2D* RenderLayer_Native(const FLandscapeBrushParameters& InParameters);
 
 	// ULandscapePatchComponent
 	virtual bool IsAffectingWeightmapLayer(const FName& InLayerName) const override;
diff --git a/Engine/Plugins/Experimental/Water/Source/Editor/Private/WaterBrushManager.cpp b/Engine/Plugins/Experimental/Water/Source/Editor/Private/WaterBrushManager.cpp
--- a/Engine/Plugins/Experimental/Water/Source/Editor/Private/WaterBrushManager.cpp
+++ b/Engine/Plugins/Experimental/Water/Source/Editor/Private/WaterBrushManager.cpp
@@ -1349,15 +1349,15 @@ void AWaterBrushManager::SetupDefaultMaterials()
 	FindEdgesMaterial = WaterEditorSettings->GetDefaultFindEdgesMaterial();
 }
 
-UTextureRenderTarget2D* AWaterBrushManager::Render_Native(bool InIsHeightmap, UTextureRenderTarget2D* InCombinedResult, FName const& InWeightmapLayerName)
+UTextureRenderTarget2D* AWaterBrushManager::RenderLayer_Native(const FLandscapeBrushParameters& InParameters)
 {
-	TRACE_CPUPROFILER_EVENT_SCOPE(AWaterBrushManager::Render_Native);
+	TRACE_CPUPROFILER_EVENT_SCOPE(AWaterBrushManager::RenderLayer_Native);
 
-	LandscapeRTRef = InCombinedResult;
+	LandscapeRTRef = InParameters.CombinedResult;
 
 	FBrushRenderContext BrushRenderContext;
-	BrushRenderContext.bHeightmapRender = InIsHeightmap;
-	BrushRenderContext.WeightmapLayerName = InWeightmapLayerName;
+	BrushRenderContext.bHeightmapRender = InParameters.LayerType == ELandscapeToolTargetType::Heightmap;
+	BrushRenderContext.WeightmapLayerName = InParameters.WeightmapLayerName;
 
 	if (!BrushRenderSetup())
 	{
@@ -1407,7 +1407,7 @@ UTextureRenderTarget2D* AWaterBrushManager::Render_Native(bool InIsHeightmap, UT
 	}
 
 	// Render Completed
-	UTextureRenderTarget2D* ReturnRT = InIsHeightmap ? HeightPingPongRead(BrushRenderContext) : WeightPingPongRead(BrushRenderContext);
+	UTextureRenderTarget2D* ReturnRT = (InParameters.LayerType == ELandscapeToolTargetType::Heightmap) ? HeightPingPongRead(BrushRenderContext) : WeightPingPongRead(BrushRenderContext);
 
 	bKillCache = false;
 
diff --git a/Engine/Plugins/Experimental/Water/Source/Editor/Private/WaterBrushManager.h b/Engine/Plugins/Experimental/Water/Source/Editor/Private/WaterBrushManager.h
--- a/Engine/Plugins/Experimental/Water/Source/Editor/Private/WaterBrushManager.h
+++ b/Engine/Plugins/Experimental/Water/Source/Editor/Private/WaterBrushManager.h
@@ -205,7 +205,7 @@ public:
 	virtual void BlueprintOnRenderTargetTexturesUpdated_Native(UTexture2D* VelocityTexture) override;
 	virtual void BlueprintWaterBodyChanged_Native(AActor* Actor) override;
 	virtual void Initialize_Native(FTransform const& InLandscapeTransform, FIntPoint const& InLandscapeSize, FIntPoint const& InLandscapeRenderTargetSize) override;
-	virtual UTextureRenderTarget2D*  Render_Native(bool InIsHeightmap, UTextureRenderTarget2D* InCombinedResult, FName const& InWeightmapLayerName) override;
+	virtual UTextureRenderTarget2D*  RenderLayer_Native(const FLandscapeBrushParameters& InParameters) override;
 		
 	virtual void GetRenderDependencies(TSet<UObject*>& OutDependencies) override;
 
diff --git a/Engine/Source/Editor/EditorStyle/Private/SlateEditorStyle.cpp b/Engine/Source/Editor/EditorStyle/Private/SlateEditorStyle.cpp
--- a/Engine/Source/Editor/EditorStyle/Private/SlateEditorStyle.cpp
+++ b/Engine/Source/Editor/EditorStyle/Private/SlateEditorStyle.cpp
@@ -6923,6 +6923,8 @@ void FSlateEditorStyle::FStyle::SetupLandscapeEditorStyle()
 		Set("LandscapeEditor.Brush.AffectsHeight.Disabled", new IMAGE_BRUSH("Icons/icon_Landscape_Affects_Height_Disabled_16x", Icon16x16));
 		Set("LandscapeEditor.Brush.AffectsWeight.Enabled", new IMAGE_BRUSH("Icons/icon_Landscape_Affects_Weight_Enabled_16x", Icon16x16));
 		Set("LandscapeEditor.Brush.AffectsWeight.Disabled", new IMAGE_BRUSH("Icons/icon_Landscape_Affects_Weight_Disabled_16x", Icon16x16));
+		Set("LandscapeEditor.Brush.AffectsVisibilityLayer.Enabled", new IMAGE_BRUSH("Icons/icon_Landscape_Affects_VisibilityLayer_Enabled_16x", Icon16x16));
+		Set("LandscapeEditor.Brush.AffectsVisibilityLayer.Disabled", new IMAGE_BRUSH("Icons/icon_Landscape_Affects_VisibilityLayer_Disabled_16x", Icon16x16));
 
 		// Target List
 		Set("LandscapeEditor.TargetList.RowBackground",        new FSlateNoResource());
diff --git a/Engine/Source/Editor/EditorStyle/Private/StarshipStyle.cpp b/Engine/Source/Editor/EditorStyle/Private/StarshipStyle.cpp
--- a/Engine/Source/Editor/EditorStyle/Private/StarshipStyle.cpp
+++ b/Engine/Source/Editor/EditorStyle/Private/StarshipStyle.cpp
@@ -6414,6 +6414,8 @@ void FStarshipEditorStyle::FStyle::SetupLandscapeEditorStyle()
 		Set("LandscapeEditor.Brush.AffectsHeight.Disabled", new IMAGE_BRUSH("Icons/icon_Landscape_Affects_Height_Disabled_16x", Icon16x16));
 		Set("LandscapeEditor.Brush.AffectsWeight.Enabled", new IMAGE_BRUSH("Icons/icon_Landscape_Affects_Weight_Enabled_16x", Icon16x16));
 		Set("LandscapeEditor.Brush.AffectsWeight.Disabled", new IMAGE_BRUSH("Icons/icon_Landscape_Affects_Weight_Disabled_16x", Icon16x16));
+		Set("LandscapeEditor.Brush.AffectsVisibilityLayer.Enabled", new IMAGE_BRUSH("Icons/icon_Landscape_Affects_VisibilityLayer_Enabled_16x", Icon16x16));
+		Set("LandscapeEditor.Brush.AffectsVisibilityLayer.Disabled", new IMAGE_BRUSH("Icons/icon_Landscape_Affects_VisibilityLayer_Disabled_16x", Icon16x16));
 
 		// Target List
 		Set("LandscapeEditor.TargetList.RowBackground",        new FSlateNoResource());
diff --git a/Engine/Source/Editor/LandscapeEditor/Private/LandscapeEdMode.cpp b/Engine/Source/Editor/LandscapeEditor/Private/LandscapeEdMode.cpp
--- a/Engine/Source/Editor/LandscapeEditor/Private/LandscapeEdMode.cpp
+++ b/Engine/Source/Editor/LandscapeEditor/Private/LandscapeEdMode.cpp
@@ -24,6 +24,7 @@
 #include "EditorSupportDelegates.h"
 #include "ScopedTransaction.h"
 #include "LandscapeEdit.h"
+#include "LandscapeEditTypes.h"
 #include "LandscapeEditorUtils.h"
 #include "LandscapeRender.h"
 #include "LandscapeDataAccess.h"
@@ -333,7 +334,7 @@ void FEdModeLandscape::PostUpdateLayerContent()
 	}
 }
 
-ELandscapeToolTargetType::Type FEdModeLandscape::GetLandscapeToolTargetType() const
+ELandscapeToolTargetType FEdModeLandscape::GetLandscapeToolTargetType() const
 {
 	if (CurrentToolMode)
 	{
diff --git a/Engine/Source/Editor/LandscapeEditor/Private/LandscapeEdMode.h b/Engine/Source/Editor/LandscapeEditor/Private/LandscapeEdMode.h
--- a/Engine/Source/Editor/LandscapeEditor/Private/LandscapeEdMode.h
+++ b/Engine/Source/Editor/LandscapeEditor/Private/LandscapeEdMode.h
@@ -13,6 +13,7 @@
 #include "LandscapeLayerInfoObject.h"
 #include "LandscapeGizmoActiveActor.h"
 #include "LandscapeEdit.h"
+#include "LandscapeEditTypes.h"
 #include "Containers/Set.h"
 #include "LandscapeImportHelper.h"
 
@@ -63,7 +64,7 @@ struct FLandscapeToolMode
 struct FLandscapeTargetListInfo
 {
 	FText TargetName;
-	ELandscapeToolTargetType::Type TargetType;
+	ELandscapeToolTargetType TargetType;
 	TWeakObjectPtr<ULandscapeInfo> LandscapeInfo;
 
 	//Values cloned from FLandscapeLayerStruct LayerStruct;			// ignored for heightmap
@@ -75,7 +76,7 @@ struct FLandscapeTargetListInfo
 	uint32 bValid : 1;												// ignored for heightmap
 	int32 LayerIndex;
 
-	FLandscapeTargetListInfo(FText InTargetName, ELandscapeToolTargetType::Type InTargetType, const FLandscapeInfoLayerSettings& InLayerSettings, int32 InLayerIndex)
+	FLandscapeTargetListInfo(FText InTargetName, ELandscapeToolTargetType InTargetType, const FLandscapeInfoLayerSettings& InLayerSettings, int32 InLayerIndex)
 		: TargetName(InTargetName)
 		, TargetType(InTargetType)
 		, LandscapeInfo(InLayerSettings.Owner->GetLandscapeInfo())
@@ -89,7 +90,7 @@ struct FLandscapeTargetListInfo
 	{
 	}
 
-	FLandscapeTargetListInfo(FText InTargetName, ELandscapeToolTargetType::Type InTargetType, ULandscapeInfo* InLandscapeInfo, int32 InLayerIndex)
+	FLandscapeTargetListInfo(FText InTargetName, ELandscapeToolTargetType InTargetType, ULandscapeInfo* InLandscapeInfo, int32 InLayerIndex)
 		: TargetName(InTargetName)
 		, TargetType(InTargetType)
 		, LandscapeInfo(InLandscapeInfo)
@@ -355,7 +356,7 @@ public:
 
 	/** ILandscapeEdModeInterface */
 	virtual void PostUpdateLayerContent() override;
-	virtual ELandscapeToolTargetType::Type GetLandscapeToolTargetType() const override;
+	virtual ELandscapeToolTargetType GetLandscapeToolTargetType() const override;
 	virtual const FLandscapeLayer* GetLandscapeSelectedLayer() const override;
 	virtual ULandscapeLayerInfoObject* GetSelectedLandscapeLayerInfo() const override;
 	virtual void OnCanHaveLayersContentChanged() override;
diff --git a/Engine/Source/Editor/LandscapeEditor/Private/LandscapeEdModeTools.h b/Engine/Source/Editor/LandscapeEditor/Private/LandscapeEdModeTools.h
--- a/Engine/Source/Editor/LandscapeEditor/Private/LandscapeEdModeTools.h
+++ b/Engine/Source/Editor/LandscapeEditor/Private/LandscapeEdModeTools.h
@@ -7,6 +7,7 @@
 #include "LandscapeProxy.h"
 #include "LandscapeToolInterface.h"
 #include "LandscapeEdMode.h"
+#include "LandscapeEditTypes.h"
 #include "EditorViewportClient.h"
 #include "LandscapeEdit.h"
 #include "LandscapeComponent.h"
@@ -1315,7 +1316,7 @@ struct FLandscapeDataCache : public TLandscapeEditCache<FDatamapAccessor<false>,
 struct FHeightmapToolTarget
 {
 	typedef FLandscapeHeightCache CacheClass;
-	static const ELandscapeToolTargetType::Type TargetType = ELandscapeToolTargetType::Heightmap;
+	static const ELandscapeToolTargetType TargetType = ELandscapeToolTargetType::Heightmap;
 
 	static float StrengthMultiplier(ULandscapeInfo* LandscapeInfo, float BrushRadius)
 	{
@@ -1346,7 +1347,7 @@ struct FHeightmapToolTarget
 struct FWeightmapToolTarget
 {
 	typedef FLandscapeAlphaCache CacheClass;
-	static const ELandscapeToolTargetType::Type TargetType = ELandscapeToolTargetType::Weightmap;
+	static const ELandscapeToolTargetType TargetType = ELandscapeToolTargetType::Weightmap;
 
 	static float StrengthMultiplier(ULandscapeInfo* LandscapeInfo, float BrushRadius)
 	{
@@ -1421,7 +1422,7 @@ public:
 
 	virtual ELandscapeLayerUpdateMode GetBeginToolContentUpdateFlag() const
 	{
-		bool bUpdateHeightmap = this->EdMode->CurrentToolTarget.TargetType == ELandscapeToolTargetType::Type::Heightmap; 
+		bool bUpdateHeightmap = this->EdMode->CurrentToolTarget.TargetType == ELandscapeToolTargetType::Heightmap; 
 		return bUpdateHeightmap ? ELandscapeLayerUpdateMode::Update_Heightmap_Editing : ELandscapeLayerUpdateMode::Update_Weightmap_Editing;
 	}
 
@@ -1432,7 +1433,7 @@ public:
 
 	virtual ELandscapeLayerUpdateMode GetEndToolContentUpdateFlag() const
 	{
-		bool bUpdateHeightmap = this->EdMode->CurrentToolTarget.TargetType == ELandscapeToolTargetType::Type::Heightmap;
+		bool bUpdateHeightmap = this->EdMode->CurrentToolTarget.TargetType == ELandscapeToolTargetType::Heightmap;
 		return bUpdateHeightmap ? ELandscapeLayerUpdateMode::Update_Heightmap_All : ELandscapeLayerUpdateMode::Update_Weightmap_All;
 	}
 
diff --git a/Engine/Source/Editor/LandscapeEditor/Private/LandscapeEditorDetailCustomization_LayersBrushStack.cpp b/Engine/Source/Editor/LandscapeEditor/Private/LandscapeEditorDetailCustomization_LayersBrushStack.cpp
--- a/Engine/Source/Editor/LandscapeEditor/Private/LandscapeEditorDetailCustomization_LayersBrushStack.cpp
+++ b/Engine/Source/Editor/LandscapeEditor/Private/LandscapeEditorDetailCustomization_LayersBrushStack.cpp
@@ -21,6 +21,7 @@
 #include "DetailCategoryBuilder.h"
 #include "PropertyCustomizationHelpers.h"
 #include "LandscapeEditorDetailCustomization_Layers.h"
+#include "LandscapeEditTypes.h"
 
 #include "ScopedTransaction.h"
 
@@ -199,6 +200,23 @@ TSharedPtr<SWidget> FLandscapeEditorCustomNodeBuilder_LayersBrushStack::Generate
                     .Image(this, &FLandscapeEditorCustomNodeBuilder_LayersBrushStack::GetAffectsWeightmapBrush, InBrushIndex)
                 ]
             ]
+			+ SHorizontalBox::Slot()
+				.AutoWidth()
+				.VAlign(VAlign_Center)
+				[
+					SNew(SButton)
+					.ContentPadding(0)
+					.ButtonStyle(FAppStyle::Get(), "NoBorder")
+					.OnClicked(this, &FLandscapeEditorCustomNodeBuilder_LayersBrushStack::OnToggleAffectsVisibilityLayer, InBrushIndex)
+					.ToolTipText(LOCTEXT("LandscapeBrushAffectsVisibilityLayer", "Toggle Affects Visibility Layer"))
+					.HAlign(HAlign_Center)
+					.VAlign(VAlign_Center)
+					.Content()
+					[
+						SNew(SImage)
+						.Image(this, &FLandscapeEditorCustomNodeBuilder_LayersBrushStack::GetAffectsVisibilityLayerBrush, InBrushIndex)
+					]
+				]
 		];
 	
 	return RowWidget;
@@ -312,6 +330,16 @@ FReply FLandscapeEditorCustomNodeBuilder_LayersBrushStack::OnToggleAffectsWeight
 	}
 	return FReply::Handled();
 }
+FReply FLandscapeEditorCustomNodeBuilder_LayersBrushStack::OnToggleAffectsVisibilityLayer(int32 InBrushIndex)
+{
+	if (ALandscapeBlueprintBrushBase* Brush = GetBrush(InBrushIndex))
+	{
+		const FScopedTransaction Transaction(LOCTEXT("Landscape_Brush_AffectsVisibilityLayer", "Set Brush Affects Visibility Layer"));
+		bool bAffectsVisibilityLayer = Brush->IsAffectingVisibilityLayer();
+		Brush->SetAffectsVisibilityLayer(!bAffectsVisibilityLayer);
+	}
+	return FReply::Handled();
+}
 
 FReply FLandscapeEditorCustomNodeBuilder_LayersBrushStack::OnToggleVisibility(int32 InBrushIndex)
 {
@@ -341,6 +369,12 @@ const FSlateBrush* FLandscapeEditorCustomNodeBuilder_LayersBrushStack::GetAffect
 	return Brush && Brush->IsAffectingHeightmap() ? FAppStyle::GetBrush("LandscapeEditor.Brush.AffectsHeight.Enabled") : FAppStyle::GetBrush("LandscapeEditor.Brush.AffectsHeight.Disabled");
 }
 
+const FSlateBrush* FLandscapeEditorCustomNodeBuilder_LayersBrushStack::GetAffectsVisibilityLayerBrush(int32 InBrushIndex) const
+{
+	ALandscapeBlueprintBrushBase* Brush = GetBrush(InBrushIndex);
+	return Brush && Brush->IsAffectingVisibilityLayer() ? FAppStyle::GetBrush("LandscapeEditor.Brush.AffectsVisibilityLayer.Enabled") : FAppStyle::GetBrush("LandscapeEditor.Brush.AffectsVisibilityLayer.Disabled");
+}
+
 const FSlateBrush* FLandscapeEditorCustomNodeBuilder_LayersBrushStack::GetVisibilityBrush(int32 InBrushIndex) const
 {
 	ALandscapeBlueprintBrushBase* Brush = GetBrush(InBrushIndex);
@@ -352,7 +386,7 @@ bool FLandscapeEditorCustomNodeBuilder_LayersBrushStack::IsBrushEnabled(int32 In
 {
 	ALandscapeBlueprintBrushBase* Brush = GetBrush(InBrushIndex);
 	const FEdModeLandscape* LandscapeEdMode = GetEditorMode();
-	return Brush && ((Brush->IsAffectingHeightmap() && LandscapeEdMode->CurrentToolTarget.TargetType == ELandscapeToolTargetType::Type::Heightmap) || (Brush->IsAffectingWeightmap() && LandscapeEdMode->CurrentToolTarget.TargetType == ELandscapeToolTargetType::Type::Weightmap));
+	return Brush && ((Brush->IsAffectingHeightmap() && LandscapeEdMode->CurrentToolTarget.TargetType == ELandscapeToolTargetType::Heightmap) || (Brush->IsAffectingWeightmap() && LandscapeEdMode->CurrentToolTarget.TargetType == ELandscapeToolTargetType::Weightmap) || (Brush->IsAffectingVisibilityLayer() && LandscapeEdMode->CurrentToolTarget.TargetType == ELandscapeToolTargetType::Visibility));
 }
 
 bool FLandscapeEditorCustomNodeBuilder_LayersBrushStack::IsBrushSelected(int32 InBrushIndex) const
diff --git a/Engine/Source/Editor/LandscapeEditor/Private/LandscapeEditorDetailCustomization_LayersBrushStack.h b/Engine/Source/Editor/LandscapeEditor/Private/LandscapeEditorDetailCustomization_LayersBrushStack.h
--- a/Engine/Source/Editor/LandscapeEditor/Private/LandscapeEditorDetailCustomization_LayersBrushStack.h
+++ b/Engine/Source/Editor/LandscapeEditor/Private/LandscapeEditorDetailCustomization_LayersBrushStack.h
@@ -65,11 +65,13 @@ protected:
 	FReply OnToggleVisibility(int32 InBrushIndex);
 	FReply OnToggleAffectsHeightmap(int32 InBrushIndex);
 	FReply OnToggleAffectsWeightmap(int32 InBrushIndex);
+	FReply OnToggleAffectsVisibilityLayer(int32 InBrushIndex);
 
 	void OnToggleVisibility(ALandscapeBlueprintBrushBase* Brush);
 
 	const FSlateBrush* GetAffectsHeightmapBrush(int32 InBrushIndex) const;
 	const FSlateBrush* GetAffectsWeightmapBrush(int32 InBrushIndex) const;
+	const FSlateBrush* GetAffectsVisibilityLayerBrush(int32 InBrushIndex) const;
 	const FSlateBrush* GetVisibilityBrush(int32 InBrushIndex) const;
 	bool IsBrushSelected(int32 InBrushIndex) const;
 	bool IsBrushEnabled(int32 InBrushIndex) const;
diff --git a/Engine/Source/Editor/LandscapeEditor/Private/SLandscapeEditor.cpp b/Engine/Source/Editor/LandscapeEditor/Private/SLandscapeEditor.cpp
--- a/Engine/Source/Editor/LandscapeEditor/Private/SLandscapeEditor.cpp
+++ b/Engine/Source/Editor/LandscapeEditor/Private/SLandscapeEditor.cpp
@@ -13,6 +13,7 @@
 #include "Styling/AppStyle.h"
 #include "EditorModeManager.h"
 #include "EditorModes.h"
+#include "LandscapeEditTypes.h"
 #include "LandscapeEditorCommands.h"
 #include "LandscapeEditorObject.h"
 #include "IDetailsView.h"
@@ -621,17 +622,17 @@ bool FLandscapeToolKit::GetIsPropertyVisibleFromProperty(const FProperty& Proper
 			TArray<FString> ShowForTargetTypes;
 			Property.GetMetaData("ShowForTargetTypes").ParseIntoArray(ShowForTargetTypes, TEXT(","), true);
 
-			const ELandscapeToolTargetType::Type CurrentTargetType = LandscapeEdMode->CurrentToolTarget.TargetType;
+			const ELandscapeToolTargetType CurrentTargetType = LandscapeEdMode->CurrentToolTarget.TargetType;
 			// ELandscapeToolTargetType::Invalid means "weightmap with no valid paint layer" so we still want to display that property if it has been marked to be displayed in Weightmap target type, to be consistent 
 			//  with other paint brush properties (that don't use ShowForTargetTypes), which are still displayed in that case, even if they are ineffective :
 			if ((CurrentTargetType == ELandscapeToolTargetType::Invalid) 
-				&& (ShowForTargetTypes.FindByKey(TargetTypeNames[ELandscapeToolTargetType::Weightmap]) != nullptr))
+				&& (ShowForTargetTypes.FindByKey(TargetTypeNames[static_cast<uint8>(ELandscapeToolTargetType::Weightmap)]) != nullptr))
 			{ 
 				return true;
 			}
 			// Otherwise, hide it, if ShowForTargetTypes was used on this property but doesn't correspond to the current target type :
 			else if ((CurrentTargetType == ELandscapeToolTargetType::Invalid)
-				|| (ShowForTargetTypes.FindByKey(TargetTypeNames[CurrentTargetType]) == nullptr))
+				|| (ShowForTargetTypes.FindByKey(TargetTypeNames[static_cast<uint8>(CurrentTargetType)]) == nullptr))
 			{
 				return false;
 			}
diff --git a/Engine/Source/Editor/LandscapeEditor/Public/LandscapeToolInterface.h b/Engine/Source/Editor/LandscapeEditor/Public/LandscapeToolInterface.h
--- a/Engine/Source/Editor/LandscapeEditor/Public/LandscapeToolInterface.h
+++ b/Engine/Source/Editor/LandscapeEditor/Public/LandscapeToolInterface.h
@@ -8,6 +8,7 @@
 #include "UnrealWidgetFwd.h"
 #include "EdMode.h"
 #include "LandscapeEdit.h"
+#include "LandscapeEditTypes.h"
 
 class FEditorViewportClient;
 class FPrimitiveDrawInterface;
@@ -146,28 +147,28 @@ namespace ELandscapeToolTargetTypeMask
 {
 	enum Type : uint8
 	{
-		Heightmap  = 1 << ELandscapeToolTargetType::Heightmap,
-		Weightmap  = 1 << ELandscapeToolTargetType::Weightmap,
-		Visibility = 1 << ELandscapeToolTargetType::Visibility,
+		Heightmap  = 1 << static_cast<uint8>(ELandscapeToolTargetType::Heightmap),
+		Weightmap  = 1 << static_cast<uint8>(ELandscapeToolTargetType::Weightmap),
+		Visibility = 1 << static_cast<uint8>(ELandscapeToolTargetType::Visibility),
 
 		NA = 0,
 		All = 0xFF,
 	};
 
-	inline ELandscapeToolTargetTypeMask::Type FromType(ELandscapeToolTargetType::Type TargetType)
+	inline ELandscapeToolTargetTypeMask::Type FromType(ELandscapeToolTargetType TargetType)
 	{
 		if (TargetType == ELandscapeToolTargetType::Invalid)
 		{
 			return ELandscapeToolTargetTypeMask::NA;
 		}
-		return (ELandscapeToolTargetTypeMask::Type)(1 << TargetType);
+		return (ELandscapeToolTargetTypeMask::Type)(1 << static_cast<uint8>(TargetType));
 	}
 }
 
 struct FLandscapeToolTarget
 {
 	TWeakObjectPtr<ULandscapeInfo> LandscapeInfo;
-	ELandscapeToolTargetType::Type TargetType;
+	ELandscapeToolTargetType TargetType;
 	TWeakObjectPtr<ULandscapeLayerInfoObject> LayerInfo;
 	FName LayerName;
 	
diff --git a/Engine/Source/Runtime/Landscape/Classes/Landscape.h b/Engine/Source/Runtime/Landscape/Classes/Landscape.h
--- a/Engine/Source/Runtime/Landscape/Classes/Landscape.h
+++ b/Engine/Source/Runtime/Landscape/Classes/Landscape.h
@@ -7,6 +7,7 @@
 #include "UObject/ObjectMacros.h"
 #include "LandscapeProxy.h"
 #include "LandscapeBlueprintBrushBase.h"
+#include "LandscapeEditTypes.h"
 #include "Delegates/DelegateCombinations.h"
 
 #include "Landscape.generated.h"
@@ -24,11 +25,6 @@ struct FUpdateLayersContentContext;
 struct FEditLayersHeightmapMergeParams;
 struct FEditLayersWeightmapMergeParams;
 
-namespace ELandscapeToolTargetType
-{
-	enum Type : int8;
-};
-
 namespace EditLayersHeightmapLocalMerge_RenderThread
 {
 	struct FMergeInfo;
@@ -131,9 +127,13 @@ struct FLandscapeLayerBrush
 
 #if WITH_EDITOR
 	UTextureRenderTarget2D* Render(bool InIsHeightmap, const FIntRect& InLandscapeSize, UTextureRenderTarget2D* InLandscapeRenderTarget, const FName& InWeightmapLayerName = NAME_None);
+	
+	UTextureRenderTarget2D* RenderLayer(const FIntRect& InLandscapeSize, const FLandscapeBrushParameters& InParameters);
+
 	LANDSCAPE_API ALandscapeBlueprintBrushBase* GetBrush() const;
 	bool IsAffectingHeightmap() const;
 	bool IsAffectingWeightmapLayer(const FName& InWeightmapLayerName) const;
+	bool IsAffectingVisibilityLayer() const;
 	void SetOwner(ALandscape* InOwner);
 #endif
 
@@ -360,6 +360,8 @@ public:
 	LANDSCAPE_API bool GetUseGeneratedLandscapeSplineMeshesActors() const;
 	LANDSCAPE_API bool PrepareTextureResources(bool bInWaitForStreaming);
 
+	LANDSCAPE_API bool GetVisibilityLayerAllocationIndex() const { return 0; }
+
 protected:
 	FName GenerateUniqueLayerName(FName InName = NAME_None) const;
 
@@ -525,7 +527,7 @@ private:
 		int32 ViewMode;
 		FGuid SelectedLayer;
 		TWeakObjectPtr<ULandscapeLayerInfoObject> SelectedLayerInfoObject;
-		ELandscapeToolTargetType::Type ToolTarget;
+		ELandscapeToolTargetType ToolTarget;
 	};
 
 	FLandscapeEdModeInfo LandscapeEdModeInfo;
diff --git a/Engine/Source/Runtime/Landscape/Private/Landscape.cpp b/Engine/Source/Runtime/Landscape/Private/Landscape.cpp
--- a/Engine/Source/Runtime/Landscape/Private/Landscape.cpp
+++ b/Engine/Source/Runtime/Landscape/Private/Landscape.cpp
@@ -3320,6 +3320,17 @@ bool ULandscapeInfo::UpdateLayerInfoMapInternal(ALandscapeProxy* Proxy, bool bIn
 						Proxy->EditorLayerSettings.RemoveAt(i--);
 					}
 				}
+
+				// Add Visibility Layer info if not initialized
+				if (ALandscapeProxy::VisibilityLayer != nullptr)
+				{
+					int32 LayerInfoIndex = GetLayerInfoIndex(ALandscapeProxy::VisibilityLayer->LayerName);
+
+					if ((LayerInfoIndex != INDEX_NONE) && (Layers[LayerInfoIndex].LayerInfoObj == nullptr))
+					{
+						Layers[LayerInfoIndex].LayerInfoObj = ALandscapeProxy::VisibilityLayer;
+					}
+				}
 			}
 		}
 		else // !Proxy
diff --git a/Engine/Source/Runtime/Landscape/Private/LandscapeBlueprintBrushBase.cpp b/Engine/Source/Runtime/Landscape/Private/LandscapeBlueprintBrushBase.cpp
--- a/Engine/Source/Runtime/Landscape/Private/LandscapeBlueprintBrushBase.cpp
+++ b/Engine/Source/Runtime/Landscape/Private/LandscapeBlueprintBrushBase.cpp
@@ -4,7 +4,9 @@
 #include "CoreMinimal.h"
 #include "LandscapeProxy.h"
 #include "Landscape.h"
+#include "LandscapeEditTypes.h"
 #include "LandscapeInfo.h"
+#include "LandscapeLayerInfoObject.h"
 #include "LandscapePrivate.h"
 #include "Misc/MapErrors.h"
 #include "Misc/UObjectToken.h"
@@ -30,6 +32,7 @@ ALandscapeBlueprintBrushBase::ALandscapeBlueprintBrushBase(const FObjectInitiali
 	, UpdateOnPropertyChange(true)
 	, AffectHeightmap(false)
 	, AffectWeightmap(false)
+	, AffectVisibilityLayer(false)
 	, bIsVisible(true)
 	, LastRequestLayersContentUpdateFrameNumber(InvalidLastRequestLayersContentUpdateFrameNumber)
 #endif
@@ -50,11 +53,28 @@ ALandscapeBlueprintBrushBase::ALandscapeBlueprintBrushBase(const FObjectInitiali
 #endif
 }
 
+PRAGMA_DISABLE_DEPRECATION_WARNINGS
+
 UTextureRenderTarget2D* ALandscapeBlueprintBrushBase::Render_Implementation(bool InIsHeightmap, UTextureRenderTarget2D* InCombinedResult, const FName& InWeightmapLayerName)
 {
 	return Render_Native(InIsHeightmap, InCombinedResult, InWeightmapLayerName);
 }
 
+UTextureRenderTarget2D* ALandscapeBlueprintBrushBase::RenderLayer_Implementation(const FLandscapeBrushParameters& InParameters)
+{
+	return RenderLayer_Native(InParameters);
+}
+
+UTextureRenderTarget2D* ALandscapeBlueprintBrushBase::RenderLayer_Native(const FLandscapeBrushParameters& InParameters)
+{
+	const bool bIsHeightmap = InParameters.LayerType == ELandscapeToolTargetType::Heightmap;
+
+	// Without any implementation, we call the former Render method so content created before the deprecation will still work as expected.
+	return Render(bIsHeightmap, InParameters.CombinedResult, InParameters.WeightmapLayerName);
+}
+
+PRAGMA_ENABLE_DEPRECATION_WARNINGS
+
 void ALandscapeBlueprintBrushBase::Initialize_Implementation(const FTransform& InLandscapeTransform, const FIntPoint& InLandscapeSize, const FIntPoint& InLandscapeRenderTargetSize)
 {
 	Initialize_Native(InLandscapeTransform, InLandscapeSize, InLandscapeRenderTargetSize);
@@ -71,7 +91,7 @@ void ALandscapeBlueprintBrushBase::RequestLandscapeUpdate(bool bInUserTriggered)
 		{
 			ModeMask |= ELandscapeLayerUpdateMode::Update_Heightmap_Editing_NoCollision;
 		}
-		if (AffectWeightmap)
+		if (AffectWeightmap || AffectVisibilityLayer)
 		{
 			ModeMask |= ELandscapeLayerUpdateMode::Update_Weightmap_Editing_NoCollision;
 		}
@@ -100,7 +120,7 @@ void ALandscapeBlueprintBrushBase::PushDeferredLayersContentUpdate()
 		{
 			ModeMask |= ELandscapeLayerUpdateMode::Update_Heightmap_All;
 		}
-		if (AffectWeightmap)
+		if (AffectWeightmap || AffectVisibilityLayer)
 		{
 			ModeMask |= ELandscapeLayerUpdateMode::Update_Weightmap_All;
 		}
@@ -171,6 +191,18 @@ void ALandscapeBlueprintBrushBase::SetAffectsWeightmap(bool bInAffectsWeightmap)
 #endif
 }
 
+void ALandscapeBlueprintBrushBase::SetAffectsVisibilityLayer(bool bInAffectsVisibilityLayer)
+{
+#if WITH_EDITORONLY_DATA
+	Modify();
+	AffectVisibilityLayer = bInAffectsVisibilityLayer;
+	if (OwningLandscape)
+	{
+		OwningLandscape->OnBlueprintBrushChanged();
+	}
+#endif
+}
+
 bool ALandscapeBlueprintBrushBase::IsAffectingWeightmapLayer(const FName& InLayerName) const
 {
 #if WITH_EDITORONLY_DATA
diff --git a/Engine/Source/Runtime/Landscape/Private/LandscapeEditLayers.cpp b/Engine/Source/Runtime/Landscape/Private/LandscapeEditLayers.cpp
--- a/Engine/Source/Runtime/Landscape/Private/LandscapeEditLayers.cpp
+++ b/Engine/Source/Runtime/Landscape/Private/LandscapeEditLayers.cpp
@@ -15,6 +15,7 @@ LandscapeEditLayers.cpp: Landscape editing layers mode
 #include "LandscapePrivate.h"
 #include "LandscapeEditReadback.h"
 #include "LandscapeEditResources.h"
+#include "LandscapeEditTypes.h"
 #include "LandscapeUtils.h"
 #include "LandscapeSubsystem.h"
 
@@ -4890,7 +4891,9 @@ int32 ALandscape::PerformLayersHeightmapsGlobalMerge(const FUpdateLayersContentC
 				// TODO: handle conversion/handling of RT not same size as internal size
 
 				FLandscapeLayerBrush& Brush = Layer.Brushes[i];
-				UTextureRenderTarget2D* BrushOutputNonAtlasRT = Brush.Render(true, LandscapeExtent, CombinedHeightmapNonAtlasRT);
+				FLandscapeBrushParameters BrushParameters(ELandscapeToolTargetType::Heightmap, CombinedHeightmapNonAtlasRT);
+
+				UTextureRenderTarget2D* BrushOutputNonAtlasRT = Brush.RenderLayer(LandscapeExtent, BrushParameters);
 				if ((BrushOutputNonAtlasRT == nullptr) 
 					|| (BrushOutputNonAtlasRT->SizeX != CombinedHeightmapNonAtlasRT->SizeX)
 					|| (BrushOutputNonAtlasRT->SizeY != CombinedHeightmapNonAtlasRT->SizeY))
@@ -6891,9 +6894,20 @@ int32 ALandscape::PerformLayersWeightmapsGlobalMerge(FUpdateLayersContentContext
 						for (int32 i = 0; i < Layer.Brushes.Num(); ++i)
 						{
 							FLandscapeLayerBrush& Brush = Layer.Brushes[i];
+							TOptional<int32> LayerInfoSettingsAllocatedIndex;
+
 							if (Brush.IsAffectingWeightmapLayer(InfoLayerSettings.GetLayerName()) && !LayerInfoObjects.Contains(InfoLayerSettings.LayerInfoObj))
 							{
-								LayerInfoObjects.Add(InfoLayerSettings.LayerInfoObj, LayerInfoSettingsIndex + 1); // due to visibility layer that is at 0
+								LayerInfoSettingsAllocatedIndex = LayerInfoSettingsIndex + 1; // due to visibility layer that is at 0
+							}
+							else if (Brush.IsAffectingVisibilityLayer() && UE::Landscape::IsVisibilityLayer(InfoLayerSettings.LayerInfoObj) && !LayerInfoObjects.Contains(InfoLayerSettings.LayerInfoObj))
+							{
+								LayerInfoSettingsAllocatedIndex = GetVisibilityLayerAllocationIndex();
+							}
+
+							if (LayerInfoSettingsAllocatedIndex.IsSet())
+							{
+								LayerInfoObjects.Add(InfoLayerSettings.LayerInfoObj, LayerInfoSettingsAllocatedIndex.GetValue());
 								bHasWeightmapData = true;
 							}
 						}
@@ -7041,13 +7055,17 @@ int32 ALandscape::PerformLayersWeightmapsGlobalMerge(FUpdateLayersContentContext
 					// Handle brush blending
 					if (Layer.bVisible && !InMergeParams.bSkipBrush)
 					{
+						const ELandscapeToolTargetType LayerType = UE::Landscape::IsVisibilityLayer(LayerInfoObj) ? ELandscapeToolTargetType::Visibility : ELandscapeToolTargetType::Weightmap;
+
 						// Draw each brushes				
 						for (int32 i = 0; i < Layer.Brushes.Num(); ++i)
 						{
 							// TODO: handle conversion/handling of RT not same size as internal size
 
 							FLandscapeLayerBrush& Brush = Layer.Brushes[i];
-							UTextureRenderTarget2D* BrushOutputRT = Brush.Render(false, LandscapeExtent, LandscapeScratchRT3, LayerInfoObj->LayerName);
+							FLandscapeBrushParameters BrushParameters(LayerType, LandscapeScratchRT3, LayerInfoObj->LayerName);
+
+							UTextureRenderTarget2D* BrushOutputRT = Brush.RenderLayer( LandscapeExtent, BrushParameters);
 							if (BrushOutputRT == nullptr || BrushOutputRT->SizeX != LandscapeScratchRT3->SizeX || BrushOutputRT->SizeY != LandscapeScratchRT3->SizeY)
 							{
 								continue;
@@ -7860,7 +7878,7 @@ void ALandscape::MonitorLandscapeEdModeChanges()
 		bRequiredEditingClientFullUpdate = true;
 	}
 
-	ELandscapeToolTargetType::Type NewValue = LandscapeEdMode ? LandscapeEdMode->GetLandscapeToolTargetType() : ELandscapeToolTargetType::Invalid;
+	ELandscapeToolTargetType NewValue = LandscapeEdMode ? LandscapeEdMode->GetLandscapeToolTargetType() : ELandscapeToolTargetType::Invalid;
 	if (LandscapeEdModeInfo.ToolTarget != NewValue)
 	{
 		LandscapeEdModeInfo.ToolTarget = NewValue;
@@ -9744,22 +9762,58 @@ bool FLandscapeLayerBrush::IsAffectingWeightmapLayer(const FName& InWeightmapLay
 #endif
 }
 
+bool FLandscapeLayerBrush::IsAffectingVisibilityLayer() const
+{
+#if WITH_EDITORONLY_DATA
+	return BlueprintBrush && BlueprintBrush->IsVisible() && BlueprintBrush->IsAffectingVisibilityLayer();
+#else
+	return false;
+#endif
+}
+
 UTextureRenderTarget2D* FLandscapeLayerBrush::Render(bool InIsHeightmap, const FIntRect& InLandscapeExtent, UTextureRenderTarget2D* InLandscapeRenderTarget, const FName& InWeightmapLayerName)
+{
+	const ELandscapeToolTargetType LayerType = InIsHeightmap ? ELandscapeToolTargetType::Heightmap : ELandscapeToolTargetType::Weightmap;
+	FLandscapeBrushParameters BrushParameters(LayerType, InLandscapeRenderTarget, InWeightmapLayerName);
+
+	return RenderLayer(InLandscapeExtent, BrushParameters);
+}
+
+UTextureRenderTarget2D* FLandscapeLayerBrush::RenderLayer(const FIntRect& InLandscapeSize, const FLandscapeBrushParameters& InParameters)
 {
 #if WITH_EDITORONLY_DATA
 	TRACE_CPUPROFILER_EVENT_SCOPE(LandscapeLayers_LayerBrushRender);
-	if ((InIsHeightmap && !IsAffectingHeightmap()) ||
-		(!InIsHeightmap && !IsAffectingWeightmapLayer(InWeightmapLayerName)))
+	if (((InParameters.LayerType == ELandscapeToolTargetType::Heightmap) && !IsAffectingHeightmap()) ||
+		((InParameters.LayerType == ELandscapeToolTargetType::Weightmap) && !IsAffectingWeightmapLayer(InParameters.WeightmapLayerName)) ||
+		((InParameters.LayerType == ELandscapeToolTargetType::Visibility) && !IsAffectingVisibilityLayer()))
 	{
 		return nullptr;
 	}
-	if (Initialize(InLandscapeExtent, InLandscapeRenderTarget))
+
+	if (Initialize(InLandscapeSize, InParameters.CombinedResult))
 	{
-		FString ProfilingEventName = FString::Format(TEXT("LandscapeLayers_RenderLayerBrush {0}: {1}"), { InIsHeightmap ? TEXT("LS Height") : TEXT("LS Weight"), BlueprintBrush->GetName() });
+		const TCHAR* LayerName = TEXT("");
+		
+		switch (InParameters.LayerType)
+		{
+		case ELandscapeToolTargetType::Heightmap:
+			LayerName = TEXT("LS Height");
+			break;
+		case ELandscapeToolTargetType::Weightmap:
+			LayerName = TEXT("LS Weight");
+			break;
+		case ELandscapeToolTargetType::Visibility:
+			LayerName = TEXT("LS VisibilityLayer");
+			break;
+		default:
+			check(false);
+		}
+
+		FString ProfilingEventName = FString::Format(TEXT("LandscapeLayers_RenderLayerBrush {0}: {1}"), { LayerName , BlueprintBrush->GetName() });
 		SCOPED_DRAW_EVENTF_GAMETHREAD(LandscapeLayers, *ProfilingEventName);
 
 		TGuardValue<bool> AutoRestore(GAllowActorScriptExecutionInEditor, true);
-		UTextureRenderTarget2D* Result = BlueprintBrush->Render(InIsHeightmap, InLandscapeRenderTarget, InWeightmapLayerName);
+		UTextureRenderTarget2D* Result = BlueprintBrush->RenderLayer(InParameters);
 
 		return Result;
 	}
diff --git a/Engine/Source/Runtime/Landscape/Private/LandscapeUtils.cpp b/Engine/Source/Runtime/Landscape/Private/LandscapeUtils.cpp
--- a/Engine/Source/Runtime/Landscape/Private/LandscapeUtils.cpp
+++ b/Engine/Source/Runtime/Landscape/Private/LandscapeUtils.cpp
@@ -4,12 +4,13 @@
 
 #include "Engine/Level.h"
 #include "DataDrivenShaderPlatformInfo.h"
+#include "LandscapeLayerInfoObject.h"
+#include "LandscapeProxy.h"
 
 #if WITH_EDITOR
 #include "EditorDirectories.h"
 #include "Engine/Texture2D.h"
 #include "LandscapeComponent.h"
-#include "LandscapeLayerInfoObject.h"
 #include "ObjectTools.h"
 #endif
 
@@ -70,6 +71,11 @@ FString GetLayerInfoObjectPackageName(const ULevel* InLevel, const FName& InLaye
 	return PackageName;
 }
 
+bool IsVisibilityLayer(const ULandscapeLayerInfoObject* InLayerInfoObject)
+{
+	return (ALandscapeProxy::VisibilityLayer != nullptr) && (ALandscapeProxy::VisibilityLayer == InLayerInfoObject);
+}
+
 uint32 GetTypeHash(const FTextureCopyRequest& InKey)
 {
 	uint32 Hash = GetTypeHash(InKey.Source);
diff --git a/Engine/Source/Runtime/Landscape/Public/LandscapeBlueprintBrushBase.h b/Engine/Source/Runtime/Landscape/Public/LandscapeBlueprintBrushBase.h
--- a/Engine/Source/Runtime/Landscape/Public/LandscapeBlueprintBrushBase.h
+++ b/Engine/Source/Runtime/Landscape/Public/LandscapeBlueprintBrushBase.h
@@ -5,11 +5,40 @@
 #include "CoreMinimal.h"
 #include "UObject/ObjectMacros.h"
 #include "GameFramework/Actor.h"
+#include "LandscapeEditTypes.h"
 
 #include "LandscapeBlueprintBrushBase.generated.h"
 
 class UTextureRenderTarget2D;
 
+USTRUCT(BlueprintType)
+struct LANDSCAPE_API FLandscapeBrushParameters
+{
+	GENERATED_BODY()
+
+	FLandscapeBrushParameters()
+		: LayerType(ELandscapeToolTargetType::Invalid)
+		, CombinedResult(nullptr)
+		, WeightmapLayerName()
+	{}
+
+	FLandscapeBrushParameters(const ELandscapeToolTargetType& InLayerType, TObjectPtr<UTextureRenderTarget2D> InCombinedResult, const FName& InWeightmapLayerName = FName())
+		: LayerType(InLayerType)
+		, CombinedResult(InCombinedResult)
+		, WeightmapLayerName(InWeightmapLayerName)
+	{}
+
+	UPROPERTY(Category = "Settings", EditAnywhere, BlueprintReadWrite)
+	ELandscapeToolTargetType LayerType;
+
+	UPROPERTY(Category = "Settings", EditAnywhere, BlueprintReadWrite)
+	TObjectPtr<UTextureRenderTarget2D> CombinedResult;
+
+	UPROPERTY(Category = "Settings", EditAnywhere, BlueprintReadWrite)
+	FName WeightmapLayerName;
+};
+
+
 UCLASS(Abstract, NotBlueprintable)
 class LANDSCAPE_API ALandscapeBlueprintBrushBase : public AActor
 {
@@ -29,6 +58,9 @@ protected:
 	UPROPERTY(Category = "Settings", EditAnywhere, BlueprintReadWrite)
 	bool AffectWeightmap;
 
+	UPROPERTY(Category = "Settings", EditAnywhere, BlueprintReadWrite)
+	bool AffectVisibilityLayer;
+
 	UPROPERTY(Category = "Settings", EditAnywhere, BlueprintReadWrite)
 	TArray<FName> AffectedWeightmapLayers;
 
@@ -39,12 +71,19 @@ protected:
 #endif
 
 public:
+
+	UE_DEPRECATED(5.3, "Please use RenderLayer_Native instead.")
 	virtual UTextureRenderTarget2D* Render_Native(bool InIsHeightmap, UTextureRenderTarget2D* InCombinedResult, const FName& InWeightmapLayerName) {return nullptr;}
 	virtual void Initialize_Native(const FTransform& InLandscapeTransform, const FIntPoint& InLandscapeSize, const FIntPoint& InLandscapeRenderTargetSize) {}
 
-	UFUNCTION(BlueprintNativeEvent)
+	UFUNCTION(BlueprintNativeEvent, meta = (DeprecatedFunction, DeprecationMessage = "Please use RenderLayer instead."))
 	UTextureRenderTarget2D* Render(bool InIsHeightmap, UTextureRenderTarget2D* InCombinedResult, const FName& InWeightmapLayerName);
 
+	UFUNCTION(BlueprintNativeEvent)
+	UTextureRenderTarget2D* RenderLayer(const FLandscapeBrushParameters& InParameters);
+
+	virtual UTextureRenderTarget2D* RenderLayer_Native(const FLandscapeBrushParameters& InParameters);
+
 	UFUNCTION(BlueprintNativeEvent)
 	void Initialize(const FTransform& InLandscapeTransform, const FIntPoint& InLandscapeSize, const FIntPoint& InLandscapeRenderTargetSize);
 
@@ -64,6 +103,7 @@ public:
 
 	bool IsAffectingHeightmap() const { return AffectHeightmap; }
 	bool IsAffectingWeightmap() const { return AffectWeightmap; }
+	bool IsAffectingVisibilityLayer() const { return AffectVisibilityLayer; }
 	virtual bool IsAffectingWeightmapLayer(const FName& InLayerName) const;
 	bool IsVisible() const { return bIsVisible; }
 	bool IsLayerUpdatePending() const;
@@ -71,6 +111,7 @@ public:
 	void SetIsVisible(bool bInIsVisible);
 	void SetAffectsHeightmap(bool bInAffectsHeightmap);
 	void SetAffectsWeightmap(bool bInAffectsWeightmap);
+	void SetAffectsVisibilityLayer(bool bInAffectsVisibilityLayer);
 
 	virtual bool ShouldTickIfViewportsOnly() const override;
 	virtual void Tick(float DeltaSeconds) override;
diff --git a/Engine/Source/Runtime/Landscape/Public/LandscapeEdit.h b/Engine/Source/Runtime/Landscape/Public/LandscapeEdit.h
--- a/Engine/Source/Runtime/Landscape/Public/LandscapeEdit.h
+++ b/Engine/Source/Runtime/Landscape/Public/LandscapeEdit.h
@@ -14,6 +14,7 @@ LandscapeEdit.h: Classes for the editor to access to Landscape data
 #include "Components/ActorComponent.h"
 #include "LandscapeInfo.h"
 #include "LandscapeComponent.h"
+#include "LandscapeEditTypes.h"
 #include "LandscapeHeightfieldCollisionComponent.h"
 #include "InstancedFoliageActor.h"
 #include "LandscapeLayerInfoObject.h"
@@ -32,23 +33,13 @@ class ULandscapeLayerInfoObject;
 
 #if WITH_EDITOR
 
-namespace ELandscapeToolTargetType
-{
-	enum Type : int8
-	{
-		Heightmap = 0,
-		Weightmap = 1,
-		Visibility = 2,
-		Invalid = -1, // only valid for LandscapeEdMode->CurrentToolTarget.TargetType
-	};
-}
 
 /** Landscape EdMode Interface (used by ALandscape in Editor mode to access EdMode properties) */
 class ILandscapeEdModeInterface
 {
 public:
 	virtual void PostUpdateLayerContent() = 0;
-	virtual ELandscapeToolTargetType::Type GetLandscapeToolTargetType() const = 0;
+	virtual ELandscapeToolTargetType GetLandscapeToolTargetType() const = 0;
 	virtual const FLandscapeLayer* GetLandscapeSelectedLayer() const = 0;
 	virtual ULandscapeLayerInfoObject* GetSelectedLandscapeLayerInfo() const = 0;
 	virtual void OnCanHaveLayersContentChanged() = 0;
diff --git a/Engine/Source/Runtime/Landscape/Public/LandscapeEditTypes.h b/Engine/Source/Runtime/Landscape/Public/LandscapeEditTypes.h
new file mode 100644
--- /dev/null
+++ b/Engine/Source/Runtime/Landscape/Public/LandscapeEditTypes.h
@@ -0,0 +1,14 @@
+// Copyright Epic Games, Inc. All Rights Reserved.
+
+#pragma once
+
+#include "CoreMinimal.h"
+
+UENUM()
+enum class ELandscapeToolTargetType : uint8
+{
+	Heightmap = 0,
+	Weightmap = 1,
+	Visibility = 2,
+	Invalid = 3, // only valid for LandscapeEdMode->CurrentToolTarget.TargetType
+};
diff --git a/Engine/Source/Runtime/Landscape/Public/LandscapeUtils.h b/Engine/Source/Runtime/Landscape/Public/LandscapeUtils.h
--- a/Engine/Source/Runtime/Landscape/Public/LandscapeUtils.h
+++ b/Engine/Source/Runtime/Landscape/Public/LandscapeUtils.h
@@ -95,6 +95,9 @@ LANDSCAPE_API FString GetSharedAssetsPath(const ULevel* InLevel);
 */
 LANDSCAPE_API FString GetLayerInfoObjectPackageName(const ULevel* InLevel, const FName& InLayerName, FName& OutLayerObjectName);
 
+/** Returns true if the provided layer info object is the current visibility layer. */
+LANDSCAPE_API bool IsVisibilityLayer(const ULandscapeLayerInfoObject* InLayerInfoObject);
+
 #endif //!WITH_EDITOR
 
 } // end namespace UE::Landscape
