Adds forgotten generated header

diff --git a/Engine/Source/Runtime/Landscape/Public/LandscapeEditTypes.h b/Engine/Source/Runtime/Landscape/Public/LandscapeEditTypes.h
--- a/Engine/Source/Runtime/Landscape/Public/LandscapeEditTypes.h
+++ b/Engine/Source/Runtime/Landscape/Public/LandscapeEditTypes.h
@@ -4,6 +4,8 @@
 
 #include "CoreMinimal.h"
 
+#include "LandscapeEditTypes.generated.h"
+
 UENUM()
 enum class ELandscapeToolTargetType : uint8
 {
