Adds back meta tag, lost following merge.

diff --git a/Engine/Source/Runtime/Landscape/Public/LandscapeEditTypes.h b/Engine/Source/Runtime/Landscape/Public/LandscapeEditTypes.h
--- a/Engine/Source/Runtime/Landscape/Public/LandscapeEditTypes.h
+++ b/Engine/Source/Runtime/Landscape/Public/LandscapeEditTypes.h
@@ -12,5 +12,5 @@ enum class ELandscapeToolTargetType : uint8
 	Heightmap = 0,
 	Weightmap = 1,
 	Visibility = 2,
-	Invalid = 3, // only valid for LandscapeEdMode->CurrentToolTarget.TargetType
+	Invalid = 3 UMETA(Hidden), // only valid for LandscapeEdMode->CurrentToolTarget.TargetType
 };
