Removed debug code

#skipci
#rnx

diff --git a/Engine/Source/Runtime/Landscape/Private/LandscapeSubsystem.cpp b/Engine/Source/Runtime/Landscape/Private/LandscapeSubsystem.cpp
--- a/Engine/Source/Runtime/Landscape/Private/LandscapeSubsystem.cpp
+++ b/Engine/Source/Runtime/Landscape/Private/LandscapeSubsystem.cpp
@@ -188,10 +188,7 @@ void ULandscapeSubsystem::Tick(float DeltaTime)
 	Super::Tick(DeltaTime);
 
 #if WITH_EDITOR
-	{
-		TRACE_CPUPROFILER_EVENT_SCOPE(YOUPI);
-		AppCurrentDateTime = FDateTime::Now();
-	}
+	AppCurrentDateTime = FDateTime::Now();
 
 	//Check if we need to start or stop creating Collision SceneProxies
 	ILandscapeModule& LandscapeModule = FModuleManager::GetModuleChecked<ILandscapeModule>("Landscape");
