Fix for crash undoing creation of landscape

diff --git a/Engine/Source/Runtime/Landscape/Private/LandscapeEdit.cpp b/Engine/Source/Runtime/Landscape/Private/LandscapeEdit.cpp
--- a/Engine/Source/Runtime/Landscape/Private/LandscapeEdit.cpp
+++ b/Engine/Source/Runtime/Landscape/Private/LandscapeEdit.cpp
@@ -675,13 +675,15 @@ void ULandscapeComponent::PreEditUndo()
 
 void ULandscapeComponent::PostEditUndo()
 {
-	ALandscapeProxy* Proxy = GetLandscapeProxy();
-	
-	// PreEditUndo of the LandscapeStreamingProxy resets it's pointer to the top level landscape so this ensures it's fixed up prior
-	// to requesting HeightMap or WeightMap updates
-	if (ULandscapeInfo* LandscapeInfo = Proxy->GetLandscapeInfo(); !LandscapeInfo->IsRegistered(Proxy))
+	ALandscapeProxy* Proxy = GetLandscapeProxy();	
+	if (IsValid(Proxy)) // this post edit can be deleting the landscape
 	{
-		LandscapeInfo->RegisterActor(Proxy, true);
+		// PreEditUndo of the LandscapeStreamingProxy resets it's pointer to the top level landscape so this ensures it's fixed up prior
+		// to requesting HeightMap or WeightMap updates
+		if (ULandscapeInfo* LandscapeInfo = Proxy->GetLandscapeInfo(); !LandscapeInfo->IsRegistered(Proxy))
+		{
+			LandscapeInfo->RegisterActor(Proxy, true);
+		}
 	}
 
 	// On undo, request a recompute weightmap usages, as they are not transacted (they combine information between multiple components)
