Added command to dump a report of all landscape grass data currently used on landscape components: grass.DumpGrassData [optional: -csv -fullnames -showempty -detailed -byproxy -bycomponent -bygrasstype -full]
- csv formats the report in a CSV-friendly way.
- fullnames displays the listed objects' full names, rather than the user-friendly version.
- showempty will dump info even from components with no grass data
- detailed shows a detailed report of all grass data, for all grass types, in all landscape components.
- byproxy shows a report of grass data per landscape proxy.
- bycomponent shows a report of grass data per landscape component.
- bygrasstype shows a report of grass data per grass type.
- full enables all sub-reports.
If no report type option specified, assume full report is requested

diff --git a/Engine/Source/Runtime/Landscape/Private/LandscapeGrass.cpp b/Engine/Source/Runtime/Landscape/Private/LandscapeGrass.cpp
--- a/Engine/Source/Runtime/Landscape/Private/LandscapeGrass.cpp
+++ b/Engine/Source/Runtime/Landscape/Private/LandscapeGrass.cpp
@@ -52,6 +52,7 @@
 #include "LandscapeRender.h"
 #include "MaterialCompiler.h"
 #include "Algo/Accumulate.h"
+#include "Algo/ForEach.h"
 #include "UObject/Package.h"
 #include "Engine/StaticMesh.h"
 #include "Components/InstancedStaticMeshComponent.h"
@@ -220,7 +221,588 @@ static FAutoConsoleVariableRef CVarGGrassDensityQualityLevelCVar(
 	GGrassQualityLevel,
 	TEXT("The quality level for grass (low, medium, high, epic). \n"),
 	ECVF_Scalability);
-	
+
+static FAutoConsoleCommand ConsoleCommandDumpLandscapeGrassData(
+	TEXT("grass.DumpGrassData"),
+	TEXT("[optional: -csv -detailed -byproxy -bycomponent -bygrasstype -full] - Dumps a report of all grass data being currently used on landscape components. \
+		-csv formats the report in a CSV-friendly way. \
+		-fullnames displays the listed objects' full names, rather than the user-friendly version. \
+		-showempty will dump info even from components with no grass data \
+		-detailed shows a detailed report of all grass data, for all grass types, in all landscape components. \
+		-byproxy shows a report of grass data per landscape proxy. \
+		-bycomponent shows a report of grass data per landscape component. \
+		-bygrasstype shows a report of grass data per grass type. \
+		-full enables all sub-reports. \
+		If no report type option specified, assume full report is requested."),
+	FConsoleCommandWithArgsDelegate::CreateLambda([](const TArray<FString>& Args)
+	{
+		bool bCSV = false;
+		bool bUseFullNames = false;
+		bool bIncludeEmpty = false;
+		bool bDetailed = false;
+		bool bByProxy = false;
+		bool bByComponent = false;
+		bool bByGrassType = false;
+		bool bFull = false;
+		for (const FString& Arg : Args)
+		{
+			if (FParse::Param(*Arg, TEXT("csv")))
+			{
+				bCSV = true;
+			}
+			if (FParse::Param(*Arg, TEXT("fullnames")))
+			{
+				bUseFullNames = true;
+			}
+			if (FParse::Param(*Arg, TEXT("showempty")))
+			{
+				bIncludeEmpty = true;
+			}
+			if (FParse::Param(*Arg, TEXT("detailed")))
+			{
+				bDetailed = true;
+			}
+			if (FParse::Param(*Arg, TEXT("byproxy")))
+			{
+				bByProxy = true;
+			}
+			if (FParse::Param(*Arg, TEXT("bycomponent")))
+			{
+				bByComponent = true;
+			}
+			if (FParse::Param(*Arg, TEXT("bygrasstype")))
+			{
+				bByGrassType = true;
+			}
+			if (FParse::Param(*Arg, TEXT("full")))
+			{
+				bFull = true;
+			}
+		}
+		if (!bDetailed && !bByProxy && !bByComponent && !bByGrassType && !bFull)
+		{
+			// If nothing specified, assume full report :
+			bFull = true;
+		}
+
+		struct FSortDataEntry
+		{
+			FSortDataEntry(const UWorld* InWorld, const ALandscape* InActor, const ALandscapeProxy* InProxy, const ULandscapeComponent* InComponent, const ULandscapeGrassType* InGrassType, int32 InOffset)
+				: World(InWorld)
+				, Actor(InActor)
+				, Proxy(InProxy)
+				, Component(InComponent)
+				, GrassType(InGrassType)
+				, Offset(InOffset)
+			{}
+			const UWorld* World = nullptr;
+			const ALandscape* Actor = nullptr;
+			const ALandscapeProxy* Proxy = nullptr;
+			const ULandscapeComponent* Component = nullptr;
+			const ULandscapeGrassType* GrassType = nullptr; // If nullptr, this entry corresponds to height data 
+			int32 Offset = 0;
+			int32 Size = 0;
+		};
+		TArray<FSortDataEntry> SortData;
+		TSet<const ULandscapeGrassType*> UniqueGrassTypes;
+
+		// Gather all data first : 
+		for (TObjectIterator<const ULandscapeComponent> It(RF_ClassDefaultObject, true, EInternalObjectFlags::Garbage); It; ++It)
+		{
+			const ULandscapeComponent* Component = *It;
+			if (Component->GrassData->HasValidData() && (bIncludeEmpty || !Component->GrassData->HeightWeightData.IsEmpty()))
+			{
+				if (const UWorld* World = Component->GetWorld())
+				{
+					if (const ALandscapeProxy* Proxy = Component->GetLandscapeProxy())
+					{
+						if (const ALandscape* Actor = Proxy->GetLandscapeActor())
+						{
+							int32 MinWeightOffset = MAX_int32;
+							for (auto ItPair : Component->GrassData->WeightOffsets)
+							{
+								const ULandscapeGrassType* GrassType = ItPair.Key;
+								int32 GrassTypeOffset = ItPair.Value;
+								SortData.Add(FSortDataEntry(World, Actor, Proxy, Component, GrassType, GrassTypeOffset));
+								MinWeightOffset = FMath::Min(GrassTypeOffset, MinWeightOffset);
+
+								UniqueGrassTypes.Add(GrassType);
+							}
+							// Height data, if any, is at the beginning so if one the weight offsets is at index 0, then there's no height data on this component at all
+							int32 HeightOffset = (MinWeightOffset == 0) ? INDEX_NONE : 0;
+							// Always add an entry for height, even if no height data was found (makes reporting easier) : 
+							SortData.Add(FSortDataEntry(World, Actor, Proxy, Component, /* InGrassType = */nullptr, HeightOffset));
+						}
+					}
+				}
+			}
+		}
+
+		const int32 NumEntries = SortData.Num();
+		if (NumEntries == 0)
+		{
+			GLog->Logf(TEXT("No grass data."));
+			return;
+		}
+
+		// Then compute the size of each landscape grass type within a component : 
+		{
+			// Sort by Component then offset: 
+			SortData.Sort([](const FSortDataEntry& InLHS, const FSortDataEntry& InRHS)
+			{
+				if (InLHS.Component == InRHS.Component)
+				{
+					return InLHS.Offset < InRHS.Offset;
+				}
+				else
+				{
+					return InLHS.Component < InRHS.Component;
+				}
+			});
+
+			// Then iterate over all entries and use the next item's offset to compute the previous item's size, component by component:
+			const ULandscapeComponent* CurrentComponent = SortData[0].Component;
+			int32 CurrentOffsetInComponent = 0;
+			for (int32 NextIndex = 1; NextIndex < NumEntries; ++NextIndex)
+			{
+				const FSortDataEntry& NextEntry = SortData[NextIndex];
+				if (NextEntry.Component != CurrentComponent)
+				{
+					SortData[NextIndex - 1].Size = CurrentComponent->GrassData->HeightWeightData.Num() - CurrentOffsetInComponent;
+					CurrentComponent = NextEntry.Component;
+					CurrentOffsetInComponent = 0;
+				}
+				else
+				{
+					SortData[NextIndex - 1].Size = NextEntry.Offset - CurrentOffsetInComponent;
+					CurrentOffsetInComponent = NextEntry.Offset;
+				}
+			}
+
+			// Since the loop sets the size at index - 1, the last item's size is not yet computed : 
+			SortData[NumEntries - 1].Size = CurrentComponent->GrassData->HeightWeightData.Num() - CurrentOffsetInComponent;
+		}
+
+		// Sort by World, then Actor, then Proxy, then Component, then GrassType : 
+		SortData.Sort([](const FSortDataEntry& InLHS, const FSortDataEntry& InRHS)
+		{
+			if (InLHS.World == InRHS.World)
+			{
+				if (InLHS.Actor == InRHS.Actor)
+				{
+					if (InLHS.Proxy == InRHS.Proxy)
+					{
+						if (InLHS.Component == InRHS.Component)
+						{
+							return InLHS.GrassType < InRHS.GrassType;
+						}
+						else
+						{
+							return InLHS.Component < InRHS.Component;
+						}
+					}
+					else
+					{
+						return InLHS.Proxy < InRHS.Proxy;
+					}
+				}
+				else
+				{
+					return InLHS.Actor < InRHS.Actor;
+				}
+			}
+			else
+			{
+				return InLHS.World < InRHS.World;
+			}
+		});
+
+		auto GetFormattedObjectName = [bUseFullNames](const UObject* InObject) -> FString { return bUseFullNames ? InObject->GetFullName() : InObject->GetName(); };
+		auto GetFormattedActorName = [bUseFullNames](const AActor* InActor) -> FString { return bUseFullNames ? InActor->GetFullName() : InActor->GetActorNameOrLabel(); };
+		auto GetFormattedGrassTypeName = [bUseFullNames](const ULandscapeGrassType* InGrassType) -> FString { return (InGrassType != nullptr) ? (bUseFullNames ? InGrassType->GetFullName() : InGrassType->GetName()) : FString(TEXT("(HeightData)")); };
+		auto LogWorldActorLine = [GetFormattedObjectName, GetFormattedActorName](const UWorld* InWorld, const ALandscape* InActor) { GLog->Logf(TEXT("-- World: %s - Landscape Actor: %s --"), *GetFormattedObjectName(InWorld), *GetFormattedActorName(InActor)); };
+
+		// List of lines to display in a sub-report (one element per entry):
+		TArray<TPair<FString, FString>> CurrentSubReportContent; // Key = entry's name (for sorting), Value = line
+		auto AddLineToSubReport = [&CurrentSubReportContent](const FString& InEntryName, const FString& InLine) { CurrentSubReportContent.Emplace(InEntryName, InLine); };
+		auto FinalizeSubReport = [&CurrentSubReportContent]()
+		{
+			// Lexicographically sort entries, output them and reset the sub-report content :
+			CurrentSubReportContent.Sort([](const TPair<FString, FString>& InLHS, const TPair<FString, FString>& InRHS) { return (InLHS.Key.Compare(InRHS.Key) < 0); });
+			Algo::ForEach(CurrentSubReportContent, [](const TPair<FString, FString>& InEntry) { GLog->Log(*InEntry.Value); });
+			CurrentSubReportContent.Empty();
+		};
+
+		if (bFull || bDetailed)
+		{
+			// Start a new report :
+			GLog->Logf(TEXT("Detailed report:"));
+			ON_SCOPE_EXIT{ GLog->Logf(TEXT("--------")); };
+
+			if (bCSV)
+			{
+				GLog->Logf(TEXT("World,Landscape Actor,Landscape Proxy,Landscape Component,Landscape Grass Type,Size (KB)"));
+				for (const FSortDataEntry& Entry : SortData)
+				{
+					if ((Entry.GrassType != nullptr) || (bIncludeEmpty || (Entry.Size > 0)))
+					{
+						FString Line = FString::Printf(TEXT("%s,%s,%s,%s,%s,%.2f"),
+							*GetFormattedObjectName(Entry.World),
+							*GetFormattedActorName(Entry.Actor),
+							*GetFormattedActorName(Entry.Proxy),
+							*Entry.Component->GetName(),
+							*GetFormattedGrassTypeName(Entry.GrassType),
+							Entry.Size / 1024.0f);
+						// Use the line itself as the sort key since it's already formatted the way we want: 
+						AddLineToSubReport(Line, Line);
+					}
+				}
+				FinalizeSubReport();
+			}
+			else
+			{
+				const UWorld* CurrentWorld = SortData[0].World;
+				const ALandscape* CurrentActor = SortData[0].Actor;
+				const ALandscapeProxy* CurrentProxy = SortData[0].Proxy;
+				const ULandscapeComponent* CurrentComponent = SortData[0].Component;
+				bool bCurrentWorldHasGrassData = false;
+				bool bCurrentActorHasGrassData = false;
+				bool bCurrentProxyHasGrassData = false;
+				bool bCurrentComponentHasGrassData = false;
+				int32 NumWorldsWithGrassData = 0;
+				int32 NumActorsWithGrassData = 0;
+				int32 NumProxiesWithGrassData = 0;
+				int32 NumComponentsWithGrassData = 0;
+				int32 TotalHeightDataSize = 0;
+				int32 TotalWeightDataSize = 0;
+				for (int32 Index = 0; Index < NumEntries; ++Index)
+				{
+					const FSortDataEntry& Entry = SortData[Index];
+					if (CurrentWorld != Entry.World)
+					{
+						// This world was fully processed, take note of the info and start tracking the new world:
+						if (bCurrentWorldHasGrassData)
+						{
+							++NumWorldsWithGrassData;
+						}
+						CurrentWorld = Entry.World;
+						bCurrentWorldHasGrassData = false;
+					}
+
+					if (CurrentActor != Entry.Actor)
+					{
+						// This actor was fully processed, take note of the info and start tracking the new actor:
+						if (bCurrentActorHasGrassData)
+						{
+							++NumActorsWithGrassData;
+						}
+						CurrentActor = Entry.Actor;
+						bCurrentActorHasGrassData = false;
+					}
+
+					if (CurrentProxy != Entry.Proxy)
+					{
+						// This proxy was fully processed, take note of the info and start tracking the new proxy:
+						if (bCurrentProxyHasGrassData)
+						{
+							++NumProxiesWithGrassData;
+						}
+						CurrentProxy = Entry.Proxy;
+						bCurrentProxyHasGrassData = false;
+					}
+
+					// This component was fully processed, take note of the info and start tracking the new component:
+					if (CurrentComponent != Entry.Component)
+					{
+						if (bCurrentComponentHasGrassData)
+						{
+							++NumComponentsWithGrassData;
+						}
+						CurrentComponent = Entry.Component;
+						bCurrentComponentHasGrassData = false;
+					}
+
+					if (Entry.Size > 0)
+					{
+						bCurrentWorldHasGrassData = true;
+						bCurrentActorHasGrassData = true;
+						bCurrentProxyHasGrassData = true;
+						bCurrentComponentHasGrassData = true;
+						if (Entry.GrassType == nullptr)
+						{
+							TotalHeightDataSize += Entry.Size;
+						}
+						else
+						{
+							TotalWeightDataSize += Entry.Size;
+						}
+					}
+				}
+
+				// And dump one more time, for the last entry, which hasn't been taken into account yet :
+				if (bCurrentWorldHasGrassData)
+				{
+					++NumWorldsWithGrassData;
+				}
+				if (bCurrentActorHasGrassData)
+				{
+					++NumActorsWithGrassData;
+				}
+				if (bCurrentProxyHasGrassData)
+				{
+					++NumProxiesWithGrassData;
+				}
+				if (bCurrentComponentHasGrassData)
+				{
+					++NumComponentsWithGrassData;
+				}
+
+				GLog->Logf(TEXT("Num worlds : %i"), NumWorldsWithGrassData);
+				GLog->Logf(TEXT("Num landscape actors with grass data: %i"), NumActorsWithGrassData);
+				GLog->Logf(TEXT("Num landscape proxies with grass data: %i"), NumProxiesWithGrassData);
+				GLog->Logf(TEXT("Num landscape components with grass data: %i"), NumComponentsWithGrassData);
+				GLog->Logf(TEXT("Num grass types used: %i"), UniqueGrassTypes.Num());
+				GLog->Logf(TEXT("Total height data size (KB): %.2f"), TotalHeightDataSize / 1024.0f);
+				GLog->Logf(TEXT("Total weight data size (KB): %.2f"), TotalWeightDataSize / 1024.0f);
+				GLog->Logf(TEXT("Total grass data size (KB): %.2f"), (TotalHeightDataSize + TotalWeightDataSize) / 1024.0f);
+			}
+		}
+
+		// Assumes the list is sorted by world / by actor / by proxy already: 
+		if (bFull || bByProxy)
+		{
+			// Start a new report :
+			GLog->Logf(TEXT("By proxy report:"));
+			ON_SCOPE_EXIT{ GLog->Logf(TEXT("--------")); };
+
+			// Start a new sub-report per world / per actor:
+			LogWorldActorLine(SortData[0].World, SortData[0].Actor);
+			const ALandscape* CurrentActor = SortData[0].Actor;
+
+			auto LogCSVHeader = [bCSV]() { if (bCSV) { GLog->Logf(TEXT("Landscape Proxy,Num Grass Types,Total Height Data Size (KB), Total Weight Data Size (KB), Total Data Size (KB)")); } };
+			auto LogLine = [bCSV, AddLineToSubReport](const TCHAR* InProxyName, int32 InNumGrassTypes, int32 InTotalHeightDataSizeBytes, int32 InTotalWeightDataSizeBytes)
+			{
+				if (bCSV)
+				{
+					AddLineToSubReport(FString(InProxyName), FString::Printf(TEXT("%s,%i,%.2f,%.2f,%.2f"),
+						InProxyName,
+						InNumGrassTypes,
+						InTotalHeightDataSizeBytes / 1024.0f,
+						InTotalWeightDataSizeBytes / 1024.0f,
+						(InTotalHeightDataSizeBytes + InTotalWeightDataSizeBytes) / 1024.0f));
+				}
+				else
+				{
+					AddLineToSubReport(FString(InProxyName), FString::Printf(TEXT("%s: Num grass types = %i, total height data size = %.2f KB, total weight data size = %.2f KB, total data size = %.2f KB"),
+						InProxyName,
+						InNumGrassTypes,
+						InTotalHeightDataSizeBytes / 1024.0f,
+						InTotalWeightDataSizeBytes / 1024.0f,
+						(InTotalHeightDataSizeBytes + InTotalWeightDataSizeBytes) / 1024.0f));
+				}
+			};
+
+			LogCSVHeader();
+			const ALandscapeProxy* CurrentProxy = SortData[0].Proxy;
+			int32 HeightDataSizeForCurrentProxy = 0;
+			int32 WeightDataSizeForCurrentProxy = 0;
+			TSet<const ULandscapeGrassType*> UniqueGrassTypesForCurrentProxy;
+			for (int32 Index = 0; Index < NumEntries; ++Index)
+			{
+				const FSortDataEntry& Entry = SortData[Index];
+				if (Entry.Proxy != CurrentProxy)
+				{
+					// The info from CurrentProxy is now complete, we can dump it :
+					LogLine(*GetFormattedActorName(CurrentProxy), UniqueGrassTypesForCurrentProxy.Num(), HeightDataSizeForCurrentProxy, WeightDataSizeForCurrentProxy);
+
+					// Start gathering stats for the next proxy :
+					CurrentProxy = Entry.Proxy;
+					HeightDataSizeForCurrentProxy = 0;
+					WeightDataSizeForCurrentProxy = 0;
+					UniqueGrassTypesForCurrentProxy.Empty();
+				}
+
+				if (Entry.Actor != CurrentActor)
+				{
+					// Start a new sub-report per world / per actor:
+					FinalizeSubReport();
+					LogWorldActorLine(Entry.World, Entry.Actor);
+					LogCSVHeader();
+					CurrentActor = Entry.Actor;
+				}
+
+				// Accumulate data about this proxy :
+				if (Entry.Size > 0)
+				{
+					if (Entry.GrassType != nullptr)
+					{
+						WeightDataSizeForCurrentProxy += Entry.Size;
+						UniqueGrassTypesForCurrentProxy.Add(Entry.GrassType);
+					}
+					else
+					{
+						HeightDataSizeForCurrentProxy += Entry.Size;
+					}
+				}
+			}
+
+			// And dump one more time, for the last proxy, which hasn't been reported yet :
+			LogLine(*GetFormattedActorName(CurrentProxy), UniqueGrassTypesForCurrentProxy.Num(), HeightDataSizeForCurrentProxy, WeightDataSizeForCurrentProxy);
+			FinalizeSubReport();
+		}
+
+		// Assumes the list is sorted by world / by actor / ... / by component already: 
+		if (bFull || bByComponent)
+		{
+			// Start a new report :
+			GLog->Logf(TEXT("By component report:"));
+			ON_SCOPE_EXIT{ GLog->Logf(TEXT("--------")); };
+
+			// Start a new sub-report per world / per actor:
+			LogWorldActorLine(SortData[0].World, SortData[0].Actor);
+			const ALandscape* CurrentActor = SortData[0].Actor;
+
+			auto LogCSVHeader = [bCSV]() { if (bCSV) { GLog->Logf(TEXT("Landscape Proxy,Landscape Component,Num Grass Types,Total Height Data Size (KB), Total Weight Data Size (KB), Total Data Size (KB)")); } };
+			auto LogLine = [bCSV, AddLineToSubReport](const TCHAR* InProxyName, const TCHAR* InComponentName, int32 InNumGrassTypes, int32 InTotalHeightDataSizeBytes, int32 InTotalWeightDataSizeBytes)
+			{
+				if (bCSV)
+				{
+					AddLineToSubReport(FString(InComponentName), FString::Printf(TEXT("%s,%s,%i,%.2f,%.2f,%.2f"),
+						InProxyName,
+						InComponentName,
+						InNumGrassTypes,
+						InTotalHeightDataSizeBytes / 1024.0f,
+						InTotalWeightDataSizeBytes / 1024.0f,
+						(InTotalHeightDataSizeBytes + InTotalWeightDataSizeBytes) / 1024.0f));
+				}
+				else
+				{
+					AddLineToSubReport(FString(InComponentName), FString::Printf(TEXT("%s (%s): Num grass types = %i, total height data size = %.2f KB, total weight data size = %.2f KB, total data size = %.2f KB"),
+						InProxyName,
+						InComponentName,
+						InNumGrassTypes,
+						InTotalHeightDataSizeBytes / 1024.0f,
+						InTotalWeightDataSizeBytes / 1024.0f,
+						(InTotalHeightDataSizeBytes + InTotalWeightDataSizeBytes) / 1024.0f));
+				}
+			};
+
+			LogCSVHeader();
+			const ULandscapeComponent* CurrentComponent = SortData[0].Component;
+			int32 NumGrassTypesForCurrentComponent = 0;
+			int32 HeightDataSizeForCurrentComponent = 0;
+			int32 WeightDataSizeForCurrentComponent = 0;
+			for (int32 Index = 0; Index < NumEntries; ++Index)
+			{
+				const FSortDataEntry& Entry = SortData[Index];
+				if (Entry.Component != CurrentComponent)
+				{
+					// The info from CurrentComponent is now complete, we can dump it :
+					LogLine(*GetFormattedActorName(CurrentComponent->GetLandscapeProxy()), *CurrentComponent->GetName(), NumGrassTypesForCurrentComponent, HeightDataSizeForCurrentComponent, WeightDataSizeForCurrentComponent);
+
+					// Start gathering stats for the next component :
+					CurrentComponent = Entry.Component;
+					NumGrassTypesForCurrentComponent = 0;
+					HeightDataSizeForCurrentComponent = 0;
+					WeightDataSizeForCurrentComponent = 0;
+				}
+
+				if (Entry.Actor != CurrentActor)
+				{
+					// Start a new sub-report per world / per actor:
+					FinalizeSubReport();
+					LogWorldActorLine(Entry.World, Entry.Actor);
+					LogCSVHeader();
+					CurrentActor = Entry.Actor;
+				}
+
+				// Accumulate data about this component :
+				if (Entry.Size > 0)
+				{
+					if (Entry.GrassType != nullptr)
+					{
+						++NumGrassTypesForCurrentComponent;
+						WeightDataSizeForCurrentComponent += Entry.Size;
+					}
+					else
+					{
+						HeightDataSizeForCurrentComponent += Entry.Size;
+					}
+				}
+			}
+
+			// And dump one more time, for the last component, which hasn't been reported yet :
+			LogLine(*GetFormattedActorName(CurrentComponent->GetLandscapeProxy()), *CurrentComponent->GetName(), NumGrassTypesForCurrentComponent, HeightDataSizeForCurrentComponent, WeightDataSizeForCurrentComponent);
+			FinalizeSubReport();
+		}
+
+		if (bFull || bByGrassType)
+		{
+			// Sort by Grass Type, then Component: 
+			SortData.Sort([](const FSortDataEntry& InLHS, const FSortDataEntry& InRHS)
+			{
+				if (InLHS.GrassType == InRHS.GrassType)
+				{
+					return (InLHS.Component < InRHS.Component);
+				}
+				return (InLHS.GrassType < InRHS.GrassType);
+			});
+
+			// Start a new sub-report :
+			GLog->Logf(TEXT("By grass type report:"));
+			ON_SCOPE_EXIT{ GLog->Logf(TEXT("--------")); };
+
+			auto LogCSVHeader = [bCSV]() { if (bCSV) { GLog->Logf(TEXT("Grass Type,Num Components,Total Data Size (KB)")); } };
+			auto LogLine = [bCSV, AddLineToSubReport](const TCHAR* InGrassTypeName, int32 InNumComponents, int32 InTotalDataSizeBytes)
+			{
+				if (bCSV)
+				{
+					AddLineToSubReport(FString(InGrassTypeName), FString::Printf(TEXT("%s,%i,%.2f"),
+						InGrassTypeName,
+						InNumComponents,
+						InTotalDataSizeBytes / 1024.0f));
+				}
+				else
+				{
+					AddLineToSubReport(FString(InGrassTypeName), FString::Printf(TEXT("%s: Num components = %i, total data size = %.2f KB"),
+						InGrassTypeName,
+						InNumComponents,
+						InTotalDataSizeBytes / 1024.0f));
+				}
+			};
+
+			LogCSVHeader();
+			const ULandscapeGrassType* CurrentGrassType = SortData[0].GrassType;
+			int32 DataSizeForCurrentGrassType = 0;
+			int32 NumComponentsForCurrentGrassType = 0;
+			for (int32 Index = 0; Index < NumEntries; ++Index)
+			{
+				const FSortDataEntry& Entry = SortData[Index];
+				if (Entry.GrassType != CurrentGrassType)
+				{
+					// The info from CurrentGrassType is now complete, we can dump it :
+					LogLine(*GetFormattedGrassTypeName(CurrentGrassType), NumComponentsForCurrentGrassType, DataSizeForCurrentGrassType);
+
+					// Start gathering stats for the next grass type :
+					CurrentGrassType = Entry.GrassType;
+					DataSizeForCurrentGrassType = 0;
+					NumComponentsForCurrentGrassType = 0;
+				}
+
+				// Accumulate data about this grass type :
+				if (Entry.Size > 0)
+				{
+					DataSizeForCurrentGrassType += Entry.Size;
+					++NumComponentsForCurrentGrassType;
+				}
+			}
+
+			// And dump one more time, for the last grass type, which hasn't been reported yet :
+			LogLine(*GetFormattedGrassTypeName(CurrentGrassType), NumComponentsForCurrentGrassType, DataSizeForCurrentGrassType);
+			FinalizeSubReport();
+		}
+	}));
+
 struct FPerQualityLevelInt;
 struct FPerQualityLevelFloat;
 
