Add NullRHI checks so the landscapeEditLayers operates or performs a pass through when nullRHI is being used.

#localization none
#tests testing nullrhi tests
Change is very safe

diff --git a/Engine/Source/Runtime/Landscape/Private/LandscapeEditLayers.cpp b/Engine/Source/Runtime/Landscape/Private/LandscapeEditLayers.cpp
--- a/Engine/Source/Runtime/Landscape/Private/LandscapeEditLayers.cpp
+++ b/Engine/Source/Runtime/Landscape/Private/LandscapeEditLayers.cpp
@@ -2614,7 +2614,10 @@ struct FLandscapeLayersCopyReadbackTextureParams
 void ExecuteCopyToReadbackTexture(TArray<FLandscapeLayersCopyReadbackTextureParams>& InParams)
 {
 	SCOPED_DRAW_EVENTF_GAMETHREAD(LandscapeLayers, TEXT("Copy to readback textures (%d copies)"), InParams.Num());
-
+	if (GUsingNullRHI)
+	{
+		return;
+	}
 	for (FLandscapeLayersCopyReadbackTextureParams& Params : InParams)
 	{
 		Params.Dest->Enqueue(Params.Source, MoveTemp(Params.Context));
@@ -2888,6 +2891,10 @@ void ALandscape::DrawHeightmapComponentsToRenderTarget(const FString& InDebugNam
 {
 	check(InHeightmapRTRead != nullptr);
 	check(InHeightmapRTWrite != nullptr);
+	if (GUsingNullRHI)
+	{
+		return;
+	}
 
 	FIntPoint HeightmapWriteTextureSize(InHeightmapRTWrite->SizeX, InHeightmapRTWrite->SizeY);
 	FIntPoint HeightmapReadTextureSize(InHeightmapRTRead->Source.GetSizeX(), InHeightmapRTRead->Source.GetSizeY());
@@ -3450,7 +3457,7 @@ bool ALandscape::PrepareTextureResources(bool bInWaitForStreaming)
 	TRACE_CPUPROFILER_EVENT_SCOPE(LandscapeLayers_PrepareTextureResources);
 
 	ULandscapeInfo* Info = GetLandscapeInfo();
-	if (Info == nullptr)
+	if (Info == nullptr || GUsingNullRHI)
 	{
 		return false;
 	}
@@ -8015,7 +8022,7 @@ void ALandscape::UpdateLayersContent(bool bInWaitForStreaming, bool bInSkipMonit
 	bool bResourcesReady = PrepareTextureResources(bInWaitForStreaming);
 
 	ULandscapeInfo* LandscapeInfo = GetLandscapeInfo();
-	if ((LandscapeInfo == nullptr) || !CanHaveLayersContent() || !LandscapeInfo->AreAllComponentsRegistered())
+	if ((LandscapeInfo == nullptr) || !CanHaveLayersContent() || !LandscapeInfo->AreAllComponentsRegistered() || GUsingNullRHI)
 	{
 		return;
 	}
@@ -8734,6 +8741,11 @@ void ALandscape::FinishDestroy()
 
 bool ALandscape::IsUpToDate() const
 {
+	if (GUsingNullRHI)
+	{
+		return true;
+	}
+
 #if WITH_EDITORONLY_DATA
 	if (CanHaveLayersContent() && GetWorld() != nullptr && !GetWorld()->IsGameWorld())
 	{
