* Fixed bad tangent space computations for various landscape utilities (static lighting, Nanite...) : mimic what is done in the vertex shader when reconstructing the normal's Z in order to avoid quantization issues
* Also factorized lots of redundant code in landscape for accessing weightmap/heightmap data

diff --git a/Engine/Source/Editor/LandscapeEditor/Private/LandscapeEdModeComponentTools.cpp b/Engine/Source/Editor/LandscapeEditor/Private/LandscapeEdModeComponentTools.cpp
--- a/Engine/Source/Editor/LandscapeEditor/Private/LandscapeEdModeComponentTools.cpp
+++ b/Engine/Source/Editor/LandscapeEditor/Private/LandscapeEdModeComponentTools.cpp
@@ -2146,8 +2146,8 @@ public:
 						const int32 ImportX = FMath::Clamp<int32>(X - OffsetX, 0, ImportSizeX - 1);
 						const int32 ImportY0 = FMath::Clamp<int32>(Y0 - OffsetY, 0, ImportSizeY - 1);
 						const int32 ImportY1 = FMath::Clamp<int32>(Y1 - OffsetY, 0, ImportSizeY - 1);
-						const float Z0 = (static_cast<float>(ImportHeights[ImportX + ImportY0 * ImportSizeX]) - 32768.0f) * LANDSCAPE_ZSCALE;
-						const float Z1 = (static_cast<float>(ImportHeights[ImportX + ImportY1 * ImportSizeX]) - 32768.0f) * LANDSCAPE_ZSCALE;
+						const float Z0 = LandscapeDataAccess::GetLocalHeight(ImportHeights[ImportX + ImportY0 * ImportSizeX]);
+						const float Z1 = LandscapeDataAccess::GetLocalHeight(ImportHeights[ImportX + ImportY1 * ImportSizeX]);
 						return LineCoords{FVector(X, Y0, Z0), FVector(X, Y1, Z1)};
 					};
 
@@ -2160,8 +2160,8 @@ public:
 						const int32 ImportY = FMath::Clamp<int32>(Y - OffsetY, 0, ImportSizeY - 1);
 						const int32 ImportX0 = FMath::Clamp<int32>(X0 - OffsetX, 0, ImportSizeX - 1);
 						const int32 ImportX1 = FMath::Clamp<int32>(X1 - OffsetX, 0, ImportSizeX - 1);
-						const float Z0 = (static_cast<float>(ImportHeights[ImportX0 + ImportY * ImportSizeX]) - 32768.0f) * LANDSCAPE_ZSCALE;
-						const float Z1 = (static_cast<float>(ImportHeights[ImportX1 + ImportY * ImportSizeX]) - 32768.0f) * LANDSCAPE_ZSCALE;
+						const float Z0 = LandscapeDataAccess::GetLocalHeight(ImportHeights[ImportX0 + ImportY * ImportSizeX]);
+						const float Z1 = LandscapeDataAccess::GetLocalHeight(ImportHeights[ImportX1 + ImportY * ImportSizeX]);
 						return LineCoords{FVector(X0, Y, Z0), FVector(X1, Y, Z1)};
 					};
 
diff --git a/Engine/Source/Editor/LandscapeEditor/Private/LandscapeEdModeErosionTools.cpp b/Engine/Source/Editor/LandscapeEditor/Private/LandscapeEdModeErosionTools.cpp
--- a/Engine/Source/Editor/LandscapeEditor/Private/LandscapeEdModeErosionTools.cpp
+++ b/Engine/Source/Editor/LandscapeEditor/Private/LandscapeEdModeErosionTools.cpp
@@ -482,7 +482,7 @@ public:
 							if (SedimentDiff > 0)
 							{
 								SedimentData[Center] -= SedimentDiff;
-								HeightData[Center] = FMath::Clamp<uint16>(HeightData[Center] + SedimentDiff, 0, 65535);
+								HeightData[Center] = FMath::Clamp<uint16>(HeightData[Center] + SedimentDiff, 0, LandscapeDataAccess::MaxValue);
 							}
 						}
 					}
diff --git a/Engine/Source/Editor/LandscapeEditor/Private/LandscapeEdModeTools.h b/Engine/Source/Editor/LandscapeEditor/Private/LandscapeEdModeTools.h
--- a/Engine/Source/Editor/LandscapeEditor/Private/LandscapeEdModeTools.h
+++ b/Engine/Source/Editor/LandscapeEditor/Private/LandscapeEdModeTools.h
@@ -737,7 +737,7 @@ struct FXYOffsetmapAccessor
 				FVector* Value = Data.Find(FIntPoint(X, Y));
 				if (Value)
 				{
-					Value->Z = ((float)NewHeights.FindRef(FIntPoint(X, Y)) - 32768.0f) * LANDSCAPE_ZSCALE;
+					Value->Z = LandscapeDataAccess::GetLocalHeight(static_cast<float>(NewHeights.FindRef(FIntPoint(X, Y))));
 				}
 			}
 		}
@@ -756,7 +756,7 @@ struct FXYOffsetmapAccessor
 				FVector* Value = Data.Find(FIntPoint(X, Y));
 				if (Value)
 				{
-					Value->Z = ((float)NewHeights.FindRef(FIntPoint(X, Y)) - 32768.0f) * LANDSCAPE_ZSCALE;
+					Value->Z = LandscapeDataAccess::GetLocalHeight(static_cast<float>(NewHeights.FindRef(FIntPoint(X, Y))));
 				}
 			}
 		}
@@ -777,7 +777,7 @@ struct FXYOffsetmapAccessor
 			{
 				for (int32 X = X1; X <= X2; ++X)
 				{
-					NewHeights[X - X1 + (Y - Y1) * (X2 - X1 + 1)] = FMath::Clamp<uint16>(Data[(X - X1 + (Y - Y1) * (X2 - X1 + 1))].Z * LANDSCAPE_INV_ZSCALE + 32768.0f, 0, 65535);
+					NewHeights[X - X1 + (Y - Y1) * (X2 - X1 + 1)] = LandscapeDataAccess::GetTexHeight(Data[(X - X1 + (Y - Y1) * (X2 - X1 + 1))].Z);
 				}
 			}
 						
@@ -1330,7 +1330,7 @@ struct FHeightmapToolTarget
 
 	static FMatrix ToWorldMatrix(ULandscapeInfo* LandscapeInfo)
 	{
-		FMatrix Result = FTranslationMatrix(FVector(0, 0, -32768.0f));
+		FMatrix Result = FTranslationMatrix(FVector(0, 0, -LandscapeDataAccess::MidValue));
 		Result *= FScaleMatrix(FVector(1.0f, 1.0f, LANDSCAPE_ZSCALE) * LandscapeInfo->DrawScale);
 		return Result;
 	}
@@ -1338,7 +1338,7 @@ struct FHeightmapToolTarget
 	static FMatrix FromWorldMatrix(ULandscapeInfo* LandscapeInfo)
 	{
 		FMatrix Result = FScaleMatrix(FVector(1.0f, 1.0f, LANDSCAPE_INV_ZSCALE) / (LandscapeInfo->DrawScale));
-		Result *= FTranslationMatrix(FVector(0, 0, 32768.0f));
+		Result *= FTranslationMatrix(FVector(0, 0, LandscapeDataAccess::MidValue));
 		return Result;
 	}
 };
diff --git a/Engine/Source/Editor/LandscapeEditor/Private/LandscapeEdModeXYOffset.cpp b/Engine/Source/Editor/LandscapeEditor/Private/LandscapeEdModeXYOffset.cpp
--- a/Engine/Source/Editor/LandscapeEditor/Private/LandscapeEdModeXYOffset.cpp
+++ b/Engine/Source/Editor/LandscapeEditor/Private/LandscapeEdModeXYOffset.cpp
@@ -16,7 +16,7 @@ namespace
 {
 	FORCEINLINE FVector4 GetWorldPos(const FMatrix& LocalToWorld, FVector2D LocalXY, uint16 Height, FVector2D XYOffset)
 	{
-		return LocalToWorld.TransformPosition(FVector(LocalXY.X + XYOffset.X, LocalXY.Y + XYOffset.Y, ((float)Height - 32768.0f) * LANDSCAPE_ZSCALE));
+		return LocalToWorld.TransformPosition(FVector(LocalXY.X + XYOffset.X, LocalXY.Y + XYOffset.Y, LandscapeDataAccess::GetLocalHeight(Height)));
 	}
 
 	FORCEINLINE FVector4 GetWorldPos(const FMatrix& LocalToWorld, FVector2D LocalXY, FVector XYOffsetVector)
diff --git a/Engine/Source/Programs/UnrealLightmass/Private/Lighting/Landscape.cpp b/Engine/Source/Programs/UnrealLightmass/Private/Lighting/Landscape.cpp
--- a/Engine/Source/Programs/UnrealLightmass/Private/Lighting/Landscape.cpp
+++ b/Engine/Source/Programs/UnrealLightmass/Private/Lighting/Landscape.cpp
@@ -57,11 +57,8 @@ namespace Lightmass
 		const int32 LocalX = X - ExpandQuadsX;
 		const int32 LocalY = Y - ExpandQuadsY;
 
-		const FColor* Data = GetHeightData( X, Y );
-
-		OutVertex.WorldTangentZ.X = 2.0f / 255.f * (float)Data->B - 1.0f;
-		OutVertex.WorldTangentZ.Y = 2.0f / 255.f * (float)Data->A - 1.0f;
-		OutVertex.WorldTangentZ.Z = FMath::Sqrt(FMath::Max(1.0f - (FMath::Square(OutVertex.WorldTangentZ.X) + FMath::Square(OutVertex.WorldTangentZ.Y)), 0.f));
+		const FColor* Data = GetHeightData(X, Y);
+		OutVertex.WorldTangentZ = LandscapeDataAccess::UnpackNormal(*Data);
 		OutVertex.WorldTangentX = FVector4f(OutVertex.WorldTangentZ.Z, 0.0f, -OutVertex.WorldTangentZ.X);
 		OutVertex.WorldTangentY = OutVertex.WorldTangentZ ^ OutVertex.WorldTangentX;
 
@@ -71,8 +68,8 @@ namespace Lightmass
 		OutVertex.WorldTangentY = LtWNoScale.TransformVector(OutVertex.WorldTangentY);
 		OutVertex.WorldTangentZ = LtWNoScale.TransformVector(OutVertex.WorldTangentZ);
 
-		const uint16 Height = (Data->R << 8) + Data->G;
-		OutVertex.WorldPosition = LocalToWorld.TransformPosition( FVector4f( LocalX, LocalY, ((float)Height - 32768.f) * LANDSCAPE_ZSCALE ) );
+		const float Height = LandscapeDataAccess::UnpackHeight(*Data);
+		OutVertex.WorldPosition = LocalToWorld.TransformPosition(FVector4f(LocalX, LocalY, Height));
 		//UE_LOG(LogLightmass, Log, TEXT("%d, %d, %d, %d, %d, %d, X:%f, Y:%f, Z:%f "), SectionBaseX + LocalX - ExpandQuadsX, SectionBaseY + LocalY - ExpandQuadsY, ClampedLocalX, ClampedLocalY, SectionBaseX, SectionBaseY, WorldPos.X, WorldPos.Y, WorldPos.Z);
 
 		const int32 LightmapUVIndex = 1;
diff --git a/Engine/Source/Runtime/Landscape/Private/LandscapeCollision.cpp b/Engine/Source/Runtime/Landscape/Private/LandscapeCollision.cpp
--- a/Engine/Source/Runtime/Landscape/Private/LandscapeCollision.cpp
+++ b/Engine/Source/Runtime/Landscape/Private/LandscapeCollision.cpp
@@ -1152,7 +1152,7 @@ bool ULandscapeMeshCollisionComponent::CookCollisionData(const FName& Format, bo
 	{
 		int32 X = i % CollisionSizeVerts;
 		int32 Y = i / CollisionSizeVerts;
-		Vertices[i].Set(X + ((float)XYOffsets[i * 2] - 32768.f) * LANDSCAPE_XYOFFSET_SCALE, Y + ((float)XYOffsets[i * 2 + 1] - 32768.f) * LANDSCAPE_XYOFFSET_SCALE, ((float)Heights[i] - 32768.f) * LANDSCAPE_ZSCALE);
+		Vertices[i].Set(X + ((float)XYOffsets[i * 2] - 32768.f) * LANDSCAPE_XYOFFSET_SCALE, Y + ((float)XYOffsets[i * 2 + 1] - 32768.f) * LANDSCAPE_XYOFFSET_SCALE, LandscapeDataAccess::GetLocalHeight(Heights[i]));
 	}
 
 	const int32 NumTris = FMath::Square(CollisionSizeQuads) * 2;
diff --git a/Engine/Source/Runtime/Landscape/Private/LandscapeDataAccess.cpp b/Engine/Source/Runtime/Landscape/Private/LandscapeDataAccess.cpp
--- a/Engine/Source/Runtime/Landscape/Private/LandscapeDataAccess.cpp
+++ b/Engine/Source/Runtime/Landscape/Private/LandscapeDataAccess.cpp
@@ -154,9 +154,7 @@ LANDSCAPE_API FVector FLandscapeComponentDataInterface::GetWorldVertex(int32 Loc
 LANDSCAPE_API void FLandscapeComponentDataInterface::GetWorldTangentVectors(int32 LocalX, int32 LocalY, FVector& WorldTangentX, FVector& WorldTangentY, FVector& WorldTangentZ) const
 {
 	FColor* Data = GetHeightData(LocalX, LocalY);
-	WorldTangentZ.X = 2.f * (float)Data->B / 255.f - 1.f;
-	WorldTangentZ.Y = 2.f * (float)Data->A / 255.f - 1.f;
-	WorldTangentZ.Z = FMath::Sqrt(1.f - (FMath::Square(WorldTangentZ.X) + FMath::Square(WorldTangentZ.Y)));
+	WorldTangentZ = LandscapeDataAccess::UnpackNormal(*Data);
 	WorldTangentX = FVector(-WorldTangentZ.Z, 0.f, WorldTangentZ.X);
 	WorldTangentY = FVector(0.f, WorldTangentZ.Z, -WorldTangentZ.Y);
 
@@ -168,19 +166,16 @@ LANDSCAPE_API void FLandscapeComponentDataInterface::GetWorldTangentVectors(int3
 LANDSCAPE_API void FLandscapeComponentDataInterface::GetWorldPositionTangents(int32 LocalX, int32 LocalY, FVector& WorldPos, FVector& WorldTangentX, FVector& WorldTangentY, FVector& WorldTangentZ) const
 {
 	FColor* Data = GetHeightData(LocalX, LocalY);
-
-	WorldTangentZ.X = 2.f * (float)Data->B / 255.f - 1.f;
-	WorldTangentZ.Y = 2.f * (float)Data->A / 255.f - 1.f;
-	WorldTangentZ.Z = FMath::Sqrt(1.f - (FMath::Square(WorldTangentZ.X) + FMath::Square(WorldTangentZ.Y)));
+	WorldTangentZ = LandscapeDataAccess::UnpackNormal(*Data);
 	WorldTangentX = FVector(WorldTangentZ.Z, 0.f, -WorldTangentZ.X);
 	WorldTangentY = WorldTangentZ ^ WorldTangentX;
 
-	uint16 Height = (Data->R << 8) + Data->G;
+	float Height = LandscapeDataAccess::UnpackHeight(*Data);
 
 	const float ScaleFactor = (float)Component->ComponentSizeQuads / (float)(ComponentSizeVerts - 1);
 	float XOffset, YOffset;
 	GetXYOffset(LocalX, LocalY, XOffset, YOffset);
-	WorldPos = Component->GetComponentTransform().TransformPosition(FVector(LocalX * ScaleFactor + XOffset, LocalY * ScaleFactor + YOffset, LandscapeDataAccess::GetLocalHeight(Height)));
+	WorldPos = Component->GetComponentTransform().TransformPosition(FVector(LocalX * ScaleFactor + XOffset, LocalY * ScaleFactor + YOffset, Height));
 	WorldTangentX = Component->GetComponentTransform().TransformVectorNoScale(WorldTangentX);
 	WorldTangentY = Component->GetComponentTransform().TransformVectorNoScale(WorldTangentY);
 	WorldTangentZ = Component->GetComponentTransform().TransformVectorNoScale(WorldTangentZ);
diff --git a/Engine/Source/Runtime/Landscape/Private/LandscapeEdit.cpp b/Engine/Source/Runtime/Landscape/Private/LandscapeEdit.cpp
--- a/Engine/Source/Runtime/Landscape/Private/LandscapeEdit.cpp
+++ b/Engine/Source/Runtime/Landscape/Private/LandscapeEdit.cpp
@@ -2852,10 +2852,10 @@ LANDSCAPE_API void ALandscapeProxy::Import(const FGuid& InGuid, int32 InMinX, in
 	{
 		for (int32 QuadX = 0; QuadX < NumPatchesX; QuadX++)
 		{
-			const FVector Vert00 = FVector(0.0f, 0.0f, ((float)HEIGHTDATA(QuadX + 0, QuadY + 0) - 32768.0f)*LANDSCAPE_ZSCALE) * DrawScale3D;
-			const FVector Vert01 = FVector(0.0f, 1.0f, ((float)HEIGHTDATA(QuadX + 0, QuadY + 1) - 32768.0f)*LANDSCAPE_ZSCALE) * DrawScale3D;
-			const FVector Vert10 = FVector(1.0f, 0.0f, ((float)HEIGHTDATA(QuadX + 1, QuadY + 0) - 32768.0f)*LANDSCAPE_ZSCALE) * DrawScale3D;
-			const FVector Vert11 = FVector(1.0f, 1.0f, ((float)HEIGHTDATA(QuadX + 1, QuadY + 1) - 32768.0f)*LANDSCAPE_ZSCALE) * DrawScale3D;
+			const FVector Vert00 = FVector(0.0f, 0.0f, LandscapeDataAccess::GetLocalHeight(HEIGHTDATA(QuadX + 0, QuadY + 0))) * DrawScale3D;
+			const FVector Vert01 = FVector(0.0f, 1.0f, LandscapeDataAccess::GetLocalHeight(HEIGHTDATA(QuadX + 0, QuadY + 1))) * DrawScale3D;
+			const FVector Vert10 = FVector(1.0f, 0.0f, LandscapeDataAccess::GetLocalHeight(HEIGHTDATA(QuadX + 1, QuadY + 0))) * DrawScale3D;
+			const FVector Vert11 = FVector(1.0f, 1.0f, LandscapeDataAccess::GetLocalHeight(HEIGHTDATA(QuadX + 1, QuadY + 1))) * DrawScale3D;
 
 			const FVector FaceNormal1 = ((Vert00 - Vert10) ^ (Vert10 - Vert11)).GetSafeNormal();
 			const FVector FaceNormal2 = ((Vert11 - Vert01) ^ (Vert01 - Vert00)).GetSafeNormal();
diff --git a/Engine/Source/Runtime/Landscape/Private/LandscapeEditInterface.cpp b/Engine/Source/Runtime/Landscape/Private/LandscapeEditInterface.cpp
--- a/Engine/Source/Runtime/Landscape/Private/LandscapeEditInterface.cpp
+++ b/Engine/Source/Runtime/Landscape/Private/LandscapeEditInterface.cpp
@@ -294,10 +294,10 @@ void FLandscapeEditDataInterface::SetHeightData(int32 X1, int32 Y1, int32 X2, in
 		{
 			for (int32 X = 0; X < NumVertsX - 1; X++)
 			{
-				FVector Vert00 = FVector(XYOffsets[(X + 0) + NumVertsX*(Y + 0)].X, XYOffsets[(X + 0) + NumVertsX*(Y + 0)].Y, ((float)InData[(X + 0) + InStride*(Y + 0)] - 32768.0f) * LANDSCAPE_ZSCALE) * DrawScale;
-				FVector Vert01 = FVector(XYOffsets[(X + 0) + NumVertsX*(Y + 0)].X, XYOffsets[(X + 0) + NumVertsX*(Y + 0)].Y + 1.0f, ((float)InData[(X + 0) + InStride*(Y + 1)] - 32768.0f) * LANDSCAPE_ZSCALE) * DrawScale;
-				FVector Vert10 = FVector(XYOffsets[(X + 0) + NumVertsX*(Y + 0)].X + 1.0f, XYOffsets[(X + 0) + NumVertsX*(Y + 0)].Y, ((float)InData[(X + 1) + InStride*(Y + 0)] - 32768.0f) * LANDSCAPE_ZSCALE) * DrawScale;
-				FVector Vert11 = FVector(XYOffsets[(X + 0) + NumVertsX*(Y + 0)].X + 1.0f, XYOffsets[(X + 0) + NumVertsX*(Y + 0)].Y + 1.0f, ((float)InData[(X + 1) + InStride*(Y + 1)] - 32768.0f) * LANDSCAPE_ZSCALE) * DrawScale;
+				FVector Vert00 = FVector(XYOffsets[(X + 0) + NumVertsX*(Y + 0)].X + 0.0f, XYOffsets[(X + 0) + NumVertsX*(Y + 0)].Y + 0.0f, LandscapeDataAccess::GetLocalHeight(InData[(X + 0) + InStride*(Y + 0)])) * DrawScale;
+				FVector Vert01 = FVector(XYOffsets[(X + 0) + NumVertsX*(Y + 0)].X + 0.0f, XYOffsets[(X + 0) + NumVertsX*(Y + 0)].Y + 1.0f, LandscapeDataAccess::GetLocalHeight(InData[(X + 0) + InStride*(Y + 1)])) * DrawScale;
+				FVector Vert10 = FVector(XYOffsets[(X + 0) + NumVertsX*(Y + 0)].X + 1.0f, XYOffsets[(X + 0) + NumVertsX*(Y + 0)].Y + 0.0f, LandscapeDataAccess::GetLocalHeight(InData[(X + 1) + InStride*(Y + 0)])) * DrawScale;
+				FVector Vert11 = FVector(XYOffsets[(X + 0) + NumVertsX*(Y + 0)].X + 1.0f, XYOffsets[(X + 0) + NumVertsX*(Y + 0)].Y + 1.0f, LandscapeDataAccess::GetLocalHeight(InData[(X + 1) + InStride*(Y + 1)])) * DrawScale;
 
 				FVector FaceNormal1 = ((Vert00 - Vert10) ^ (Vert10 - Vert11)).GetSafeNormal();
 				FVector FaceNormal2 = ((Vert11 - Vert01) ^ (Vert01 - Vert00)).GetSafeNormal();
@@ -560,10 +560,10 @@ void FLandscapeEditDataInterface::RecalculateNormals()
 		{
 			for( int32 X=0;X<Stride-1;X++ )
 			{
-				FVector Vert00 = FVector(XYOffsets[(X+0) + Stride*(Y+0)].X,			XYOffsets[(X+0) + Stride*(Y+0)].Y,		((float)HeightData[(X+0) + Stride*(Y+0)] - 32768.0f)*LANDSCAPE_ZSCALE) * DrawScale;
-				FVector Vert01 = FVector(XYOffsets[(X+0) + Stride*(Y+0)].X,			XYOffsets[(X+0) + Stride*(Y+0)].Y+1.0f, ((float)HeightData[(X+0) + Stride*(Y+1)] - 32768.0f)*LANDSCAPE_ZSCALE) * DrawScale;
-				FVector Vert10 = FVector(XYOffsets[(X+0) + Stride*(Y+0)].X+1.0f,	XYOffsets[(X+0) + Stride*(Y+0)].Y,		((float)HeightData[(X+1) + Stride*(Y+0)] - 32768.0f)*LANDSCAPE_ZSCALE) * DrawScale;
-				FVector Vert11 = FVector(XYOffsets[(X+0) + Stride*(Y+0)].X+1.0f,	XYOffsets[(X+0) + Stride*(Y+0)].Y+1.0f,	((float)HeightData[(X+1) + Stride*(Y+1)] - 32768.0f)*LANDSCAPE_ZSCALE) * DrawScale;
+				FVector Vert00 = FVector(XYOffsets[(X + 0) + Stride * (Y + 0)].X + 0.0f, XYOffsets[(X + 0) + Stride * (Y + 0)].Y + 0.0f, LandscapeDataAccess::GetLocalHeight(HeightData[(X + 0) + Stride * (Y + 0)])) * DrawScale;
+				FVector Vert01 = FVector(XYOffsets[(X + 0) + Stride * (Y + 0)].X + 0.0f, XYOffsets[(X + 0) + Stride * (Y + 0)].Y + 1.0f, LandscapeDataAccess::GetLocalHeight(HeightData[(X + 0) + Stride * (Y + 1)])) * DrawScale;
+				FVector Vert10 = FVector(XYOffsets[(X + 0) + Stride * (Y + 0)].X + 1.0f, XYOffsets[(X + 0) + Stride * (Y + 0)].Y + 0.0f, LandscapeDataAccess::GetLocalHeight(HeightData[(X + 1) + Stride * (Y + 0)])) * DrawScale;
+				FVector Vert11 = FVector(XYOffsets[(X + 0) + Stride * (Y + 0)].X + 1.0f, XYOffsets[(X + 0) + Stride * (Y + 0)].Y + 1.0f, LandscapeDataAccess::GetLocalHeight(HeightData[(X + 1) + Stride * (Y + 1)])) * DrawScale;
 
 				FVector FaceNormal1 = ((Vert00-Vert10) ^ (Vert10-Vert11)).GetSafeNormal();
 				FVector FaceNormal2 = ((Vert11-Vert01) ^ (Vert01-Vert00)).GetSafeNormal(); 
diff --git a/Engine/Source/Runtime/Landscape/Private/LandscapeEditLayers.cpp b/Engine/Source/Runtime/Landscape/Private/LandscapeEditLayers.cpp
--- a/Engine/Source/Runtime/Landscape/Private/LandscapeEditLayers.cpp
+++ b/Engine/Source/Runtime/Landscape/Private/LandscapeEditLayers.cpp
@@ -3271,7 +3271,7 @@ void ALandscape::PrintLayersDebugHeightData(const FString& InContext, const TArr
 
 			if (DisplayHeightAsDelta)
 			{
-				HeightDelta = HeightDelta >= 32768 ? HeightDelta - 32768 : HeightDelta;
+				HeightDelta = HeightDelta >= LandscapeDataAccess::MidValue ? HeightDelta - LandscapeDataAccess::MidValue : HeightDelta;
 			}
 
 			if (X > 0 && MipSize > 0 && X % MipSize == 0)
diff --git a/Engine/Source/Runtime/Landscape/Private/LandscapeGizmoActor.cpp b/Engine/Source/Runtime/Landscape/Private/LandscapeGizmoActor.cpp
--- a/Engine/Source/Runtime/Landscape/Private/LandscapeGizmoActor.cpp
+++ b/Engine/Source/Runtime/Landscape/Private/LandscapeGizmoActor.cpp
@@ -985,7 +985,7 @@ LANDSCAPE_API void ALandscapeGizmoActiveActor::Import( int32 VertsX, int32 Verts
 		{
 			FGizmoSelectData Data;
 			Data.Ratio = 1.f;
-			Data.HeightData = (float)HeightData[X + Y*VertsX] / 65535.f; //GetNormalizedHeight(HeightData[X + Y*VertsX]);
+			Data.HeightData = (float)HeightData[X + Y*VertsX] / LandscapeDataAccess::MaxValue; //GetNormalizedHeight(HeightData[X + Y*VertsX]);
 			for (int32 i = 0; i < ImportLayerInfos.Num(); ++i)
 			{
 				Data.WeightDataMap.Add( ImportLayerInfos[i], LayerDataPointers[i][X + Y*VertsX] );
@@ -1061,7 +1061,7 @@ void ALandscapeGizmoActiveActor::Export(int32 Index, TArray<FString>& Filenames)
 						int32 Idx = (X-MinX) + Y *(1+MaxX-MinX);
 						if (!bExportOneTarget || Index == -1)
 						{
-							pHeightData[Idx] = FMath::Clamp<uint16>(Data->HeightData * 65535.f, 0, 65535);
+							pHeightData[Idx] = FMath::Clamp<uint16>(Data->HeightData * LandscapeDataAccess::MaxValue, 0, LandscapeDataAccess::MaxValue);
 						}
 
 						for( int32 i=1;i<Filenames.Num();i++ )
diff --git a/Engine/Source/Runtime/Landscape/Private/LandscapeLight.cpp b/Engine/Source/Runtime/Landscape/Private/LandscapeLight.cpp
--- a/Engine/Source/Runtime/Landscape/Private/LandscapeLight.cpp
+++ b/Engine/Source/Runtime/Landscape/Private/LandscapeLight.cpp
@@ -600,10 +600,7 @@ void FLandscapeStaticLightingMesh::GetStaticLightingVertex(int32 VertexIndex, FS
 	const int32 LocalY = Y - ExpandQuadsY;
 
 	const FColor* Data = &HeightData[X + Y * NumVertices];
-
-	OutVertex.WorldTangentZ.X = 2.0f / 255.f * (float)Data->B - 1.0f;
-	OutVertex.WorldTangentZ.Y = 2.0f / 255.f * (float)Data->A - 1.0f;
-	OutVertex.WorldTangentZ.Z = FMath::Sqrt(1.0f - (FMath::Square(OutVertex.WorldTangentZ.X) + FMath::Square(OutVertex.WorldTangentZ.Y)));
+	OutVertex.WorldTangentZ = LandscapeDataAccess::UnpackNormal(*Data);
 	OutVertex.WorldTangentX = FVector4(OutVertex.WorldTangentZ.Z, 0.0f, -OutVertex.WorldTangentZ.X);
 	OutVertex.WorldTangentY = OutVertex.WorldTangentZ ^ OutVertex.WorldTangentX;
 
@@ -612,8 +609,8 @@ void FLandscapeStaticLightingMesh::GetStaticLightingVertex(int32 VertexIndex, FS
 	OutVertex.WorldTangentY = LocalToWorld.TransformVectorNoScale(OutVertex.WorldTangentY);
 	OutVertex.WorldTangentZ = LocalToWorld.TransformVectorNoScale(OutVertex.WorldTangentZ);
 
-	const uint16 Height = (Data->R << 8) + Data->G;
-	OutVertex.WorldPosition = LocalToWorld.TransformPosition(FVector(LocalX, LocalY, LandscapeDataAccess::GetLocalHeight(Height)));
+	const float Height = LandscapeDataAccess::UnpackHeight(*Data);
+	OutVertex.WorldPosition = LocalToWorld.TransformPosition(FVector(LocalX, LocalY, Height));
 
 	OutVertex.TextureCoordinates[0] = FVector2D((float)X / NumVertices, (float)Y / NumVertices); 
 	OutVertex.TextureCoordinates[LANDSCAPE_LIGHTMAP_UV_INDEX].X = X * UVFactor;
diff --git a/Engine/Source/Runtime/Landscape/Private/LandscapeSplineRaster.cpp b/Engine/Source/Runtime/Landscape/Private/LandscapeSplineRaster.cpp
--- a/Engine/Source/Runtime/Landscape/Private/LandscapeSplineRaster.cpp
+++ b/Engine/Source/Runtime/Landscape/Private/LandscapeSplineRaster.cpp
@@ -87,7 +87,7 @@ protected:
 
 			uint16& DestAlphaValue = (*HeightAlphaBlendData)[DataIndex];
 			float InterpAlphaValue = DestAlphaValue * (1.f - Alpha);
-			DestAlphaValue = (uint16)FMath::Clamp<float>(InterpAlphaValue, 0.f, 65535.f);
+			DestAlphaValue = (uint16)FMath::Clamp<float>(InterpAlphaValue, 0.f, static_cast<float>(LandscapeDataAccess::MaxValue));
 
 			if (HeightFlagsData)
 			{
diff --git a/Engine/Source/Runtime/Landscape/Public/LandscapeDataAccess.h b/Engine/Source/Runtime/Landscape/Public/LandscapeDataAccess.h
--- a/Engine/Source/Runtime/Landscape/Public/LandscapeDataAccess.h
+++ b/Engine/Source/Runtime/Landscape/Public/LandscapeDataAccess.h
@@ -28,7 +28,7 @@ namespace LandscapeDataAccess
 	// 2nd significant bit - Triangle flip, not implemented yet
 	FORCEINLINE float GetLocalHeight(uint16 Height)
 	{
-		return ((float)Height - MidValue) * LANDSCAPE_ZSCALE;
+		return (static_cast<float>(Height) - MidValue) * LANDSCAPE_ZSCALE;
 	}
 
 	FORCEINLINE uint16 GetTexHeight(float Height)
@@ -43,6 +43,21 @@ namespace LandscapeDataAccess
 		Color.G = Height & 255;
 		return MoveTemp(Color);
 	}
+
+	FORCEINLINE float UnpackHeight(const FColor& InHeightmapSample)
+	{
+		uint16 Height = (InHeightmapSample.R << 8) + InHeightmapSample.G;
+		return GetLocalHeight(Height);
+	}
+
+	FORCEINLINE FVector UnpackNormal(const FColor& InHeightmapSample)
+	{
+		FVector Normal;
+		Normal.X = 2.f * static_cast<float>(InHeightmapSample.B) / 255.f - 1.f;
+		Normal.Y = 2.f * static_cast<float>(InHeightmapSample.A) / 255.f - 1.f;
+		Normal.Z = FMath::Sqrt(FMath::Max(1.0f - (FMath::Square(Normal.X) + FMath::Square(Normal.Y)), 0.0f));
+		return Normal;
+	}
 };
 
 #if WITH_EDITOR
@@ -283,9 +298,7 @@ struct FLandscapeComponentDataInterface
 		// Note: these are still pre-scaled, just not rotated
 
 		FColor* Data = GetHeightData( LocalX, LocalY );
-		LocalTangentZ.X = 2.f * (float)Data->B / 255.f - 1.f;
-		LocalTangentZ.Y = 2.f * (float)Data->A / 255.f - 1.f;
-		LocalTangentZ.Z = FMath::Sqrt(1.f - (FMath::Square(LocalTangentZ.X)+FMath::Square(LocalTangentZ.Y)));
+		LocalTangentZ = LandscapeDataAccess::UnpackNormal(*Data);
 		LocalTangentX = FVector(-LocalTangentZ.Z, 0.f, LocalTangentZ.X);
 		LocalTangentY = FVector(0.f, LocalTangentZ.Z, -LocalTangentZ.Y);
 	}
