Updating to use CanEverRender

#localization none

diff --git a/Engine/Source/Runtime/Landscape/Private/LandscapeEditLayers.cpp b/Engine/Source/Runtime/Landscape/Private/LandscapeEditLayers.cpp
--- a/Engine/Source/Runtime/Landscape/Private/LandscapeEditLayers.cpp
+++ b/Engine/Source/Runtime/Landscape/Private/LandscapeEditLayers.cpp
@@ -2614,7 +2614,7 @@ struct FLandscapeLayersCopyReadbackTextureParams
 void ExecuteCopyToReadbackTexture(TArray<FLandscapeLayersCopyReadbackTextureParams>& InParams)
 {
 	SCOPED_DRAW_EVENTF_GAMETHREAD(LandscapeLayers, TEXT("Copy to readback textures (%d copies)"), InParams.Num());
-	if (GUsingNullRHI)
+	if (!FApp::CanEverRender())
 	{
 		return;
 	}
@@ -2891,7 +2891,7 @@ void ALandscape::DrawHeightmapComponentsToRenderTarget(const FString& InDebugNam
 {
 	check(InHeightmapRTRead != nullptr);
 	check(InHeightmapRTWrite != nullptr);
-	if (GUsingNullRHI)
+	if (!FApp::CanEverRender())
 	{
 		return;
 	}
@@ -3457,7 +3457,7 @@ bool ALandscape::PrepareTextureResources(bool bInWaitForStreaming)
 	TRACE_CPUPROFILER_EVENT_SCOPE(LandscapeLayers_PrepareTextureResources);
 
 	ULandscapeInfo* Info = GetLandscapeInfo();
-	if (Info == nullptr || GUsingNullRHI)
+	if (Info == nullptr || !FApp::CanEverRender())
 	{
 		return false;
 	}
@@ -8022,7 +8022,7 @@ void ALandscape::UpdateLayersContent(bool bInWaitForStreaming, bool bInSkipMonit
 	bool bResourcesReady = PrepareTextureResources(bInWaitForStreaming);
 
 	ULandscapeInfo* LandscapeInfo = GetLandscapeInfo();
-	if ((LandscapeInfo == nullptr) || !CanHaveLayersContent() || !LandscapeInfo->AreAllComponentsRegistered() || GUsingNullRHI)
+	if ((LandscapeInfo == nullptr) || !CanHaveLayersContent() || !LandscapeInfo->AreAllComponentsRegistered() || !FApp::CanEverRender())
 	{
 		return;
 	}
@@ -8741,7 +8741,7 @@ void ALandscape::FinishDestroy()
 
 bool ALandscape::IsUpToDate() const
 {
-	if (GUsingNullRHI)
+	if (!FApp::CanEverRender())
 	{
 		return true;
 	}
