Fixed crash when starting PIE of default template open world map : Skip FixupLandscapeGuidsIfInstanced on temporary dirty actors (only used for PIE duplicate and outered to an actorContainer that is not outered to the world).

diff --git a/Engine/Source/Runtime/Landscape/Private/Landscape.cpp b/Engine/Source/Runtime/Landscape/Private/Landscape.cpp
--- a/Engine/Source/Runtime/Landscape/Private/Landscape.cpp
+++ b/Engine/Source/Runtime/Landscape/Private/Landscape.cpp
@@ -3385,35 +3385,37 @@ void FixupLandscapeGuidsIfInstanced(ALandscapeProxy* LandscapeProxy)
 	// if the outer world is instanced, we need to change our landscape guid (in a deterministic way)
 	// this avoids guid collisions when you instance a world (and it's landscapes) multiple times,
 	// while maintaining the same GUID between landscape proxy objects within an instance
-	UWorld* OuterWorld = LandscapeProxy->GetTypedOuter<UWorld>();
-	UPackage* OuterPackage = OuterWorld->GetPackage();
-	if (OuterPackage->GetFName() != OuterPackage->GetLoadedPath().GetPackageFName())	// OuterWorld->IsInstanced()) is broken when PIE uses a memory built package
+	if (UWorld* OuterWorld = LandscapeProxy->GetTypedOuter<UWorld>())
 	{
-		FArchiveMD5 Ar;
-		FGuid OldLandscapeGuid = LandscapeProxy->GetLandscapeGuid();
-		Ar << OldLandscapeGuid;
+		UPackage* OuterPackage = OuterWorld->GetPackage();
+		if (OuterPackage->GetFName() != OuterPackage->GetLoadedPath().GetPackageFName())	// OuterWorld->IsInstanced()) is broken when PIE uses a memory built package
+		{
+			FArchiveMD5 Ar;
+			FGuid OldLandscapeGuid = LandscapeProxy->GetLandscapeGuid();
+			Ar << OldLandscapeGuid;
 
 #if WITH_EDITOR
-		// to work around PIE issues, we use the world partition to find the package
-		UWorldPartition* WorldPartition = FWorldPartitionHelpers::GetWorldPartition(LandscapeProxy);
-		if (WorldPartition)
-		{
-			OuterPackage = WorldPartition->GetPackage();
-			FName PackageName = OuterPackage->GetFName();
-			Ar << PackageName;
-		}
-		else
+			// to work around PIE issues, we use the world partition to find the package
+			UWorldPartition* WorldPartition = FWorldPartitionHelpers::GetWorldPartition(LandscapeProxy);
+			if (WorldPartition)
+			{
+				OuterPackage = WorldPartition->GetPackage();
+				FName PackageName = OuterPackage->GetFName();
+				Ar << PackageName;
+			}
+			else
 #endif // WITH_EDITOR
-		{
-			FObjectKey ContainerKey(OuterPackage);
-			Ar << ContainerKey;
-		}
+			{
+				FObjectKey ContainerKey(OuterPackage);
+				Ar << ContainerKey;
+			}
 
-		FMD5Hash MD5Hash;
-		Ar.GetHash(MD5Hash);
+			FMD5Hash MD5Hash;
+			Ar.GetHash(MD5Hash);
 
-		FGuid NewLandscapeGuid = MD5HashToGuid(MD5Hash);
-		LandscapeProxy->SetLandscapeGuid(NewLandscapeGuid);
+			FGuid NewLandscapeGuid = MD5HashToGuid(MD5Hash);
+			LandscapeProxy->SetLandscapeGuid(NewLandscapeGuid);
+		}
 	}
 }
 
