Expose rendered physical materials to runtime.

diff --git a/Engine/Source/Runtime/Landscape/Classes/LandscapeHeightfieldCollisionComponent.h b/Engine/Source/Runtime/Landscape/Classes/LandscapeHeightfieldCollisionComponent.h
--- a/Engine/Source/Runtime/Landscape/Classes/LandscapeHeightfieldCollisionComponent.h
+++ b/Engine/Source/Runtime/Landscape/Classes/LandscapeHeightfieldCollisionComponent.h
@@ -250,12 +250,17 @@ public:
 	struct FWriteRuntimeDataParams
 	{
 		bool bUseDefaultMaterialOnly = false;
+		bool bProcessRenderIndices = true;
+		bool bProcessVisibilityLayer = true;
 		TArrayView<const uint16> Heights;
 		TArrayView<const uint16> SimpleHeights;
 		TArrayView<const uint8> DominantLayers;
 		TArrayView<const uint8> SimpleDominantLayers;
 		TArrayView<const uint8> RenderPhysicalMaterialIds;
 		TArrayView<const uint8> SimpleRenderPhysicalMaterialIds;
+		TArrayView<const TObjectPtr<UPhysicalMaterial>> PhysicalMaterialRenderObjects;
+		TArrayView<const TObjectPtr<ULandscapeLayerInfoObject>> ComponentLayerInfos;
+		int32 VisibilityLayerIndex = INDEX_NONE;
 	};
 
 	void GetCollisionSampleInfo(int32& OutCollisionSizeVerts, int32& OutSimpleCollisionSizeVerts, int32& OutNumSamples, int32& OutNumSimpleSamples) const;
diff --git a/Engine/Source/Runtime/Landscape/Private/LandscapeCollision.cpp b/Engine/Source/Runtime/Landscape/Private/LandscapeCollision.cpp
--- a/Engine/Source/Runtime/Landscape/Private/LandscapeCollision.cpp
+++ b/Engine/Source/Runtime/Landscape/Private/LandscapeCollision.cpp
@@ -856,7 +856,7 @@ bool ULandscapeHeightfieldCollisionComponent::WriteRuntimeData(const FWriteRunti
 	TArray<uint8> MaterialIndices;
 	MaterialIndices.Reserve(NumSamples + NumSimpleSamples);
 
-	auto ResolveMaterials = [&MaterialIndices, &bIsMirrored, &Params, &DefMaterial, &InOutMaterials, this](int32 InCollisionVertExtent, TArrayView<const uint8> InDominantLayers, TArrayView<const uint8> InRenderMaterialIds)
+	auto ResolveMaterials = [&MaterialIndices, &bIsMirrored, &Params, &DefMaterial, &InOutMaterials](int32 InCollisionVertExtent, TArrayView<const uint8> InDominantLayers, TArrayView<const uint8> InRenderMaterialIds)
 	{
 		for (int32 RowIndex = 0; RowIndex < InCollisionVertExtent; RowIndex++)
 		{
@@ -872,22 +872,20 @@ bool ULandscapeHeightfieldCollisionComponent::WriteRuntimeData(const FWriteRunti
 					if (!Params.bUseDefaultMaterialOnly)
 					{
 						uint8 DominantLayerIdx = InDominantLayers.IsEmpty() ? -1 : InDominantLayers[SrcSampleIndex];
-						ULandscapeLayerInfoObject* Layer = ComponentLayerInfos.IsValidIndex(DominantLayerIdx) ? ToRawPtr(ComponentLayerInfos[DominantLayerIdx]) : nullptr;
+						ULandscapeLayerInfoObject* Layer = Params.ComponentLayerInfos.IsValidIndex(DominantLayerIdx) ? ToRawPtr(Params.ComponentLayerInfos[DominantLayerIdx]) : nullptr;
 
-#if WITH_EDITOR
-						if (Layer == ALandscapeProxy::VisibilityLayer)
+						if (Params.bProcessVisibilityLayer && DominantLayerIdx == Params.VisibilityLayerIndex)
 						{
 							// If it's a hole, use the final index
 							MaterialIndex = TNumericLimits<uint8>::Max();
 						}
-						else if (!InRenderMaterialIds.IsEmpty())
+						else if (Params.bProcessRenderIndices && !InRenderMaterialIds.IsEmpty())
 						{
 							uint8 RenderIdx = InRenderMaterialIds[SrcSampleIndex];
-							UPhysicalMaterial* DominantMaterial = RenderIdx > 0 ? ToRawPtr(PhysicalMaterialRenderObjects[RenderIdx - 1]) : DefMaterial;
+							UPhysicalMaterial* DominantMaterial = RenderIdx > 0 ? ToRawPtr(Params.PhysicalMaterialRenderObjects[RenderIdx - 1]) : DefMaterial;
 							MaterialIndex = InOutMaterials.AddUnique(DominantMaterial);
 						}
 						else
-#endif
 						{
 							UPhysicalMaterial* DominantMaterial = Layer && Layer->PhysMaterial ? ToRawPtr(Layer->PhysMaterial) : DefMaterial;
 							MaterialIndex = InOutMaterials.AddUnique(DominantMaterial);
@@ -1025,12 +1023,17 @@ bool ULandscapeHeightfieldCollisionComponent::CookCollisionData(const FName& For
 
 	FWriteRuntimeDataParams WriteParams;
 	WriteParams.bUseDefaultMaterialOnly = bUseDefaultMaterialOnly;
+	WriteParams.bProcessRenderIndices = true;
+	WriteParams.bProcessVisibilityLayer = true;
 	WriteParams.Heights = MakeSafeArrayView(Heights, NumSamples);
 	WriteParams.SimpleHeights = MakeSafeArrayView(SimpleHeights, NumSimpleSamples);
 	WriteParams.DominantLayers = MakeSafeArrayView(DominantLayers, NumSamples);
 	WriteParams.SimpleDominantLayers = MakeSafeArrayView(SimpleDominantLayers, NumSimpleSamples);
 	WriteParams.RenderPhysicalMaterialIds = MakeSafeArrayView(RenderPhysicalMaterialIds, NumSamples);
 	WriteParams.SimpleRenderPhysicalMaterialIds = MakeSafeArrayView(SimpleRenderPhysicalMaterialIds, NumSimpleSamples);
+	WriteParams.PhysicalMaterialRenderObjects = MakeSafeArrayView(PhysicalMaterialRenderObjects.GetData(), PhysicalMaterialRenderObjects.Num());
+	WriteParams.ComponentLayerInfos = MakeSafeArrayView(ComponentLayerInfos.GetData(), ComponentLayerInfos.Num());
+	WriteParams.VisibilityLayerIndex = ComponentLayerInfos.IndexOfByKey(ALandscapeProxy::VisibilityLayer);
 
 	TArray<uint8> OutData;
 	bool Succeeded = WriteRuntimeData(WriteParams, OutData, InOutMaterials);
