Force all Landscape Grass to regenerate when bAllowMaskedMaterials is toggled in WorldSettings.

diff --git a/Engine/Source/Runtime/Engine/Private/WorldSettings.cpp b/Engine/Source/Runtime/Engine/Private/WorldSettings.cpp
--- a/Engine/Source/Runtime/Engine/Private/WorldSettings.cpp
+++ b/Engine/Source/Runtime/Engine/Private/WorldSettings.cpp
@@ -838,6 +838,9 @@ void AWorldSettings::InternalPostPropertyChanged(FName PropertyName)
 		{
 			// Need to recreate scene proxies when this flag changes.
 			FGlobalComponentRecreateRenderStateContext Context;
+
+			extern void ForceRegenerateAllGrass();
+			ForceRegenerateAllGrass();
 		}
 	}
 }
diff --git a/Engine/Source/Runtime/Landscape/Private/LandscapeGrass.cpp b/Engine/Source/Runtime/Landscape/Private/LandscapeGrass.cpp
--- a/Engine/Source/Runtime/Landscape/Private/LandscapeGrass.cpp
+++ b/Engine/Source/Runtime/Landscape/Private/LandscapeGrass.cpp
@@ -814,6 +814,17 @@ DECLARE_CYCLE_STAT(TEXT("Grass Update"), STAT_GrassUpdate, STATGROUP_Foliage);
 
 int32 ALandscapeProxy::GrassUpdateInterval = 1;
 
+void ForceRegenerateAllGrass()
+{
+	for (UWorld* World : TObjectRange<UWorld>(RF_ClassDefaultObject | RF_ArchetypeObject, true, EInternalObjectFlags::Garbage))
+	{
+		if (ULandscapeSubsystem* LandscapeSubsystem = World->GetSubsystem<ULandscapeSubsystem>())
+		{
+			LandscapeSubsystem->RegenerateGrass(/*bInFlushGrass = */true, /*bInForceSync = */true);
+		}
+	}
+}
+
 static void GrassCVarSinkFunction()
 {
 	static float CachedGrassDensityScale = 1.0f;
@@ -845,13 +856,7 @@ static void GrassCVarSinkFunction()
 		CachedDetailMode = DetailMode;
 		CachedNaniteEnabled = NaniteEnabled;
 
-		for (UWorld* World : TObjectRange<UWorld>(RF_ClassDefaultObject | RF_ArchetypeObject, true, EInternalObjectFlags::Garbage))
-		{
-			if (ULandscapeSubsystem* LandscapeSubsystem = World->GetSubsystem<ULandscapeSubsystem>())
-			{
-				LandscapeSubsystem->RegenerateGrass(/*bInFlushGrass = */true, /*bInForceSync = */true);
-			}
-		}
+		ForceRegenerateAllGrass();
 	}
 }
 
