[Backout] - CL24504022 (backed out due to cis errors)
Original CL Desc
-----------------------------------------------------------------
Force all Landscape Grass to regenerate when bAllowMaskedMaterials is toggled in WorldSettings.

diff --git a/Engine/Source/Runtime/Engine/Private/WorldSettings.cpp b/Engine/Source/Runtime/Engine/Private/WorldSettings.cpp
--- a/Engine/Source/Runtime/Engine/Private/WorldSettings.cpp
+++ b/Engine/Source/Runtime/Engine/Private/WorldSettings.cpp
@@ -838,9 +838,6 @@ void AWorldSettings::InternalPostPropertyChanged(FName PropertyName)
 		{
 			// Need to recreate scene proxies when this flag changes.
 			FGlobalComponentRecreateRenderStateContext Context;
-
-			extern void ForceRegenerateAllGrass();
-			ForceRegenerateAllGrass();
 		}
 	}
 }
diff --git a/Engine/Source/Runtime/Landscape/Private/LandscapeGrass.cpp b/Engine/Source/Runtime/Landscape/Private/LandscapeGrass.cpp
--- a/Engine/Source/Runtime/Landscape/Private/LandscapeGrass.cpp
+++ b/Engine/Source/Runtime/Landscape/Private/LandscapeGrass.cpp
@@ -814,17 +814,6 @@ DECLARE_CYCLE_STAT(TEXT("Grass Update"), STAT_GrassUpdate, STATGROUP_Foliage);
 
 int32 ALandscapeProxy::GrassUpdateInterval = 1;
 
-void ForceRegenerateAllGrass()
-{
-	for (UWorld* World : TObjectRange<UWorld>(RF_ClassDefaultObject | RF_ArchetypeObject, true, EInternalObjectFlags::Garbage))
-	{
-		if (ULandscapeSubsystem* LandscapeSubsystem = World->GetSubsystem<ULandscapeSubsystem>())
-		{
-			LandscapeSubsystem->RegenerateGrass(/*bInFlushGrass = */true, /*bInForceSync = */true);
-		}
-	}
-}
-
 static void GrassCVarSinkFunction()
 {
 	static float CachedGrassDensityScale = 1.0f;
@@ -856,7 +845,13 @@ static void GrassCVarSinkFunction()
 		CachedDetailMode = DetailMode;
 		CachedNaniteEnabled = NaniteEnabled;
 
-		ForceRegenerateAllGrass();
+		for (UWorld* World : TObjectRange<UWorld>(RF_ClassDefaultObject | RF_ArchetypeObject, true, EInternalObjectFlags::Garbage))
+		{
+			if (ULandscapeSubsystem* LandscapeSubsystem = World->GetSubsystem<ULandscapeSubsystem>())
+			{
+				LandscapeSubsystem->RegenerateGrass(/*bInFlushGrass = */true, /*bInForceSync = */true);
+			}
+		}
 	}
 }
 
