Respect Nanite material audit in ULandscapeProxy::UpdateRenderingMethod()
- Fallback to non-nanite landscape if material is not compatible with Nanite.
- Force masking to be allowed to ensure that all landscape components use the same rendering path (otherwise some components might using visibility masking while others don't).

diff --git a/Engine/Source/Runtime/Landscape/Private/Landscape.cpp b/Engine/Source/Runtime/Landscape/Private/Landscape.cpp
--- a/Engine/Source/Runtime/Landscape/Private/Landscape.cpp
+++ b/Engine/Source/Runtime/Landscape/Private/Landscape.cpp
@@ -72,6 +72,7 @@ Landscape.cpp: Terrain rendering
 #include "RenderCaptureInterface.h"
 #include "VisualLogger/VisualLogger.h"
 #include "LandscapeSettings.h"
+#include "NaniteSceneProxy.h"
 #include "Misc/ArchiveMD5.h"
 
 #if WITH_EDITOR
@@ -5543,6 +5544,16 @@ void ALandscapeProxy::UpdateRenderingMethod()
 #endif
 	}
 
+	if (bNaniteActive)
+	{
+		Nanite::FMaterialAudit NaniteMaterials;
+		Nanite::AuditMaterials(NaniteComponent, NaniteMaterials);
+
+		//const bool bIsMaskingAllowed = Nanite::IsMaskingAllowedForWorld(GetWorld()) || bForceNaniteForMasked;
+		const bool bIsMaskingAllowed = true;
+		bNaniteActive = NaniteMaterials.IsValid(bIsMaskingAllowed);
+	}
+
 	for (ULandscapeComponent* Component : LandscapeComponents)
 	{
 		if (Component)
