Fix Nanite Landscape not rendering when Nanite is disabled for masked materials in WorldSettings.
- There was a mismatch between the logic in ALandscapeProxy::UpdateRenderingMethod() and ShouldCreateNaniteProxy(...)
- Enable bForceNaniteForMasked on the ULandscapeNaniteComponent instead of forcing it in ALandscapeProxy::UpdateRenderingMethod

diff --git a/Engine/Source/Runtime/Landscape/Private/Landscape.cpp b/Engine/Source/Runtime/Landscape/Private/Landscape.cpp
--- a/Engine/Source/Runtime/Landscape/Private/Landscape.cpp
+++ b/Engine/Source/Runtime/Landscape/Private/Landscape.cpp
@@ -318,6 +318,7 @@ void ALandscapeProxy::UpdateNaniteRepresentation(const ITargetPlatform* InTarget
 			NaniteComponent->CanCharacterStepUpOn = ECanBeCharacterBase::ECB_No;
 			NaniteComponent->bSelectable = false;
 			NaniteComponent->DepthPriorityGroup = SDPG_World;
+			NaniteComponent->bForceNaniteForMasked = true;
 			NaniteComponent->RegisterComponent();
 			NaniteComponent->AttachToComponent(GetRootComponent(), FAttachmentTransformRules::KeepRelativeTransform);
 		}
@@ -5549,8 +5550,7 @@ void ALandscapeProxy::UpdateRenderingMethod()
 		Nanite::FMaterialAudit NaniteMaterials;
 		Nanite::AuditMaterials(NaniteComponent, NaniteMaterials);
 
-		//const bool bIsMaskingAllowed = Nanite::IsMaskingAllowedForWorld(GetWorld()) || bForceNaniteForMasked;
-		const bool bIsMaskingAllowed = true;
+		const bool bIsMaskingAllowed = Nanite::IsMaskingAllowedForWorld(GetWorld()) || NaniteComponent->bForceNaniteForMasked;
 		bNaniteActive = NaniteMaterials.IsValid(bIsMaskingAllowed);
 	}
 
