Add a world settings callback for when Nanite settings change and register for it to regenerate Landscape grass

diff --git a/Engine/Source/Runtime/Engine/Classes/GameFramework/WorldSettings.h b/Engine/Source/Runtime/Engine/Classes/GameFramework/WorldSettings.h
--- a/Engine/Source/Runtime/Engine/Classes/GameFramework/WorldSettings.h
+++ b/Engine/Source/Runtime/Engine/Classes/GameFramework/WorldSettings.h
@@ -1011,6 +1011,9 @@ public:
 	static FOnNumberOfBookmarksChanged OnNumberOfBoomarksChanged;
 #endif
 
+	DECLARE_MULTICAST_DELEGATE_OneParam(FOnNaniteSettingsChanged, AWorldSettings*);
+	static FOnNaniteSettingsChanged OnNaniteSettingsChanged;
+
 	const int32 GetMaxNumberOfBookmarks() const
 	{
 		return MaxNumberOfBookmarks;
diff --git a/Engine/Source/Runtime/Engine/Private/WorldSettings.cpp b/Engine/Source/Runtime/Engine/Private/WorldSettings.cpp
--- a/Engine/Source/Runtime/Engine/Private/WorldSettings.cpp
+++ b/Engine/Source/Runtime/Engine/Private/WorldSettings.cpp
@@ -51,6 +51,7 @@ ENGINE_API float GNewWorldToMetersScale = 0.0f;
 AWorldSettings::FOnBookmarkClassChanged AWorldSettings::OnBookmarkClassChanged;
 AWorldSettings::FOnNumberOfBookmarksChanged AWorldSettings::OnNumberOfBoomarksChanged;
 #endif
+AWorldSettings::FOnNaniteSettingsChanged AWorldSettings::OnNaniteSettingsChanged;
 
 AWorldSettings::AWorldSettings(const FObjectInitializer& ObjectInitializer)
 	: Super(ObjectInitializer.DoNotCreateDefaultSubobject(TEXT("Sprite")))
@@ -261,6 +262,8 @@ void AWorldSettings::OnRep_NaniteSettings()
 {
 	// Need to recreate scene proxies when Nanite settings changes.
 	FGlobalComponentRecreateRenderStateContext Context;
+
+	OnNaniteSettingsChanged.Broadcast(this);
 }
 
 void AWorldSettings::SetAllowMaskedMaterials(bool bState)
@@ -838,6 +841,8 @@ void AWorldSettings::InternalPostPropertyChanged(FName PropertyName)
 		{
 			// Need to recreate scene proxies when this flag changes.
 			FGlobalComponentRecreateRenderStateContext Context;
+		
+			OnNaniteSettingsChanged.Broadcast(this);
 		}
 	}
 }
diff --git a/Engine/Source/Runtime/Landscape/Private/LandscapeSubsystem.cpp b/Engine/Source/Runtime/Landscape/Private/LandscapeSubsystem.cpp
--- a/Engine/Source/Runtime/Landscape/Private/LandscapeSubsystem.cpp
+++ b/Engine/Source/Runtime/Landscape/Private/LandscapeSubsystem.cpp
@@ -77,6 +77,14 @@ void ULandscapeSubsystem::Initialize(FSubsystemCollectionBase& Collection)
 {
 	Super::Initialize(Collection);
 
+	if (UWorld* World = GetWorld())
+	{
+		if (AWorldSettings* WorldSettings = World->GetWorldSettings())
+		{
+			OnNaniteWorldSettingsChangedHandle = WorldSettings->OnNaniteSettingsChanged.AddUObject(this, &ULandscapeSubsystem::OnNaniteWorldSettingsChanged);
+		}
+	}
+
 #if WITH_EDITOR
 	GrassMapsBuilder = new FLandscapeGrassMapsBuilder(GetWorld());
 	GIBakedTextureBuilder = new FLandscapeGIBakedTextureBuilder(GetWorld());
@@ -91,6 +99,18 @@ void ULandscapeSubsystem::Initialize(FSubsystemCollectionBase& Collection)
 
 void ULandscapeSubsystem::Deinitialize()
 {
+	if (OnNaniteWorldSettingsChangedHandle.IsValid())
+	{
+		UWorld* World = GetWorld();
+		check(World != nullptr);
+
+		AWorldSettings* WorldSettings = World->GetWorldSettings();
+		check(WorldSettings != nullptr);
+
+		WorldSettings->OnNaniteSettingsChanged.Remove(OnNaniteWorldSettingsChangedHandle);
+		OnNaniteWorldSettingsChangedHandle.Reset();
+	}
+	
 #if WITH_EDITOR
 	delete GrassMapsBuilder;
 	delete GIBakedTextureBuilder;
diff --git a/Engine/Source/Runtime/Landscape/Public/LandscapeSubsystem.h b/Engine/Source/Runtime/Landscape/Public/LandscapeSubsystem.h
--- a/Engine/Source/Runtime/Landscape/Public/LandscapeSubsystem.h
+++ b/Engine/Source/Runtime/Landscape/Public/LandscapeSubsystem.h
@@ -8,6 +8,7 @@
 #include "LandscapeSubsystem.generated.h"
 
 class ALandscapeProxy;
+class AWorldSettings;
 class ULandscapeInfo;
 class FLandscapeNotificationManager;
 struct FDateTime;
@@ -100,8 +101,11 @@ private:
 	virtual void Deinitialize() override;
 	// End USubsystem
 
+	void OnNaniteWorldSettingsChanged(AWorldSettings* WorldSettings) { RegenerateGrass(true, true); }
+
 	bool bIsGrassCreationPrioritized;
 	TArray<TWeakObjectPtr<ALandscapeProxy>> Proxies;
+	FDelegateHandle OnNaniteWorldSettingsChangedHandle;
 
 #if WITH_EDITOR
 	class FLandscapeGrassMapsBuilder* GrassMapsBuilder;
