Fix a rare crash seen when cooking where the landscape system will attempt to lock a texture for reading while the texture has outstanding async texture compilation work, resulting in asserts about the texture source's lock state.

#rnx

- Changed FSourceDataMipNumber to hold a FTextureSource::FMipData> MipData rather than locking each mip. This way of accessing mip data from a texture is thread safe and will not clash with the async texture compilation tasks.
-- Note that it is a TOptional as FMipData has no default constructor nor copy constructor and has problems being in a container, this is because of its strict lifetime rules.
- Add a new accessor to FTextureSource::FMipData which returns the mip as a FSharedBuffer. Instead of making a new allocation and copying the mip data to it as the other accessor does, this one returns the FSharedBuffer as a view around its internal allocation.

diff --git a/Engine/Source/Runtime/Engine/Classes/Engine/Texture.h b/Engine/Source/Runtime/Engine/Classes/Engine/Texture.h
--- a/Engine/Source/Runtime/Engine/Classes/Engine/Texture.h
+++ b/Engine/Source/Runtime/Engine/Classes/Engine/Texture.h
@@ -457,6 +457,15 @@ struct FTextureSource
 		/** Get a copy of a given texture mip, to be stored in OutMipData */
 		bool GetMipData(TArray64<uint8>& OutMipData, int32 BlockIndex, int32 LayerIndex, int32 MipIndex) const;
 
+		/** 
+		 * Access the given texture mip in the form of a read only FSharedBuffer.
+		 * This method does not make a copy of the mip but instead returns a view.
+		 * The returned buffer will keep a ref count on the original full mip chain
+		 * allocation, so will remain valid even after the FMipData object has been 
+		 * destroyed.
+		 */
+		FSharedBuffer GetMipData(int32 BlockIndex, int32 LayerIndex, int32 MipIndex) const;
+
 	private:
 		// We only want to allow FTextureSource to create FMipData objects
 		friend struct FTextureSource;
diff --git a/Engine/Source/Runtime/Engine/Private/Texture.cpp b/Engine/Source/Runtime/Engine/Private/Texture.cpp
--- a/Engine/Source/Runtime/Engine/Private/Texture.cpp
+++ b/Engine/Source/Runtime/Engine/Private/Texture.cpp
@@ -3837,6 +3837,22 @@ bool FTextureSource::FMipData::GetMipData(TArray64<uint8>& OutMipData, int32 Blo
 	return false;
 }
 
+FSharedBuffer FTextureSource::FMipData::GetMipData(int32 BlockIndex, int32 LayerIndex, int32 MipIndex) const
+{
+	if (BlockIndex < TextureSource.GetNumBlocks() && LayerIndex < TextureSource.GetNumLayers() && MipIndex < TextureSource.GetNumMips() && !MipData.IsNull())
+	{
+		const int64 MipOffset = TextureSource.CalcMipOffset(BlockIndex, LayerIndex, MipIndex);
+		const int64 MipSize = TextureSource.CalcMipSize(BlockIndex, LayerIndex, MipIndex);
+
+		if ((int64)MipData.GetSize() >= MipOffset + MipSize)
+		{
+			return FSharedBuffer::MakeView((const uint8*)MipData.GetData() + MipOffset, MipSize, MipData);
+		}
+	}
+	
+	return FSharedBuffer();
+}
+
 #endif //WITH_EDITOR
 
 #if WITH_EDITOR
diff --git a/Engine/Source/Runtime/Landscape/Private/LandscapeUtils.cpp b/Engine/Source/Runtime/Landscape/Private/LandscapeUtils.cpp
--- a/Engine/Source/Runtime/Landscape/Private/LandscapeUtils.cpp
+++ b/Engine/Source/Runtime/Landscape/Private/LandscapeUtils.cpp
@@ -121,7 +121,7 @@ bool FBatchTextureCopy::AddWeightmapCopy(UTexture2D* InDestination, int8 InDesti
 
 struct FSourceDataMipNumber
 {
-	TArray<const uint8*> SourceDataPtr;
+	TOptional<FTextureSource::FMipData> MipData;
 	int32 MipNumber = 0;
 };
 
@@ -152,18 +152,10 @@ bool FBatchTextureCopy::ProcessTextureCopies()
 		DestinationData.MipNumber = CopyRequest.Key.Destination->Source.GetNumMips();
 	}
 
-	// Lock all sources mips
+	// Decompress (if needed) and get the source textures ready for access
 	for (TPair<UTexture2D*, FSourceDataMipNumber>& Source : Sources)
 	{
-		int32 MipNumber = Source.Value.MipNumber;
-		TArray<const uint8*>& SourceDataPtr = Source.Value.SourceDataPtr;
-		
-		SourceDataPtr.Reserve(MipNumber);
-
-		for (int32 MipLevel = 0; MipLevel < MipNumber; ++MipLevel)
-		{
-			SourceDataPtr.Add(Source.Key->Source.LockMipReadOnly(MipLevel));
-		}
+		Source.Value.MipData = Source.Key->Source.GetMipData(nullptr);
 	}
 
 	// Lock all destinations mips
@@ -194,7 +186,8 @@ bool FBatchTextureCopy::ProcessTextureCopies()
 			check(MipSize == (CopyRequest.Key.Destination->Source.GetSizeY() >> MipLevel));
 
 			int32 MipSizeSquare = FMath::Square(MipSize);
-			const uint8* SourceTextureData = SourceDataMipNumber->SourceDataPtr[MipLevel];
+			FSharedBuffer MipSrcData = SourceDataMipNumber->MipData->GetMipData(0, 0, MipLevel);
+			const uint8* SourceTextureData = static_cast<const uint8*>(MipSrcData.GetData());
 			uint8* DestTextureData = DestinationDataMipNumber->DestinationDataPtr[MipLevel];
 
 			check((SourceTextureData != nullptr) && (DestTextureData != nullptr));
@@ -219,16 +212,7 @@ bool FBatchTextureCopy::ProcessTextureCopies()
 		}
 	}
 
-	// Unlock all sources mips
-	for (TPair<UTexture2D*, FSourceDataMipNumber>& Source : Sources)
-	{
-		int32 MipNumber = Source.Value.MipNumber;
-		
-		for (int32 MipLevel = 0; MipLevel < MipNumber; ++MipLevel)
-		{
-			Source.Key->Source.UnlockMip(MipLevel);
-		}
-	}
+	// Note that source textures do not need unlocking, data will be released once the FMipData go out of scope
 	
 	// Unlock all destination mips
 	for (TPair<UTexture2D*, FDestinationDataMipNumber>& Destination : Destinations)
