Fix for MapCheck landscape errors in the Project.Maps.AllInPIE.Maps.Tools.Landscape.FallenSmall.FallenSmall tests.

diff --git a/Engine/Source/Runtime/Landscape/Private/Landscape.cpp b/Engine/Source/Runtime/Landscape/Private/Landscape.cpp
--- a/Engine/Source/Runtime/Landscape/Private/Landscape.cpp
+++ b/Engine/Source/Runtime/Landscape/Private/Landscape.cpp
@@ -4446,7 +4446,10 @@ void ULandscapeInfo::RegisterActor(ALandscapeProxy* Proxy, bool bMapCheck, bool
 			StreamingProxies.Insert(StreamingProxyPtr, InsertIndex);
 		}
 
-		StreamingProxy->SetLandscapeActor(LandscapeActor.Get());
+		if (LandscapeActor.IsValid())	// don't overwrite the proxy's landscape actor if we don't have one registered yet
+		{
+			StreamingProxy->SetLandscapeActor(LandscapeActor.Get());
+		}
 #if WITH_EDITOR
 		StreamingProxy->FixupSharedData(LandscapeActor.Get());
 #endif // WITH_EDITOR
