Fixes several issues with 24125650 -- level instancing support for landscape.

diff --git a/Engine/Source/Runtime/Landscape/Private/Landscape.cpp b/Engine/Source/Runtime/Landscape/Private/Landscape.cpp
--- a/Engine/Source/Runtime/Landscape/Private/Landscape.cpp
+++ b/Engine/Source/Runtime/Landscape/Private/Landscape.cpp
@@ -3381,36 +3381,38 @@ bool ULandscapeInfo::UpdateLayerInfoMap(ALandscapeProxy* Proxy /*= nullptr*/, bo
 
 #endif // WITH_EDITOR
 
-
-void FixupLandscapeGuidsIfInstanced(ALandscapeProxy* LandscapeProxy)
+/* if the outer world is instanced, we need to change our landscape guid(in a deterministic way)
+ * this avoids guid collisions when you instance a world (and its landscapes) multiple times,
+ * while maintaining the same GUID between landscape proxy objects within an instance
+ */ 
+namespace LandscapeInstancing
 {
-	// if the outer world is instanced, we need to change our landscape guid (in a deterministic way)
-	// this avoids guid collisions when you instance a world (and it's landscapes) multiple times,
-	// while maintaining the same GUID between landscape proxy objects within an instance
-	if (UWorld* OuterWorld = LandscapeProxy->GetTypedOuter<UWorld>())
+	void FixupLandscapeGuidsIfInstanced(ALandscapeProxy* LandscapeProxy)
 	{
-		UPackage* OuterPackage = OuterWorld->GetPackage();
-		if (OuterPackage->GetFName() != OuterPackage->GetLoadedPath().GetPackageFName())	// OuterWorld->IsInstanced()) is broken when PIE uses a memory built package
+		// we shouldn't be dealing with any instanced landscapes in these cases, early out
+		if (LandscapeProxy->IsTemplate() || IsRunningCookCommandlet())
+		{
+			return;
+		}
+
+		UWorldPartition* WorldPartition = FWorldPartitionHelpers::GetWorldPartition(LandscapeProxy);
+		UWorld* OuterWorld = WorldPartition ? WorldPartition->GetTypedOuter<UWorld>() : LandscapeProxy->GetTypedOuter<UWorld>();
+
+		// TODO [chris.tchou] : Note this is not 100% correct, IsInstanced() returns TRUE when using PIE on non-instanced landscapes.
+		// That is generally ok however, as the GUID remaps are still deterministic and landscape still works.
+		// the only issue is if some external system is tracking or interacting with landscape via the GUID, as it changes.
+		if (OuterWorld && OuterWorld->IsInstanced())
 		{
 			FArchiveMD5 Ar;
 			FGuid OldLandscapeGuid = LandscapeProxy->GetLandscapeGuid();
 			Ar << OldLandscapeGuid;
 
-#if WITH_EDITOR
-			// to work around PIE issues, we use the world partition to find the package
-			UWorldPartition* WorldPartition = FWorldPartitionHelpers::GetWorldPartition(LandscapeProxy);
-			if (WorldPartition)
+			UPackage* OuterWorldPackage = OuterWorld->GetPackage();
+			if (ensure(OuterWorldPackage))
 			{
-				OuterPackage = WorldPartition->GetPackage();
-				FName PackageName = OuterPackage->GetFName();
+				FName PackageName = OuterWorldPackage->GetFName();
 				Ar << PackageName;
 			}
-			else
-#endif // WITH_EDITOR
-			{
-				FObjectKey ContainerKey(OuterPackage);
-				Ar << ContainerKey;
-			}
 
 			FMD5Hash MD5Hash;
 			Ar.GetHash(MD5Hash);
@@ -3425,7 +3427,7 @@ void ALandscapeProxy::PostLoad()
 {
 	Super::PostLoad();
 
-	FixupLandscapeGuidsIfInstanced(this);
+	LandscapeInstancing::FixupLandscapeGuidsIfInstanced(this);
 
 	// Temporary
 	if (ComponentSizeQuads == 0 && LandscapeComponents.Num() > 0)
