Insight asset meta data tracking on DX12 buffers

diff --git a/Engine/Plugins/Runtime/HairStrands/Source/HairStrandsCore/Private/HairStrandsVertexFactory.cpp b/Engine/Plugins/Runtime/HairStrands/Source/HairStrandsCore/Private/HairStrandsVertexFactory.cpp
--- a/Engine/Plugins/Runtime/HairStrands/Source/HairStrandsCore/Private/HairStrandsVertexFactory.cpp
+++ b/Engine/Plugins/Runtime/HairStrands/Source/HairStrandsCore/Private/HairStrandsVertexFactory.cpp
@@ -45,6 +45,7 @@ public:
 	{
 		{
 			FRHIResourceCreateInfo CreateInfo(TEXT("FDummyCulledDispatchVertexIdsBuffer"));
+			CreateInfo.ClassName = TEXT("FDummyCulledDispatchVertexIdsBuffer");
 			uint32 NumBytes = sizeof(uint32) * 4;
 			VertexBufferRHI = RHICreateBuffer(NumBytes, BUF_Static | BUF_VertexBuffer | BUF_ShaderResource, 0, ERHIAccess::VertexOrIndexBuffer | ERHIAccess::SRVMask, CreateInfo);
 			uint32* DummyContents = (uint32*)RHILockBuffer(VertexBufferRHI, 0, NumBytes, RLM_WriteOnly);
@@ -53,7 +54,8 @@ public:
 		}
 
 		{
-			FRHIResourceCreateInfo CreateInfo(TEXT("FDummyByteAddressBufferBuffer"));
+			FRHIResourceCreateInfo CreateInfo(TEXT("FDummyByteAddressBuffer"));
+			CreateInfo.ClassName = TEXT("FDummyCulledDispatchVertexIdsBuffer");
 			uint32 NumBytes = sizeof(uint32) * 4;
 			ByteAddressBufferRHI = RHICreateBuffer(NumBytes, BUF_Static | BUF_ShaderResource |BUF_ByteAddressBuffer, 0, ERHIAccess::SRVMask, CreateInfo);
 			uint32* DummyContents = (uint32*)RHILockBuffer(ByteAddressBufferRHI, 0, NumBytes, RLM_WriteOnly);
diff --git a/Engine/Source/Runtime/D3D12RHI/Private/D3D12Adapter.h b/Engine/Source/Runtime/D3D12RHI/Private/D3D12Adapter.h
--- a/Engine/Source/Runtime/D3D12RHI/Private/D3D12Adapter.h
+++ b/Engine/Source/Runtime/D3D12RHI/Private/D3D12Adapter.h
@@ -353,7 +353,9 @@ public:
 		bool bHasInitialData,
 		const FRHIGPUMask& InGPUMask,
 		ID3D12ResourceAllocator* ResourceAllocator,
-		const TCHAR* InDebugName);
+		const TCHAR* InDebugName,
+		const FName& InAssetName = NAME_None,
+		const FName& InClassName = NAME_None);
 
 	void CreateUAVAliasResource(
 		D3D12_CLEAR_VALUE* ClearValuePtr,
diff --git a/Engine/Source/Runtime/D3D12RHI/Private/D3D12Buffer.cpp b/Engine/Source/Runtime/D3D12RHI/Private/D3D12Buffer.cpp
--- a/Engine/Source/Runtime/D3D12RHI/Private/D3D12Buffer.cpp
+++ b/Engine/Source/Runtime/D3D12RHI/Private/D3D12Buffer.cpp
@@ -5,6 +5,12 @@ D3D12Buffer.cpp: D3D Common code for buffers.
 =============================================================================*/
 
 #include "D3D12RHIPrivate.h"
+#include "ProfilingDebugging/AssetMetadataTrace.h"
+
+FName ChooseMetadataName(const FName& FirstChoice, const FName& SecondChoice)
+{
+	return FirstChoice == NAME_None ? SecondChoice : FirstChoice;
+}
 
 FD3D12Buffer::~FD3D12Buffer()
 {
@@ -66,6 +72,8 @@ struct FRHICommandRenameUploadBuffer final : public FRHICommand<FRHICommandRenam
 
 	void Execute(FRHICommandListBase& CmdList)
 	{
+		UE_TRACE_METADATA_SCOPE_ASSET_FNAME(ChooseMetadataName(Resource->GetOwnerName(), Resource->GetName()), FName(TEXT("FRHICommandRenameUploadBuffer::Execute")));
+
 		// Make sure we have an active pipeline when running at the top of the pipe, so the lambda below can obtain a context.
 		TOptional<ERHIPipeline> PreviousPipeline;
 		if (CmdList.IsTopOfPipe() && CmdList.GetPipeline() == ERHIPipeline::None)
@@ -192,8 +200,10 @@ struct FD3D12RHICommandInitializeBuffer final : public FRHICommand<FD3D12RHIComm
 };
 
 
-void FD3D12Buffer::UploadResourceData(FRHICommandListBase& RHICmdList, FResourceArrayInterface* InResourceArray, D3D12_RESOURCE_STATES InDestinationState)
+void FD3D12Buffer::UploadResourceData(FRHICommandListBase& RHICmdList, FResourceArrayInterface* InResourceArray, D3D12_RESOURCE_STATES InDestinationState, const FName& AssetName, const FName& ClassName)
 {
+	UE_TRACE_METADATA_SCOPE_ASSET_FNAME(AssetName, ClassName);
+
 	check(InResourceArray);
 	check(ResourceLocation.IsValid());
 
@@ -337,11 +347,15 @@ FD3D12Buffer* FD3D12Adapter::CreateRHIBuffer(
 	bool bHasInitialData,
 	const FRHIGPUMask& InGPUMask,
 	ID3D12ResourceAllocator* ResourceAllocator,
-	const TCHAR* InDebugName)
+	const TCHAR* InDebugName,
+	const FName& AssetName,
+	const FName& ClassName)
 {
 	TRACE_CPUPROFILER_EVENT_SCOPE(D3D12RHI::CreateRHIBuffer);
 	SCOPE_CYCLE_COUNTER(STAT_D3D12CreateBufferTime);
 
+	UE_TRACE_METADATA_SCOPE_ASSET_FNAME(ChooseMetadataName(AssetName, InDebugName), ChooseMetadataName(ClassName, TEXT("FRHIBuffer")));
+
 	check(InDesc.Width == Size);
 
 	FD3D12Buffer* BufferOut = nullptr;
@@ -525,6 +539,9 @@ FBufferRHIRef FD3D12DynamicRHI::CreateBuffer(FRHICommandListBase& RHICmdList, ui
 
 FD3D12Buffer* FD3D12DynamicRHI::CreateD3D12Buffer(class FRHICommandListBase* RHICmdList, uint32 Size, EBufferUsageFlags Usage, uint32 Stride, ERHIAccess InResourceState, FRHIResourceCreateInfo& CreateInfo, ID3D12ResourceAllocator* ResourceAllocator)
 {
+	FName TraceAssetName = ChooseMetadataName(CreateInfo.GetTraceAssetName(), CreateInfo.DebugName);
+	FName TraceClassName = ChooseMetadataName(CreateInfo.GetTraceClassName(), TEXT("FRHIBuffer"));
+
 	D3D12_RESOURCE_DESC Desc;
 	uint32 Alignment;
 	FD3D12Buffer::GetResourceDescAndAlignment(Size, Stride, Usage, Desc, Alignment);
@@ -548,14 +565,14 @@ FD3D12Buffer* FD3D12DynamicRHI::CreateD3D12Buffer(class FRHICommandListBase* RHI
 	D3D12_RESOURCE_STATES CreateState = (CreateInfo.ResourceArray && bSupportResourceStateTracking) ? D3D12_RESOURCE_STATE_COPY_DEST : DesiredState;
 	bool bHasInitialData = CreateInfo.ResourceArray != nullptr;
 
-	FD3D12Buffer* Buffer = GetAdapter().CreateRHIBuffer(Desc, Alignment, Stride, Size, Usage, StateMode, CreateState, bHasInitialData, CreateInfo.GPUMask, ResourceAllocator, CreateInfo.DebugName);
+	FD3D12Buffer* Buffer = GetAdapter().CreateRHIBuffer(Desc, Alignment, Stride, Size, Usage, StateMode, CreateState, bHasInitialData, CreateInfo.GPUMask, ResourceAllocator, CreateInfo.DebugName, TraceAssetName, TraceClassName);
 	check(Buffer->ResourceLocation.IsValid());
 
 	// Copy the resource data if available 
 	if (bHasInitialData)
 	{
 		check(RHICmdList);
-		Buffer->UploadResourceData(*RHICmdList, CreateInfo.ResourceArray, DesiredState);
+		Buffer->UploadResourceData(*RHICmdList, CreateInfo.ResourceArray, DesiredState, TraceAssetName, TraceClassName);
 	}
 
 	return Buffer;
@@ -570,6 +587,7 @@ FRHIBuffer* FD3D12DynamicRHI::CreateBuffer(const FRHIBufferCreateInfo& CreateInf
 void* FD3D12DynamicRHI::LockBuffer(FRHICommandListBase& RHICmdList, FD3D12Buffer* Buffer, uint32 BufferSize, EBufferUsageFlags BufferUsage, uint32 Offset, uint32 Size, EResourceLockMode LockMode)
 {
 	SCOPE_CYCLE_COUNTER(STAT_D3D12LockBufferTime);
+	UE_TRACE_METADATA_SCOPE_ASSET_FNAME(ChooseMetadataName(Buffer->GetOwnerName(), Buffer->GetName()), Buffer->GetName());
 
 	checkf(Size <= BufferSize, TEXT("Requested lock size %u is larger than the total size %u for buffer '%s'."), Size, BufferSize, *Buffer->GetName().ToString());
 
diff --git a/Engine/Source/Runtime/D3D12RHI/Private/D3D12Commands.cpp b/Engine/Source/Runtime/D3D12RHI/Private/D3D12Commands.cpp
--- a/Engine/Source/Runtime/D3D12RHI/Private/D3D12Commands.cpp
+++ b/Engine/Source/Runtime/D3D12RHI/Private/D3D12Commands.cpp
@@ -18,6 +18,7 @@ D3D12Commands.cpp: D3D RHI commands implementation.
 #include "RenderUtils.h"
 #include "GlobalRenderResources.h"
 #include "RHIShaderParametersShared.h"
+#include "ProfilingDebugging/AssetMetadataTrace.h"
 
 static int32 GD3D12TransientAllocatorFullAliasingBarrier = 0;
 static FAutoConsoleVariableRef CVarD3D12TransientAllocatorFullAliasingBarrier(
@@ -675,6 +676,9 @@ void FD3D12CommandContext::RHICopyToStagingBuffer(FRHIBuffer* SourceBufferRHI, F
 {
 	SCOPE_CYCLE_COUNTER(STAT_D3D12CopyToStagingBufferTime);
 
+	FName AssetName = SourceBufferRHI->GetOwnerName() == NAME_None ? SourceBufferRHI->GetName() : SourceBufferRHI->GetOwnerName();
+	UE_TRACE_METADATA_SCOPE_ASSET_FNAME(AssetName, FName(TEXT("FRHIStagingBuffer")));
+
 	FD3D12StagingBuffer* StagingBuffer = FD3D12DynamicRHI::ResourceCast(StagingBufferRHI);
 	check(StagingBuffer);
 	ensureMsgf(!StagingBuffer->bIsLocked, TEXT("Attempting to Copy to a locked staging buffer. This may have undefined behavior"));
diff --git a/Engine/Source/Runtime/D3D12RHI/Private/D3D12Query.cpp b/Engine/Source/Runtime/D3D12RHI/Private/D3D12Query.cpp
--- a/Engine/Source/Runtime/D3D12RHI/Private/D3D12Query.cpp
+++ b/Engine/Source/Runtime/D3D12RHI/Private/D3D12Query.cpp
@@ -5,6 +5,7 @@
 =============================================================================*/
 
 #include "D3D12RHIPrivate.h"
+#include "ProfilingDebugging/AssetMetadataTrace.h"
 
 namespace D3D12RHI
 {
@@ -74,6 +75,8 @@ FD3D12QueryHeap::FD3D12QueryHeap(FD3D12Device* Device, D3D12_QUERY_TYPE QueryTyp
 		break;
 	}
 
+	UE_TRACE_METADATA_SCOPE_ASSET_FNAME(FName(ResultBufferName), FName(TEXT("FD3D12QueryHeap")));
+
 	// Create the query heap
 #if D3D12RHI_PLATFORM_USES_TIMESTAMP_QUERIES == 0
 	if (HeapType != D3D12_QUERY_HEAP_TYPE_TIMESTAMP)
diff --git a/Engine/Source/Runtime/D3D12RHI/Private/D3D12Resources.cpp b/Engine/Source/Runtime/D3D12RHI/Private/D3D12Resources.cpp
--- a/Engine/Source/Runtime/D3D12RHI/Private/D3D12Resources.cpp
+++ b/Engine/Source/Runtime/D3D12RHI/Private/D3D12Resources.cpp
@@ -8,6 +8,7 @@ D3D12Resources.cpp: D3D RHI utility implementation.
 #include "EngineModule.h"
 #include "HAL/LowLevelMemTracker.h"
 #include "ProfilingDebugging/MemoryTrace.h"
+#include "ProfilingDebugging/AssetMetadataTrace.h"
 
 #if INTEL_EXTENSIONS
 	#define INTC_IGDEXT_D3D12 1
@@ -1101,6 +1102,8 @@ bool FD3D12ResourceLocation::OnAllocationMoved(FRHIPoolAllocationData* InNewData
 	FD3D12Resource* CurrentResource = GetResource();
 	FD3D12PoolAllocator* NewAllocator = GetPoolAllocator();
 
+	UE_TRACE_METADATA_SCOPE_ASSET_FNAME(CurrentResource->GetName(), FName(TEXT("FD3D12ResourceLocation::OnAllocationMoved")));
+
 	// If sub allocated and not placed only update the internal data
 	if (NewAllocator->GetAllocationStrategy() == EResourceAllocationStrategy::kManualSubAllocation)
 	{
diff --git a/Engine/Source/Runtime/D3D12RHI/Private/D3D12TransientResourceAllocator.cpp b/Engine/Source/Runtime/D3D12RHI/Private/D3D12TransientResourceAllocator.cpp
--- a/Engine/Source/Runtime/D3D12RHI/Private/D3D12TransientResourceAllocator.cpp
+++ b/Engine/Source/Runtime/D3D12RHI/Private/D3D12TransientResourceAllocator.cpp
@@ -4,6 +4,7 @@
 #include "D3D12TransientResourceAllocator.h"
 #include "D3D12Stats.h"
 #include "ProfilingDebugging/MemoryTrace.h"
+#include "ProfilingDebugging/AssetMetadataTrace.h"
 
 D3D12_RESOURCE_STATES GetInitialResourceState(const D3D12_RESOURCE_DESC& InDesc)
 {
@@ -28,6 +29,8 @@ D3D12_RESOURCE_STATES GetInitialResourceState(const D3D12_RESOURCE_DESC& InDesc)
 FD3D12TransientHeap::FD3D12TransientHeap(const FInitializer& Initializer, FD3D12Adapter* Adapter, FD3D12Device* Device, FRHIGPUMask VisibleNodeMask)
 	: FRHITransientHeap(Initializer)
 {
+	UE_TRACE_METADATA_SCOPE_ASSET_FNAME(FName(TEXT("FD3D12TransientHeap")), FName(TEXT("FD3D12TransientHeap")));
+
 	D3D12_HEAP_FLAGS HeapFlags = D3D12_HEAP_FLAG_ALLOW_ALL_BUFFERS_AND_TEXTURES;
 
 	if (Initializer.Flags != ERHITransientHeapFlags::AllowAll)
diff --git a/Engine/Source/Runtime/D3D12RHI/Private/Windows/WindowsD3D12Device.cpp b/Engine/Source/Runtime/D3D12RHI/Private/Windows/WindowsD3D12Device.cpp
--- a/Engine/Source/Runtime/D3D12RHI/Private/Windows/WindowsD3D12Device.cpp
+++ b/Engine/Source/Runtime/D3D12RHI/Private/Windows/WindowsD3D12Device.cpp
@@ -10,6 +10,7 @@
 #include "Windows/AllowWindowsPlatformTypes.h"
 #include "Windows/WindowsPlatformCrashContext.h"
 #include "HAL/FileManager.h"
+#include "ProfilingDebugging/AssetMetadataTrace.h"
 #include <delayimp.h>
 
 #if WITH_AMD_AGS
@@ -1968,6 +1969,8 @@ void FWindowsD3D12Adapter::CreateCommandSignatures()
 
 TUniquePtr<FD3D12DiagnosticBuffer> FD3D12Device::CreateDiagnosticBuffer(const D3D12_RESOURCE_DESC& Desc, const TCHAR* Name)
 {
+	UE_TRACE_METADATA_SCOPE_ASSET_FNAME(FName(TEXT("FD3D12DiagnosticBuffer")), FName(TEXT("FD3D12Device::CreateDiagnosticBuffer")));
+
 	TRefCountPtr<ID3D12Device3> D3D12Device3;
 	HRESULT hr = GetDevice()->QueryInterface(IID_PPV_ARGS(D3D12Device3.GetInitReference()));
 	if (SUCCEEDED(hr))
diff --git a/Engine/Source/Runtime/D3D12RHI/Public/D3D12Resources.h b/Engine/Source/Runtime/D3D12RHI/Public/D3D12Resources.h
--- a/Engine/Source/Runtime/D3D12RHI/Public/D3D12Resources.h
+++ b/Engine/Source/Runtime/D3D12RHI/Public/D3D12Resources.h
@@ -879,7 +879,7 @@ public:
 
 	virtual uint32 GetParentGPUIndex() const override;
 
-	void UploadResourceData(FRHICommandListBase& InRHICmdList, FResourceArrayInterface* InResourceArray, D3D12_RESOURCE_STATES InDestinationState);
+	void UploadResourceData(FRHICommandListBase& InRHICmdList, FResourceArrayInterface* InResourceArray, D3D12_RESOURCE_STATES InDestinationState, const FName& AssetName = NAME_None, const FName& ClassName = NAME_None);
 	FD3D12SyncPointRef UploadResourceDataViaCopyQueue(FResourceArrayInterface* InResourceArray);
 
 	// FRHIResource overrides
diff --git a/Engine/Source/Runtime/Engine/Private/Components/TextRenderComponent.cpp b/Engine/Source/Runtime/Engine/Private/Components/TextRenderComponent.cpp
--- a/Engine/Source/Runtime/Engine/Private/Components/TextRenderComponent.cpp
+++ b/Engine/Source/Runtime/Engine/Private/Components/TextRenderComponent.cpp
@@ -699,6 +699,11 @@ void FTextRenderSceneProxy::CreateRenderThreadResources()
 	TArray<FDynamicMeshVertex> OutVertices;
 	if(BuildStringMesh(OutVertices, IndexBuffer.Indices))
 	{
+#if RHI_ENABLE_RESOURCE_INFO
+		FName Name = FName(TEXT("FTextRenderSceneProxy ") + GetOwnerName().ToString());
+		VertexBuffers.SetOwnerName(Name);
+		IndexBuffer.SetOwnerName(Name);
+#endif
 		VertexBuffers.InitFromDynamicVertex(&VertexFactory, OutVertices);
 		// Enqueue initialization of render resources
 		BeginInitResource(&IndexBuffer);
diff --git a/Engine/Source/Runtime/Engine/Private/RawIndexBuffer.cpp b/Engine/Source/Runtime/Engine/Private/RawIndexBuffer.cpp
--- a/Engine/Source/Runtime/Engine/Private/RawIndexBuffer.cpp
+++ b/Engine/Source/Runtime/Engine/Private/RawIndexBuffer.cpp
@@ -341,6 +341,8 @@ FBufferRHIRef FRawStaticIndexBuffer::CreateRHIBuffer_Internal()
 		// Create the index buffer.
 		FBufferRHIRef Ret;
 		FRHIResourceCreateInfo CreateInfo(Is32Bit() ? TEXT("FRawStaticIndexBuffer32") : TEXT("FRawStaticIndexBuffer16"), &IndexStorage);
+		CreateInfo.ClassName = CreateInfo.DebugName;
+		CreateInfo.AssetName = GetOwnerName();
 		CreateInfo.bWithoutNativeResource = !SizeInBytes;
 		if (bRenderThread)
 		{
@@ -510,6 +512,8 @@ FBufferRHIRef FRawStaticIndexBuffer16or32Interface::CreateRHIIndexBufferInternal
 {
 	// Create the index buffer.
 	FRHIResourceCreateInfo CreateInfo(InDebugName, ResourceArray);
+	CreateInfo.ClassName = InDebugName;
+	CreateInfo.AssetName = InOwnerName;
 	EBufferUsageFlags Flags = EBufferUsageFlags::Static;
 
 	if (bNeedSRV)
diff --git a/Engine/Source/Runtime/Engine/Private/Rendering/MorphTargetVertexInfoBuffers.cpp b/Engine/Source/Runtime/Engine/Private/Rendering/MorphTargetVertexInfoBuffers.cpp
--- a/Engine/Source/Runtime/Engine/Private/Rendering/MorphTargetVertexInfoBuffers.cpp
+++ b/Engine/Source/Runtime/Engine/Private/Rendering/MorphTargetVertexInfoBuffers.cpp
@@ -20,6 +20,8 @@ void FMorphTargetVertexInfoBuffers::InitRHI()
 
 	const uint32 BufferSize = MorphData.Num() * sizeof(uint32);
 	FRHIResourceCreateInfo CreateInfo(TEXT("MorphData"));
+	CreateInfo.ClassName = TEXT("FMorphTargetVertexInfoBuffers");
+	CreateInfo.AssetName = GetOwnerName();
 	MorphDataBuffer = RHICreateStructuredBuffer(sizeof(uint32), BufferSize, BUF_Static | BUF_ByteAddressBuffer | BUF_ShaderResource, ERHIAccess::SRVMask, CreateInfo);
 	MorphDataBuffer->SetOwnerName(GetOwnerName());
 	
diff --git a/Engine/Source/Runtime/Engine/Private/StaticMesh.cpp b/Engine/Source/Runtime/Engine/Private/StaticMesh.cpp
--- a/Engine/Source/Runtime/Engine/Private/StaticMesh.cpp
+++ b/Engine/Source/Runtime/Engine/Private/StaticMesh.cpp
@@ -1228,6 +1228,13 @@ void FStaticMeshVertexBuffers::InitFromDynamicVertex(FLocalVertexFactory* Vertex
 		});
 };
 
+void FStaticMeshVertexBuffers::SetOwnerName(const FName& OwnerName)
+{
+	PositionVertexBuffer.SetOwnerName(OwnerName);
+	StaticMeshVertexBuffer.SetOwnerName(OwnerName);
+	ColorVertexBuffer.SetOwnerName(OwnerName);
+}
+
 FStaticMeshLODResources::FStaticMeshLODResources(bool bAddRef)
 	: CardRepresentationData(nullptr)
 	, MaxDeviation(0.0f)
diff --git a/Engine/Source/Runtime/Engine/Public/GPUSkinCache.h b/Engine/Source/Runtime/Engine/Public/GPUSkinCache.h
--- a/Engine/Source/Runtime/Engine/Public/GPUSkinCache.h
+++ b/Engine/Source/Runtime/Engine/Public/GPUSkinCache.h
@@ -199,6 +199,8 @@ public:
 		{
 			for (int32 Index = 0; Index < NUM_BUFFERS; ++Index)
 			{
+				PositionBuffers[Index].Buffer.ClassName = TEXT("SkinCachePositions");
+				PositionBuffers[Index].Buffer.AssetName = OwnerName;
 				PositionBuffers[Index].Buffer.Initialize(TEXT("SkinCachePositions"), PosBufferBytesPerElement, NumVertices * 3, PF_R32_FLOAT, BUF_Static);
 				PositionBuffers[Index].Buffer.Buffer->SetOwnerName(OwnerName);
 				PositionBuffers[Index].AccessState = ERHIAccess::Unknown;
@@ -208,11 +210,15 @@ public:
 				// OpenGL ES does not support writing to RGBA16_SNORM images, instead pack data into SINT in the shader
 				const EPixelFormat TangentsFormat = IsOpenGLPlatform(GMaxRHIShaderPlatform) ? PF_R16G16B16A16_SINT : PF_R16G16B16A16_SNORM;
 				
+				Tangents.Buffer.ClassName = TEXT("SkinCacheTangents");
+				Tangents.Buffer.AssetName = OwnerName;
 				Tangents.Buffer.Initialize(TEXT("SkinCacheTangents"), TangentBufferBytesPerElement, NumVertices * 2, TangentsFormat, BUF_Static);
 				Tangents.Buffer.Buffer->SetOwnerName(OwnerName);
 				Tangents.AccessState = ERHIAccess::Unknown;
 				if (UseIntermediateTangents)
 				{
+					IntermediateTangents.Buffer.ClassName = TEXT("SkinCacheIntermediateTangents");
+					IntermediateTangents.Buffer.AssetName = OwnerName;
 					IntermediateTangents.Buffer.Initialize(TEXT("SkinCacheIntermediateTangents"), TangentBufferBytesPerElement, NumVertices * 2, TangentsFormat, BUF_Static);
 					IntermediateTangents.Buffer.Buffer->SetOwnerName(OwnerName);
 					IntermediateTangents.AccessState = ERHIAccess::Unknown;
@@ -220,6 +226,8 @@ public:
 			}
 			if (IntermediateAccumulatedTangentsSize > 0)
 			{
+				IntermediateAccumulatedTangents.Buffer.ClassName = TEXT("SkinCacheIntermediateAccumulatedTangents");
+				IntermediateAccumulatedTangents.Buffer.AssetName = OwnerName;
 				IntermediateAccumulatedTangents.Buffer.Initialize(TEXT("SkinCacheIntermediateAccumulatedTangents"), sizeof(int32), IntermediateAccumulatedTangentsSize * FGPUSkinCache::IntermediateAccumBufferNumInts, PF_R32_SINT, BUF_UnorderedAccess);
 				IntermediateAccumulatedTangents.Buffer.Buffer->SetOwnerName(OwnerName);
 				IntermediateAccumulatedTangents.AccessState = ERHIAccess::Unknown;
diff --git a/Engine/Source/Runtime/Engine/Public/Rendering/SkeletalMeshDuplicatedVerticesBuffer.h b/Engine/Source/Runtime/Engine/Public/Rendering/SkeletalMeshDuplicatedVerticesBuffer.h
--- a/Engine/Source/Runtime/Engine/Public/Rendering/SkeletalMeshDuplicatedVerticesBuffer.h
+++ b/Engine/Source/Runtime/Engine/Public/Rendering/SkeletalMeshDuplicatedVerticesBuffer.h
@@ -112,6 +112,8 @@ public:
             check(ResourceArray->GetResourceDataSize() > 0);
 
             FRHIResourceCreateInfo CreateInfo(TEXT("DuplicatedVerticesIndexBuffer"), ResourceArray);
+			CreateInfo.ClassName = TEXT("FDuplicatedVerticesBuffer DuplicatedVerticesIndexBuffer");
+			CreateInfo.AssetName = GetOwnerName();
             DuplicatedVerticesIndexBuffer.VertexBufferRHI = RHICreateVertexBuffer(ResourceArray->GetResourceDataSize(), BUF_Static | BUF_ShaderResource, CreateInfo);
 			DuplicatedVerticesIndexBuffer.VertexBufferRHI->SetOwnerName(GetOwnerName());
             DuplicatedVerticesIndexBuffer.VertexBufferSRV = RHICreateShaderResourceView(DuplicatedVerticesIndexBuffer.VertexBufferRHI, sizeof(uint32), PF_R32_UINT);
@@ -122,6 +124,8 @@ public:
             check(ResourceArray->GetResourceDataSize() > 0);
 
             FRHIResourceCreateInfo CreateInfo(TEXT("LengthAndIndexDuplicatedVerticesIndexBuffer"), ResourceArray);
+			CreateInfo.ClassName = TEXT("FDuplicatedVerticesBuffer LengthAndIndexDuplicatedVerticesIndexBuffer");
+			CreateInfo.AssetName = GetOwnerName();
             LengthAndIndexDuplicatedVerticesIndexBuffer.VertexBufferRHI = RHICreateVertexBuffer(ResourceArray->GetResourceDataSize(), BUF_Static | BUF_ShaderResource, CreateInfo);
 			LengthAndIndexDuplicatedVerticesIndexBuffer.VertexBufferRHI->SetOwnerName(GetOwnerName());
             LengthAndIndexDuplicatedVerticesIndexBuffer.VertexBufferSRV = RHICreateShaderResourceView(LengthAndIndexDuplicatedVerticesIndexBuffer.VertexBufferRHI, sizeof(uint32), PF_R32_UINT);
diff --git a/Engine/Source/Runtime/Engine/Public/StaticMeshResources.h b/Engine/Source/Runtime/Engine/Public/StaticMeshResources.h
--- a/Engine/Source/Runtime/Engine/Public/StaticMeshResources.h
+++ b/Engine/Source/Runtime/Engine/Public/StaticMeshResources.h
@@ -325,6 +325,8 @@ struct FStaticMeshVertexBuffers
 
 	/* This is a temporary function to refactor and convert old code, do not copy this as is and try to build your data as SoA from the beginning.*/
 	void ENGINE_API InitModelVF(FLocalVertexFactory* VertexFactory);
+
+	void SetOwnerName(const FName& OwnerName);
 };
 
 struct FAdditionalStaticMeshIndexBuffers
diff --git a/Engine/Source/Runtime/Landscape/Private/LandscapeRender.cpp b/Engine/Source/Runtime/Landscape/Private/LandscapeRender.cpp
--- a/Engine/Source/Runtime/Landscape/Private/LandscapeRender.cpp
+++ b/Engine/Source/Runtime/Landscape/Private/LandscapeRender.cpp
@@ -673,6 +673,7 @@ void FLandscapeRenderSystem::UpdateBuffers()
 		if (!SectionLODBiasBuffer.IsValid())
 		{
 			FRHIResourceCreateInfo CreateInfo(TEXT("SectionLODBiasBuffer"), &SectionLODBiases);
+			CreateInfo.ClassName = TEXT("FLandscapeRenderSystem");
 			SectionLODBiasBuffer = RHICreateVertexBuffer(SectionLODBiases.GetResourceDataSize(), BUF_ShaderResource | BUF_Dynamic, CreateInfo);
 			SectionLODBiasSRV = RHICreateShaderResourceView(SectionLODBiasBuffer, sizeof(float), PF_R32_FLOAT);
 			bUpdateUB = true;
@@ -1180,7 +1181,7 @@ void FLandscapeComponentSceneProxy::CreateRenderThreadResources()
 	if (SharedBuffers == nullptr)
 	{
 		SharedBuffers = new FLandscapeSharedBuffers(
-			SharedBuffersKey, SubsectionSizeQuads, NumSubsections, FeatureLevel);
+			SharedBuffersKey, SubsectionSizeQuads, NumSubsections, FeatureLevel, FName(LandscapeComponent->GetPathName()));
 
 		FLandscapeComponentSceneProxy::SharedBuffersMap.Add(SharedBuffersKey, SharedBuffers);
 
@@ -2592,6 +2593,8 @@ void FLandscapeVertexBuffer::InitRHI()
 
 	// create a static vertex buffer
 	FRHIResourceCreateInfo CreateInfo(TEXT("FLandscapeVertexBuffer"));
+	CreateInfo.ClassName = FName(TEXT("FLandscapeVertexBuffer"));
+	CreateInfo.AssetName = GetOwnerName();
 	VertexBufferRHI = RHICreateBuffer(NumVertices * sizeof(FLandscapeVertex), BUF_Static | BUF_VertexBuffer, 0, ERHIAccess::VertexOrIndexBuffer, CreateInfo);
 	FLandscapeVertex* Vertex = (FLandscapeVertex*)RHILockBuffer(VertexBufferRHI, 0, NumVertices * sizeof(FLandscapeVertex), RLM_WriteOnly);
 	int32 VertexIndex = 0;
@@ -2622,7 +2625,7 @@ void FLandscapeVertexBuffer::InitRHI()
 //
 
 template <typename INDEX_TYPE>
-void FLandscapeSharedBuffers::CreateIndexBuffers()
+void FLandscapeSharedBuffers::CreateIndexBuffers(const FName& InOwnerName)
 {
 	TArray<INDEX_TYPE> VertexToIndexMap;
 	VertexToIndexMap.AddUninitialized(FMath::Square(SubsectionSizeVerts * NumSubsections));
@@ -2704,7 +2707,7 @@ void FLandscapeSharedBuffers::CreateIndexBuffers()
 			IndexBuffer = new FRawStaticIndexBuffer16or32<INDEX_TYPE>(false);
 		}
 		IndexBuffer->AssignNewBuffer(NewIndices);
-
+		IndexBuffer->SetOwnerName(InOwnerName);
 		IndexBuffer->InitResource();
 
 		IndexBuffers[Mip] = IndexBuffer;
@@ -2744,7 +2747,7 @@ void FLandscapeSharedBuffers::CreateIndexBuffers()
 
 #if WITH_EDITOR
 template <typename INDEX_TYPE>
-void FLandscapeSharedBuffers::CreateGrassIndexBuffer()
+void FLandscapeSharedBuffers::CreateGrassIndexBuffer(const FName& InOwnerName)
 {
 	TArray<INDEX_TYPE> NewIndices;
 
@@ -2784,12 +2787,13 @@ void FLandscapeSharedBuffers::CreateGrassIndexBuffer()
 	// Create and init new index buffer with index data
 	FRawStaticIndexBuffer16or32<INDEX_TYPE>* IndexBuffer = new FRawStaticIndexBuffer16or32<INDEX_TYPE>(false);
 	IndexBuffer->AssignNewBuffer(NewIndices);
+	IndexBuffer->SetOwnerName(InOwnerName);
 	IndexBuffer->InitResource();
 	GrassIndexBuffer = IndexBuffer;
 }
 #endif
 
-FLandscapeSharedBuffers::FLandscapeSharedBuffers(const int32 InSharedBuffersKey, const int32 InSubsectionSizeQuads, const int32 InNumSubsections, const ERHIFeatureLevel::Type InFeatureLevel)
+FLandscapeSharedBuffers::FLandscapeSharedBuffers(const int32 InSharedBuffersKey, const int32 InSubsectionSizeQuads, const int32 InNumSubsections, const ERHIFeatureLevel::Type InFeatureLevel, const FName& InOwnerName)
 	: SharedBuffersKey(InSharedBuffersKey)
 	, NumIndexBuffers(FMath::CeilLogTwo(InSubsectionSizeQuads + 1))
 	, SubsectionSizeVerts(InSubsectionSizeQuads + 1)
@@ -2805,7 +2809,7 @@ FLandscapeSharedBuffers::FLandscapeSharedBuffers(const int32 InSharedBuffersKey,
 {
 	NumVertices = FMath::Square(SubsectionSizeVerts) * FMath::Square(NumSubsections);
 	
-	VertexBuffer = new FLandscapeVertexBuffer(InFeatureLevel, NumVertices, SubsectionSizeVerts, NumSubsections);
+	VertexBuffer = new FLandscapeVertexBuffer(InFeatureLevel, NumVertices, SubsectionSizeVerts, NumSubsections, InOwnerName);
 
 	IndexBuffers = new FIndexBuffer * [NumIndexBuffers];
 	FMemory::Memzero(IndexBuffers, sizeof(FIndexBuffer*) * NumIndexBuffers);
@@ -2822,21 +2826,21 @@ FLandscapeSharedBuffers::FLandscapeSharedBuffers(const int32 InSharedBuffersKey,
 	if (NumVertices > 65535)
 	{
 		bUse32BitIndices = true;
-		CreateIndexBuffers<uint32>();
+		CreateIndexBuffers<uint32>(InOwnerName);
 #if WITH_EDITOR
 		if (InFeatureLevel > ERHIFeatureLevel::ES3_1)
 		{
-			CreateGrassIndexBuffer<uint32>();
+			CreateGrassIndexBuffer<uint32>(InOwnerName);
 		}
 #endif
 	}
 	else
 	{
-		CreateIndexBuffers<uint16>();
+		CreateIndexBuffers<uint16>(InOwnerName);
 #if WITH_EDITOR
 		if (InFeatureLevel > ERHIFeatureLevel::ES3_1)
 		{
-			CreateGrassIndexBuffer<uint16>();
+			CreateGrassIndexBuffer<uint16>(InOwnerName);
 		}
 #endif
 	}
diff --git a/Engine/Source/Runtime/Landscape/Public/LandscapeRender.h b/Engine/Source/Runtime/Landscape/Public/LandscapeRender.h
--- a/Engine/Source/Runtime/Landscape/Public/LandscapeRender.h
+++ b/Engine/Source/Runtime/Landscape/Public/LandscapeRender.h
@@ -287,12 +287,13 @@ class FLandscapeVertexBuffer final : public FVertexBuffer
 public:
 
 	/** Constructor. */
-	FLandscapeVertexBuffer(ERHIFeatureLevel::Type InFeatureLevel, int32 InNumVertices, int32 InSubsectionSizeVerts, int32 InNumSubsections)
+	FLandscapeVertexBuffer(ERHIFeatureLevel::Type InFeatureLevel, int32 InNumVertices, int32 InSubsectionSizeVerts, int32 InNumSubsections, const FName& InOwnerName)
 		: FeatureLevel(InFeatureLevel)
 		, NumVertices(InNumVertices)
 		, SubsectionSizeVerts(InSubsectionSizeVerts)
 		, NumSubsections(InNumSubsections)
 	{
+		SetOwnerName(InOwnerName);
 		InitResource();
 	}
 
@@ -343,14 +344,14 @@ public:
 	TArray<FIndexBuffer*> ZeroOffsetIndexBuffers;
 #endif
 
-	FLandscapeSharedBuffers(int32 SharedBuffersKey, int32 SubsectionSizeQuads, int32 NumSubsections, ERHIFeatureLevel::Type FeatureLevel);
+	FLandscapeSharedBuffers(int32 SharedBuffersKey, int32 SubsectionSizeQuads, int32 NumSubsections, ERHIFeatureLevel::Type FeatureLevel, const FName& OwnerName = NAME_None);
 
 	template <typename INDEX_TYPE>
-	void CreateIndexBuffers();
+	void CreateIndexBuffers(const FName& OwnerName);
 	
 #if WITH_EDITOR
 	template <typename INDEX_TYPE>
-	void CreateGrassIndexBuffer();
+	void CreateGrassIndexBuffer(const FName& InOwnerName);
 #endif
 
 	virtual ~FLandscapeSharedBuffers();
diff --git a/Engine/Source/Runtime/RHI/Public/RHIResources.h b/Engine/Source/Runtime/RHI/Public/RHIResources.h
--- a/Engine/Source/Runtime/RHI/Public/RHIResources.h
+++ b/Engine/Source/Runtime/RHI/Public/RHIResources.h
@@ -391,6 +391,9 @@ struct FRHIResourceCreateInfo
 		ExtData = InExtData;
 	}
 
+	FName GetTraceClassName() const									{ return ClassName == NAME_None ? FName(TEXT("FRHIBuffer")) : ClassName; }
+	FName GetTraceAssetName() const									{ return AssetName == NAME_None ? FName(DebugName) : AssetName; }
+
 	// for CreateTexture calls
 	FResourceBulkDataInterface* BulkData;
 
@@ -409,6 +412,9 @@ struct FRHIResourceCreateInfo
 
 	// optional data that would have come from an offline cooker or whatever - general purpose
 	uint32 ExtData;
+
+	FName ClassName = NAME_None;	// The owner class of FRHIBuffer used for Insight asset metadata tracing
+	FName AssetName = NAME_None;	// The owner asset name used for Insight asset metadata tracing
 };
 
 class FExclusiveDepthStencil
diff --git a/Engine/Source/Runtime/RHI/Public/RHIUtilities.h b/Engine/Source/Runtime/RHI/Public/RHIUtilities.h
--- a/Engine/Source/Runtime/RHI/Public/RHIUtilities.h
+++ b/Engine/Source/Runtime/RHI/Public/RHIUtilities.h
@@ -162,6 +162,8 @@ struct FRWBuffer
 	FUnorderedAccessViewRHIRef UAV;
 	FShaderResourceViewRHIRef SRV;
 	uint32 NumBytes;
+	FName ClassName = NAME_None;	// The owner class of Buffer used for Insight asset metadata tracing, set it before calling Initialize()
+	FName AssetName = NAME_None;	// The owner asset name used for Insight asset metadata tracing, set it before calling Initialize()
 
 	FRWBuffer()
 		: NumBytes(0)
@@ -217,6 +219,8 @@ struct FRWBuffer
 		ensure(!(EnumHasAnyFlags(AdditionalUsage, BUF_FastVRAM) && !InDebugName));
 		NumBytes = BytesPerElement * NumElements;
 		FRHIResourceCreateInfo CreateInfo(InDebugName);
+		CreateInfo.ClassName = ClassName;
+		CreateInfo.AssetName = AssetName;
 		CreateInfo.ResourceArray = InResourceArray;
 		Buffer = RHICreateVertexBuffer(NumBytes, BUF_UnorderedAccess | BUF_ShaderResource | AdditionalUsage, InResourceState, CreateInfo);
 		UAV = RHICreateUnorderedAccessView(Buffer, UE_PIXELFORMAT_TO_UINT8(Format));
diff --git a/Engine/Source/Runtime/RenderCore/Private/RenderResource.cpp b/Engine/Source/Runtime/RenderCore/Private/RenderResource.cpp
--- a/Engine/Source/Runtime/RenderCore/Private/RenderResource.cpp
+++ b/Engine/Source/Runtime/RenderCore/Private/RenderResource.cpp
@@ -334,6 +334,8 @@ FBufferRHIRef FRenderResource::CreateRHIBufferInternal(
 {
 	const uint32 SizeInBytes = ResourceArray ? ResourceArray->GetResourceDataSize() : 0;
 	FRHIResourceCreateInfo CreateInfo(InDebugName, ResourceArray);
+	CreateInfo.ClassName = FName(InDebugName);
+	CreateInfo.AssetName = InOwnerName;
 	CreateInfo.bWithoutNativeResource = bWithoutNativeResource;
 
 	FBufferRHIRef Buffer;
