* Prevent landscape gizmo actor from being duplicated
* Prevent alt-dragging the landscape gizmo in landscape mode (tries to duplicate it and crashes)
* Nullptr-check for landscape copy/paste tool for CurrentLandscapeInfo, like is done elsewhere

diff --git a/Engine/Source/Editor/LandscapeEditor/Private/LandscapeEdMode.cpp b/Engine/Source/Editor/LandscapeEditor/Private/LandscapeEdMode.cpp
--- a/Engine/Source/Editor/LandscapeEditor/Private/LandscapeEdMode.cpp
+++ b/Engine/Source/Editor/LandscapeEditor/Private/LandscapeEdMode.cpp
@@ -1933,6 +1933,13 @@ bool FEdModeLandscape::InputKey(FEditorViewportClient* ViewportClient, FViewport
 				}
 				return true;
 			}
+
+			if (CurrentGizmoActor.IsValid() && CurrentGizmoActor->IsSelected() && (GLandscapeEditRenderMode & ELandscapeEditRenderMode::Gizmo) && IsAltDown(Viewport))
+			{ 
+				// When manipulating the gizmo with Alt down, don't do anything but return true to indicate it's on purpose and prevent the editor from doing the usual thing
+				//  (alt-drag is usually "duplicate actor" and we want to prevent duplicating the gizmo actor) 
+				return true; 
+			}
 		}
 
 		if (Key == EKeys::LeftMouseButton ||
diff --git a/Engine/Source/Runtime/Landscape/Classes/LandscapeGizmoActor.h b/Engine/Source/Runtime/Landscape/Classes/LandscapeGizmoActor.h
--- a/Engine/Source/Runtime/Landscape/Classes/LandscapeGizmoActor.h
+++ b/Engine/Source/Runtime/Landscape/Classes/LandscapeGizmoActor.h
@@ -45,13 +45,10 @@ public:
 
 
 #if WITH_EDITOR
+	virtual bool ShouldExport() override { return false; } // Prevent copy-pasting
 	virtual void Duplicate(ALandscapeGizmoActor* Gizmo); 
-	//virtual void EditorApplyTranslation(const FVector& DeltaTranslation, bool bAltDown, bool bShiftDown, bool bCtrlDown);
 
-	bool EditorCanAttachTo(const AActor* InParent, FText& OutReason) const override
-	{
-		return false;
-	}
+	bool EditorCanAttachTo(const AActor* InParent, FText& OutReason) const override	{ return false;	}
 #endif
 
 	/** 
diff --git a/Engine/Source/Runtime/Landscape/Private/LandscapeGizmoActor.cpp b/Engine/Source/Runtime/Landscape/Private/LandscapeGizmoActor.cpp
--- a/Engine/Source/Runtime/Landscape/Private/LandscapeGizmoActor.cpp
+++ b/Engine/Source/Runtime/Landscape/Private/LandscapeGizmoActor.cpp
@@ -566,42 +566,44 @@ void ALandscapeGizmoActiveActor::PostEditMove(bool bFinished)
 
 FVector ALandscapeGizmoActiveActor::SnapToLandscapeGrid(const FVector& GizmoLocation) const
 {
-	check(TargetLandscapeInfo);
-	const ALandscapeProxy* LandscapeProxy = TargetLandscapeInfo->GetLandscapeProxy();
-
-	const FTransform LToW = LandscapeProxy->LandscapeActorToWorld();
-	const FVector LandscapeSpaceLocation = LToW.InverseTransformPosition(GizmoLocation);
-	float SnapDimension;
-	switch(SnapType)
+	FVector ResultLocation = GizmoLocation;
+	if (TargetLandscapeInfo != nullptr)
 	{
-	case ELandscapeGizmoSnapType::Component:
-		SnapDimension = LandscapeProxy->ComponentSizeQuads;
-		break;
-	case ELandscapeGizmoSnapType::Texel:
-		SnapDimension = 1.0f;
-		break;
-	case ELandscapeGizmoSnapType::None:
-		SnapDimension = 0.0f; 
-		break;
-	default:
-		SnapDimension = 0.0f;
-		break;
-	}
+		const ALandscapeProxy* LandscapeProxy = TargetLandscapeInfo->GetLandscapeProxy();
 
-	FVector ResultLocation;
-	if (SnapDimension > 0.0f)
-	{
-		const FVector SnappedLandscapeSpaceLocation = LandscapeSpaceLocation.GridSnap(SnapDimension);
-		ResultLocation = LToW.TransformPosition(SnappedLandscapeSpaceLocation);	
-	}
-	else
-	{
-		ResultLocation = GizmoLocation;
-	}
+		const FTransform LToW = LandscapeProxy->LandscapeActorToWorld();
+		const FVector LandscapeSpaceLocation = LToW.InverseTransformPosition(GizmoLocation);
+		float SnapDimension;
+		switch (SnapType)
+		{
+		case ELandscapeGizmoSnapType::Component:
+			SnapDimension = LandscapeProxy->ComponentSizeQuads;
+			break;
+		case ELandscapeGizmoSnapType::Texel:
+			SnapDimension = 1.0f;
+			break;
+		case ELandscapeGizmoSnapType::None:
+			SnapDimension = 0.0f;
+			break;
+		default:
+			SnapDimension = 0.0f;
+			break;
+		}
 
-	if (bFollowTerrainHeight)
-	{
-		ResultLocation.Z = LandscapeProxy->GetHeightAtLocation(ResultLocation, EHeightfieldSource::Editor).Get(GizmoLocation.Z);
+		if (SnapDimension > 0.0f)
+		{
+			const FVector SnappedLandscapeSpaceLocation = LandscapeSpaceLocation.GridSnap(SnapDimension);
+			ResultLocation = LToW.TransformPosition(SnappedLandscapeSpaceLocation);
+		}
+		else
+		{
+			ResultLocation = GizmoLocation;
+		}
+
+		if (bFollowTerrainHeight)
+		{
+			ResultLocation.Z = LandscapeProxy->GetHeightAtLocation(ResultLocation, EHeightfieldSource::Editor).Get(GizmoLocation.Z);
+		}
 	}
 	
 	return ResultLocation;
