Fix landscape spline actors not getting hidden in PIE.

This is only an issue in editor, as only in editor does the landscape info registration stomp on the spline's landscape actor reference -- so if the guid is wrong you end up with no landscape actor.
In non-editor, it leaves the landscape actor alone.

diff --git a/Engine/Source/Runtime/Landscape/Private/LandscapeSplineActor.cpp b/Engine/Source/Runtime/Landscape/Private/LandscapeSplineActor.cpp
--- a/Engine/Source/Runtime/Landscape/Private/LandscapeSplineActor.cpp
+++ b/Engine/Source/Runtime/Landscape/Private/LandscapeSplineActor.cpp
@@ -105,6 +105,15 @@ void ALandscapeSplineActor::PostRegisterAllComponents()
 	Super::PostRegisterAllComponents();
 	if (!IsPendingKillPending())
 	{
+#if WITH_EDITORONLY_DATA
+		// note that when running PIE or when instancing a level, Landscape might change its GUID (see FixupLandscapeGuidsIfInstanced), so we need to copy it.
+		// (otherwise we would register with the wrong landscape info and our LandscapeActor reference will be stomped)
+		if (LandscapeActor)
+		{
+			LandscapeGuid = LandscapeActor->GetLandscapeGuid();
+		}
+#endif // WITH_EDITORONLY_DATA
+
 		UWorld* World = GetWorld();
 		if (LandscapeGuid.IsValid())
 		{
