Fixing static analysis warnings

diff --git a/Engine/Source/Runtime/Landscape/Private/Landscape.cpp b/Engine/Source/Runtime/Landscape/Private/Landscape.cpp
--- a/Engine/Source/Runtime/Landscape/Private/Landscape.cpp
+++ b/Engine/Source/Runtime/Landscape/Private/Landscape.cpp
@@ -3065,12 +3065,10 @@ void ALandscapeProxy::Serialize(FArchive& Ar)
 
 	if (Ar.IsLoading())
 	{
-		// Fixup Nanite meshes which were using the wrong material and didn't have proper UVs :
-		if ((Ar.CustomVer(FFortniteReleaseBranchCustomObjectVersion::GUID) < FFortniteReleaseBranchCustomObjectVersion::FixupNaniteLandscapeMeshes)
-			// Remove cooked collision data from Nanite landscape meshes, since collisions are handled by ULandscapeHeightfieldCollisionComponent :		
-			|| (Ar.CustomVer(FFortniteReleaseBranchCustomObjectVersion::GUID) < FFortniteReleaseBranchCustomObjectVersion::RemoveUselessLandscapeMeshesCookedCollisionData)
-			// Fix the names of the generated Nanite landcape UStaticMesh so that it's unique in a given package : 
-			|| (Ar.CustomVer(FFortniteReleaseBranchCustomObjectVersion::GUID) < FFortniteReleaseBranchCustomObjectVersion::FixNaniteLandscapeMeshNames))
+		// Fixup Nanite meshes which were using the wrong material and didn't have proper UVs : FFortniteReleaseBranchCustomObjectVersion::FixupNaniteLandscapeMeshes
+		// Remove cooked collision data from Nanite landscape meshes, since collisions are handled by ULandscapeHeightfieldCollisionComponent : FFortniteReleaseBranchCustomObjectVersion::RemoveUselessLandscapeMeshesCookedCollisionData
+		// Fix the names of the generated Nanite landcape UStaticMesh so that it's unique in a given package : FFortniteReleaseBranchCustomObjectVersion::FixNaniteLandscapeMeshNames
+		if (Ar.CustomVer(FFortniteReleaseBranchCustomObjectVersion::GUID) < FFortniteReleaseBranchCustomObjectVersion::FixNaniteLandscapeMeshNames)
 		{
 			// This will force the Nanite meshes to be properly regenerated during the next save :
 			InvalidateNaniteRepresentation(/* bCheckContentId = */ false);
