LandscapeStreamingProxies' FixupSharedData now removes Materials being the same as the parent LandscapeActor.
GetSharedProperties now only copies material when explicitely specified.
Moves ValidateCombinationMaterial calling code to PostRegisterAllComponents since GetLandscapeMaterial will need the ALandscape link to be valid.
Adds assert to ALandscapeStreamingProxy::GetLandscapeMaterial to enforce the method is not used during postload (to be sure we can always get the parent actor).
Fixes Landscape material assignment in WP.

diff --git a/Engine/Source/Editor/UnrealEd/Private/Commandlets/WorldPartitionConvertCommandlet.cpp b/Engine/Source/Editor/UnrealEd/Private/Commandlets/WorldPartitionConvertCommandlet.cpp
--- a/Engine/Source/Editor/UnrealEd/Private/Commandlets/WorldPartitionConvertCommandlet.cpp
+++ b/Engine/Source/Editor/UnrealEd/Private/Commandlets/WorldPartitionConvertCommandlet.cpp
@@ -987,7 +987,9 @@ int32 UWorldPartitionConvertCommandlet::Main(const FString& Params)
 				FTransform LandscapeTransform = FirstProxy->LandscapeActorToWorld();
 				ALandscape* NewLandscape = MainWorld->SpawnActor<ALandscape>(ALandscape::StaticClass(), LandscapeTransform, SpawnParams);
 
-				NewLandscape->GetSharedProperties(FirstProxy);
+				// We want to clone this StreamingProxy to create the parent ALandscape, including materials.
+				const bool bIncludeLandscapeMaterial = true;
+				NewLandscape->GetSharedProperties(FirstProxy, bIncludeLandscapeMaterial);
 
 				LandscapeInfo->RegisterActor(NewLandscape);
 			}
diff --git a/Engine/Source/Editor/WorldBrowser/Private/Tiles/WorldTileModel.cpp b/Engine/Source/Editor/WorldBrowser/Private/Tiles/WorldTileModel.cpp
--- a/Engine/Source/Editor/WorldBrowser/Private/Tiles/WorldTileModel.cpp
+++ b/Engine/Source/Editor/WorldBrowser/Private/Tiles/WorldTileModel.cpp
@@ -979,8 +979,9 @@ bool FWorldTileModel::CreateAdjacentLandscapeProxy(ALandscapeProxy* SourceLandsc
 	ALandscapeProxy* AdjacentLandscape = ImportLandscapeTile(ImportSettings);
 	if (AdjacentLandscape)
 	{
-		// Copy source landscape properties 
-		AdjacentLandscape->GetSharedProperties(SourceLandscape);
+		// Copy source landscape properties, including materials since this proxy could exist by itself, without a parent ALandscape.
+		const bool bIncludeLandscapeMaterial = true;
+		AdjacentLandscape->GetSharedProperties(SourceLandscape, bIncludeLandscapeMaterial);
 		
 		// Refresh level model bounding box
 		FBox AdjacentLandscapeBounds = AdjacentLandscape->GetComponentsBoundingBox(true);
diff --git a/Engine/Source/Runtime/Landscape/Classes/LandscapeProxy.h b/Engine/Source/Runtime/Landscape/Classes/LandscapeProxy.h
--- a/Engine/Source/Runtime/Landscape/Classes/LandscapeProxy.h
+++ b/Engine/Source/Runtime/Landscape/Classes/LandscapeProxy.h
@@ -996,7 +996,7 @@ public:
 	LANDSCAPE_API TArray<float> GetLODScreenSizeArray() const;
 
 	// Copy properties from parent Landscape actor
-	LANDSCAPE_API void GetSharedProperties(ALandscapeProxy* Landscape);
+	LANDSCAPE_API void GetSharedProperties(ALandscapeProxy* Landscape, const bool bIncludeLandscapeMaterial);
 
 #if WITH_EDITOR
 	/* Serialize all hashes/guids that record the current state of this proxy */
@@ -1059,7 +1059,7 @@ public:
 	LANDSCAPE_API void ChangedPhysMaterial();
 
 	// Assign only mismatching data and mark proxy package dirty
-	LANDSCAPE_API void FixupSharedData(ALandscape* Landscape);
+	LANDSCAPE_API virtual void FixupSharedData(ALandscape* Landscape);
 
 	/** Set landscape absolute location in section space */
 	LANDSCAPE_API void SetAbsoluteSectionBase(FIntPoint SectionOffset);
diff --git a/Engine/Source/Runtime/Landscape/Classes/LandscapeStreamingProxy.h b/Engine/Source/Runtime/Landscape/Classes/LandscapeStreamingProxy.h
--- a/Engine/Source/Runtime/Landscape/Classes/LandscapeStreamingProxy.h
+++ b/Engine/Source/Runtime/Landscape/Classes/LandscapeStreamingProxy.h
@@ -58,6 +58,12 @@ public:
 #endif
 	//~ End ALandscapeBase Interface
 
+	//~ Begin ALandscapeProxy Interface
+#if WITH_EDITOR
+	virtual void FixupSharedData(ALandscape* Landscape) override;
+#endif
+	//~ End ALandscapeProxy Interface
+
 	// Check input Landscape actor is match for this LandscapeProxy (by GUID)
 	bool IsValidLandscapeActor(ALandscape* Landscape);
 };
diff --git a/Engine/Source/Runtime/Landscape/Private/Landscape.cpp b/Engine/Source/Runtime/Landscape/Private/Landscape.cpp
--- a/Engine/Source/Runtime/Landscape/Private/Landscape.cpp
+++ b/Engine/Source/Runtime/Landscape/Private/Landscape.cpp
@@ -92,10 +92,12 @@ Landscape.cpp: Terrain rendering
 #include "LandscapeHeightfieldCollisionComponent.h"
 #include "UObject/EditorObjectVersion.h"
 #include "Algo/BinarySearch.h"
+#include "UObject/UObjectThreadContext.h"
 #include "WorldPartition/WorldPartition.h"
 #include "WorldPartition/WorldPartitionHelpers.h"
 #include "WorldPartition/WorldPartitionHandle.h"
 #include "WorldPartition/Landscape/LandscapeActorDesc.h"
+
 #if WITH_EDITOR
 #include "Rendering/StaticLightingSystemInterface.h"
 #include "Misc/ScopedSlowTask.h"
@@ -2322,6 +2324,7 @@ void ULandscapeComponent::DeleteLayerAllocation(const FGuid& InEditLayerGuid, in
 void ALandscapeProxy::PostRegisterAllComponents()
 {
 	Super::PostRegisterAllComponents();
+
 	ULandscapeInfo* LandscapeInfo = nullptr;
 	if (!IsPendingKillPending())
 	{
@@ -2381,6 +2384,55 @@ void ALandscapeProxy::PostRegisterAllComponents()
 	{
 		LandscapeInfo->FixupProxiesTransform();
 	}
+
+	bool bFixedUpInvalidMaterialInstances = false;
+
+	for (ULandscapeComponent* Comp : LandscapeComponents)
+	{
+		if (Comp == nullptr)
+		{
+			continue;
+		}
+
+		// Validate the layer combination and store it in the MaterialInstanceConstantMap
+		UMaterialInstance* MaterialInstance = Comp->GetMaterialInstance(0, false);
+
+		if (MaterialInstance == nullptr)
+		{
+			continue;
+		}
+
+		UMaterialInstanceConstant* CombinationMaterialInstance = Cast<UMaterialInstanceConstant>(MaterialInstance->Parent);
+		// Only validate if uncooked and in the editor/commandlet mode (we cannot re-build material instance constants if this is not the case : see UMaterialInstance::CacheResourceShadersForRendering, which is only called if FApp::CanEverRender() returns true) 
+		if (!Comp->GetOutermost()->HasAnyPackageFlags(PKG_FilterEditorOnly) && (GIsEditor && FApp::CanEverRender()))
+		{
+			if (Comp->ValidateCombinationMaterial(CombinationMaterialInstance))
+			{
+				MaterialInstanceConstantMap.Add(*ULandscapeComponent::GetLayerAllocationKey(Comp->GetWeightmapLayerAllocations(), CombinationMaterialInstance->Parent), CombinationMaterialInstance);
+			}
+			else
+			{
+				// There was a problem with the loaded material : it doesn't match the expected material combination, we need to regenerate the material instances : 
+				Comp->UpdateMaterialInstances();
+				bFixedUpInvalidMaterialInstances = true;
+			}
+		}
+		else if (CombinationMaterialInstance)
+		{
+			// Skip ValidateCombinationMaterial
+			MaterialInstanceConstantMap.Add(*ULandscapeComponent::GetLayerAllocationKey(Comp->GetWeightmapLayerAllocations(), CombinationMaterialInstance->Parent), CombinationMaterialInstance);
+		}
+	}
+
+	if (bFixedUpInvalidMaterialInstances)
+	{
+		FFormatNamedArguments Arguments;
+		Arguments.Add(TEXT("LandscapeName"), FText::FromString(GetPathName()));
+		Arguments.Add(TEXT("ProxyPackage"), FText::FromString(GetOutermost()->GetName()));
+		FMessageLog("MapCheck").Info()
+			->AddToken(FTextToken::Create(FText::Format(LOCTEXT("MapCheck_Message_FixedUpInvalidLandscapeMaterialInstances", "{LandscapeName} : Fixed up invalid landscape material instances. Please re-save {ProxyPackage}."), Arguments)))
+			->AddToken(FMapErrorToken::Create(FMapErrors::FixedUpInvalidLandscapeMaterialInstances));
+	}
 #endif // WITH_EDITOR
 }
 
@@ -3499,49 +3551,6 @@ void ALandscapeProxy::PostLoad()
 		EditorCachedLayerInfos_DEPRECATED.Empty();
 	}
 
-	bool bFixedUpInvalidMaterialInstances = false;
-
-	for (ULandscapeComponent* Comp : LandscapeComponents)
-	{
-		if (Comp)
-		{
-			// Validate the layer combination and store it in the MaterialInstanceConstantMap
-			if (UMaterialInstance* MaterialInstance = Comp->GetMaterialInstance(0, false))
-			{
-				UMaterialInstanceConstant* CombinationMaterialInstance = Cast<UMaterialInstanceConstant>(MaterialInstance->Parent);
-				// Only validate if uncooked and in the editor/commandlet mode (we cannot re-build material instance constants if this is not the case : see UMaterialInstance::CacheResourceShadersForRendering, which is only called if FApp::CanEverRender() returns true) 
-				if (!Comp->GetOutermost()->HasAnyPackageFlags(PKG_FilterEditorOnly) && (GIsEditor && FApp::CanEverRender()))
-				{
-					if (Comp->ValidateCombinationMaterial(CombinationMaterialInstance))
-					{
-						MaterialInstanceConstantMap.Add(*ULandscapeComponent::GetLayerAllocationKey(Comp->GetWeightmapLayerAllocations(), CombinationMaterialInstance->Parent), CombinationMaterialInstance);
-					}
-					else
-					{
-						// There was a problem with the loaded material : it doesn't match the expected material combination, we need to regenerate the material instances : 
-						Comp->UpdateMaterialInstances();
-						bFixedUpInvalidMaterialInstances = true;
-					}
-				}
-				else if (CombinationMaterialInstance)
-				{
-					// Skip ValidateCombinationMaterial
-					MaterialInstanceConstantMap.Add(*ULandscapeComponent::GetLayerAllocationKey(Comp->GetWeightmapLayerAllocations(), CombinationMaterialInstance->Parent), CombinationMaterialInstance);
-				}
-			}
-		}
-	}
-
-	if (bFixedUpInvalidMaterialInstances)
-	{
-		FFormatNamedArguments Arguments;
-		Arguments.Add(TEXT("LandscapeName"), FText::FromString(GetPathName()));
-		Arguments.Add(TEXT("ProxyPackage"), FText::FromString(GetOutermost()->GetName()));
-		FMessageLog("MapCheck").Info()
-			->AddToken(FTextToken::Create(FText::Format(LOCTEXT("MapCheck_Message_FixedUpInvalidLandscapeMaterialInstances", "{LandscapeName} : Fixed up invalid landscape material instances. Please re-save {ProxyPackage}."), Arguments)))
-			->AddToken(FMapErrorToken::Create(FMapErrors::FixedUpInvalidLandscapeMaterialInstances));
-	}
-
 	UWorld* World = GetWorld();
 
 	// track feature level change to flush grass cache
@@ -3604,7 +3613,7 @@ void ALandscapeProxy::Destroyed()
 }
 #endif // WITH_EDITOR
 
-void ALandscapeProxy::GetSharedProperties(ALandscapeProxy* Landscape)
+void ALandscapeProxy::GetSharedProperties(ALandscapeProxy* Landscape, const bool bIncludeLandscapeMaterial)
 {
 	if (Landscape)
 	{
@@ -3649,12 +3658,12 @@ void ALandscapeProxy::GetSharedProperties(ALandscapeProxy* Landscape)
 		VirtualTextureNumLods = Landscape->VirtualTextureNumLods;
 		VirtualTextureRenderPassType = Landscape->VirtualTextureRenderPassType;
 
-		if (!LandscapeMaterial)
+		if (bIncludeLandscapeMaterial && !LandscapeMaterial)
 		{
 			LandscapeMaterial = Landscape->LandscapeMaterial;
 			PerLODOverrideMaterials = Landscape->PerLODOverrideMaterials;
 		}
-		if (!LandscapeHoleMaterial)
+		if (bIncludeLandscapeMaterial && !LandscapeHoleMaterial)
 		{
 			LandscapeHoleMaterial = Landscape->LandscapeHoleMaterial;
 		}
@@ -3851,6 +3860,8 @@ UMaterialInterface* ALandscapeStreamingProxy::GetLandscapeMaterial(int8 InLODInd
 		return LandscapeMaterial;
 	}
 
+	checkf(!FUObjectThreadContext::Get().IsRoutingPostLoad, TEXT("This method should not be called during PostLoad since the LandscapeActor may not be fully unserialized yet."));
+
 	if (const ALandscape* Landscape = GetLandscapeActor())
 	{
 		return Landscape->GetLandscapeMaterial(InLODIndex);
@@ -4400,7 +4411,7 @@ void ULandscapeInfo::RegisterActor(ALandscapeProxy* Proxy, bool bMapCheck, bool
 			{
 				if (ALandscapeStreamingProxy* StreamingProxy = StreamingProxyPtr.Get())
 				{
-					StreamingProxy->SetLandscapeActor(LandscapeActor.Get());
+					StreamingProxy->SetLandscapeActor(Landscape);
 #if WITH_EDITOR
 					StreamingProxy->FixupSharedData(Landscape);
 #endif // WITH_EDITOR
diff --git a/Engine/Source/Runtime/Landscape/Private/LandscapeConfigHelper.cpp b/Engine/Source/Runtime/Landscape/Private/LandscapeConfigHelper.cpp
--- a/Engine/Source/Runtime/Landscape/Private/LandscapeConfigHelper.cpp
+++ b/Engine/Source/Runtime/Landscape/Private/LandscapeConfigHelper.cpp
@@ -69,8 +69,10 @@ ALandscapeProxy* FLandscapeConfigHelper::FindOrAddLandscapeStreamingProxy(UActor
 		const FIntPoint CellLocation(static_cast<int32>(InCellCoord.X) * Landscape->GridSize, static_cast<int32>(InCellCoord.Y) * Landscape->GridSize);
 
 		ALandscapeProxy* LandscapeProxy = CastChecked<ALandscapeProxy>(PartitionActor);
-		// copy shared properties to this new proxy
-		LandscapeProxy->GetSharedProperties(Landscape);
+		// copy shared properties to this new proxy. Not including materials since in this architecture, we share the material with the parent ALandscape.
+		const bool bIncludeLandscapeMaterial = false;
+		LandscapeProxy->GetSharedProperties(Landscape, bIncludeLandscapeMaterial);
+
 		const FVector ProxyLocation = Landscape->GetActorLocation() + FVector(CellLocation.X * Landscape->GetActorRelativeScale3D().X, CellLocation.Y * Landscape->GetActorRelativeScale3D().Y, 0.0f);
 
 		LandscapeProxy->CreateLandscapeInfo();
@@ -484,8 +486,10 @@ ULandscapeInfo* FLandscapeConfigHelper::ChangeConfiguration(ULandscapeInfo* InLa
 			
 	FActorSpawnParameters SpawnParams;
 	ALandscape* NewLandscape = World->SpawnActor<ALandscape>(OldLandscape->GetActorLocation() + ActorOffset, OldLandscape->GetActorRotation(), SpawnParams);
-	
-	NewLandscape->GetSharedProperties(OldLandscape);
+	// We'll clone the previous actor properties to the new one, including materials.
+	const bool bIncludeLandscapeMaterial = true;
+
+	NewLandscape->GetSharedProperties(OldLandscape, bIncludeLandscapeMaterial);
 	NewLandscape->SetLandscapeGuid(FGuid::NewGuid());
 	NewLandscape->GridSize = InNewConfig.GetGridSizeQuads();
 	NewLandscape->ComponentSizeQuads = InNewConfig.GetComponentSizeQuads();
diff --git a/Engine/Source/Runtime/Landscape/Private/LandscapeEdit.cpp b/Engine/Source/Runtime/Landscape/Private/LandscapeEdit.cpp
--- a/Engine/Source/Runtime/Landscape/Private/LandscapeEdit.cpp
+++ b/Engine/Source/Runtime/Landscape/Private/LandscapeEdit.cpp
@@ -254,6 +254,12 @@ bool ULandscapeComponent::ValidateCombinationMaterial(UMaterialInstanceConstant*
 		return false;
 	}
 
+	// Check if the currently used Landscape Material changed, since a ALandscapeStreamingProxy can use its ALandscape parent material.
+	if (GetLandscapeMaterial() != InCombinationMaterial->Parent)
+	{
+		return false;
+	}
+
 	const TArray<FStaticTerrainLayerWeightParameter>& TerrainLayerWeightParameters = InCombinationMaterial->GetEditorOnlyStaticParameters().TerrainLayerWeightParameters;
 	const TArray<FWeightmapLayerAllocationInfo>& ComponentWeightmapAllocations = GetWeightmapLayerAllocations();
 
@@ -4874,8 +4880,10 @@ ALandscapeProxy* ULandscapeInfo::MoveComponentsToLevel(const TArray<ULandscapeCo
 		SpawnParams.OverrideLevel = TargetLevel;
 		LandscapeProxy = TargetLevel->GetWorld()->SpawnActor<ALandscapeStreamingProxy>(SpawnParams);
 
-		// copy shared properties to this new proxy
-		LandscapeProxy->GetSharedProperties(Landscape);
+		// copy shared properties to this new proxy, not including material since we'll use the ALandscape material.
+		const bool bIncludeLandscapeMaterial = false;
+		
+		LandscapeProxy->GetSharedProperties(Landscape, bIncludeLandscapeMaterial);
 		LandscapeProxy->CreateLandscapeInfo();
 		LandscapeProxy->SetActorLabel(LandscapeProxy->GetName());
 		bSetPositionAndOffset = true;
@@ -6011,7 +6019,10 @@ void ALandscape::PostEditChangeProperty(FPropertyChangedEvent& PropertyChangedEv
 			{
 				if (ALandscapeProxy* Proxy = ProxyPtr.Get())
 				{
-					Proxy->GetSharedProperties(this);
+					// This should not include material since the ALandscapeStreamingProxy is either using the ALandscape material, or has a specific override.
+					const bool bIncludeLandscapeMaterial = false;
+					
+					Proxy->GetSharedProperties(this, bIncludeLandscapeMaterial);
 					Proxy->PostEditChangeProperty(PropertyChangedEvent);
 				}
 			}
@@ -6050,7 +6061,7 @@ void ALandscape::PostEditChangeProperty(FPropertyChangedEvent& PropertyChangedEv
 
 			if (ChangedMaterial)
 			{
-				UpdateAllComponentMaterialInstances();
+				Info->UpdateAllComponentMaterialInstances();
 
 				UWorld* World = GetWorld();
 
@@ -7460,6 +7471,34 @@ bool ALandscapeProxy::LandscapeExportHeightmapToRenderTarget(UTextureRenderTarge
 }
 
 #if WITH_EDITOR
+
+void ALandscapeStreamingProxy::FixupSharedData(ALandscape* InLandscape)
+{
+	bool bUpdated = false;
+
+	if (InLandscape != nullptr)
+	{
+		if (LandscapeMaterial == InLandscape->GetLandscapeMaterial())
+		{
+			LandscapeMaterial = nullptr;
+			bUpdated = true;
+		}
+
+		if (LandscapeHoleMaterial == InLandscape->GetLandscapeHoleMaterial())
+		{
+			LandscapeHoleMaterial = nullptr;
+			bUpdated = true;
+		}
+	}
+
+	if (bUpdated)
+	{
+		MarkPackageDirty();
+	}
+
+	Super::FixupSharedData(InLandscape);
+}
+
 bool ALandscapeProxy::LandscapeImportWeightmapFromRenderTarget(UTextureRenderTarget2D* InRenderTarget, FName InLayerName)
 {
 	ALandscape* Landscape = GetLandscapeActor();
