Removed rogue ULandscapeNaniteComponent from ALandscapeProxy's OwnedComponents upon copy/pasting (the NaniteComponent property is marked NonPIEDuplicate+TextExportTransient but the component itself cannot be, so it ends up being duplicated all the same and there'll be a copy of it in the duplicated actor's OwnedComponents so we need this fixup to ensure we regenerate the Nanite data upon copy/paste

diff --git a/Engine/Source/Runtime/Landscape/Private/LandscapeEdit.cpp b/Engine/Source/Runtime/Landscape/Private/LandscapeEdit.cpp
--- a/Engine/Source/Runtime/Landscape/Private/LandscapeEdit.cpp
+++ b/Engine/Source/Runtime/Landscape/Private/LandscapeEdit.cpp
@@ -4544,6 +4544,15 @@ void ALandscape::PostEditImport()
 		}
 	}
 
+	// Even if the component's UPROPERTY is TextExportTransient/NonPIEDuplicate, it still gets duplicated and added to the OwnedComponents so we need to remove it after duplicating the actor : 
+	check(NaniteComponent == nullptr);
+	TInlineComponentArray<ULandscapeNaniteComponent*> OwnedNaniteComponents;
+	GetComponents<ULandscapeNaniteComponent>(OwnedNaniteComponents, /*bIncludeFromChildActors = */false);
+	for (ULandscapeNaniteComponent* OwnedNaniteComponent : OwnedNaniteComponents)
+	{
+		OwnedNaniteComponent->DestroyComponent();
+	}
+
 	// Some edit layers could be affected by BP brushes, which might need to be updated when the landscape is transformed :
 	RequestLayersContentUpdate(ELandscapeLayerUpdateMode::Update_All);
 
