Removed the deprecated baked landscape GI texture system

diff --git a/Engine/Source/Editor/LandscapeEditor/Private/LandscapeEdModeTools.h b/Engine/Source/Editor/LandscapeEditor/Private/LandscapeEdModeTools.h
--- a/Engine/Source/Editor/LandscapeEditor/Private/LandscapeEdModeTools.h
+++ b/Engine/Source/Editor/LandscapeEditor/Private/LandscapeEdModeTools.h
@@ -1530,9 +1530,6 @@ public:
 			}
 			TimeSinceLastInteractorMove += DeltaTime;
 
-			// Prevent landscape from baking textures while tool stroke is active
-			EdMode->CurrentToolTarget.LandscapeInfo->PostponeTextureBaking();
-
 			if (ShouldUpdateEditingLayer())
 			{
 				ALandscape* Landscape = this->EdMode->CurrentToolTarget.LandscapeInfo->LandscapeActor.Get();
diff --git a/Engine/Source/Editor/LandscapeEditor/Private/LandscapeEditorModule.cpp b/Engine/Source/Editor/LandscapeEditor/Private/LandscapeEditorModule.cpp
--- a/Engine/Source/Editor/LandscapeEditor/Private/LandscapeEditorModule.cpp
+++ b/Engine/Source/Editor/LandscapeEditor/Private/LandscapeEditorModule.cpp
@@ -132,9 +132,6 @@ public:
 		{
 			FToolMenuSection& Section = BuildMenu->FindOrAddSection("LevelEditorLandscape");
 
-			FUIAction ActionBakeTextures(FExecuteAction::CreateStatic(&BuildGITextures), FCanExecuteAction());
-			Section.AddMenuEntry(NAME_None, LOCTEXT("BuildGITexturesOnly", "Build GI Textures Only"),LOCTEXT("BuildGIBakedTextures ", "Build GI baked base color textures"), TAttribute<FSlateIcon>(), ActionBakeTextures,EUserInterfaceActionType::Button);
-			
 			FUIAction ActionBuildGrassMaps(FExecuteAction::CreateStatic(&BuildGrassMaps), FCanExecuteAction());
 			Section.AddMenuEntry(TEXT("BuildGrassMapsOnly"), LOCTEXT("BuildGrassMapsOnly", "Build Grass Maps Only"), LOCTEXT("BuildLandscapeGrassMaps", "Build landscape grass maps"), TAttribute<FSlateIcon>(), ActionBuildGrassMaps, EUserInterfaceActionType::Button);
 
@@ -242,17 +239,6 @@ public:
 		);
 	}
 
-	static void BuildGITextures()
-	{
-		if (UWorld* World = GEditor->GetEditorWorldContext().World())
-		{
-			if (ULandscapeSubsystem* LandscapeSubsystem = World->GetSubsystem<ULandscapeSubsystem>())
-			{
-				LandscapeSubsystem->BuildGIBakedTextures();
-			}
-		}
-	}
-
 	static void BuildGrassMaps()
 	{
 		if (UWorld* World = GEditor->GetEditorWorldContext().World())
diff --git a/Engine/Source/Editor/LevelEditor/Private/LevelEditorActions.cpp b/Engine/Source/Editor/LevelEditor/Private/LevelEditorActions.cpp
--- a/Engine/Source/Editor/LevelEditor/Private/LevelEditorActions.cpp
+++ b/Engine/Source/Editor/LevelEditor/Private/LevelEditorActions.cpp
@@ -3530,7 +3530,7 @@ void FLevelEditorCommands::RegisterCommands()
 	UI_COMMAND( BuildLandscapeSplineMeshes, "Build Landscape Spline Meshes", "Builds landscape spline meshes for the current world", EUserInterfaceActionType::Button, FInputChord());
 	UI_COMMAND( BuildTextureStreamingOnly, "Build Texture Streaming", "Build texture streaming data", EUserInterfaceActionType::Button, FInputChord() );
 	UI_COMMAND( BuildVirtualTextureOnly, "Build Virtual Textures", "Build runtime virtual texture low mips streaming data", EUserInterfaceActionType::Button, FInputChord());
-	UI_COMMAND( BuildAllLandscape, "Build Landscape", "Build all data related to Landscape (grass maps, GI textures, physical material, Nanite, dirty height and weight maps)", EUserInterfaceActionType::Button, FInputChord());
+	UI_COMMAND( BuildAllLandscape, "Build Landscape", "Build all data related to landscape (grass maps, physical material, Nanite, dirty height and weight maps)", EUserInterfaceActionType::Button, FInputChord());
 
 	UI_COMMAND( LightingQuality_Production, "Production", "Sets precomputed lighting quality to highest possible quality (slowest computation time.)", EUserInterfaceActionType::RadioButton, FInputChord() );
 	UI_COMMAND( LightingQuality_High, "High", "Sets precomputed lighting quality to high quality", EUserInterfaceActionType::RadioButton, FInputChord() );
diff --git a/Engine/Source/Runtime/Engine/Public/PrimitiveSceneProxy.h b/Engine/Source/Runtime/Engine/Public/PrimitiveSceneProxy.h
--- a/Engine/Source/Runtime/Engine/Public/PrimitiveSceneProxy.h
+++ b/Engine/Source/Runtime/Engine/Public/PrimitiveSceneProxy.h
@@ -383,10 +383,9 @@ public:
 
 	virtual bool StaticMeshHasPendingStreaming() const { return false; }
 
-	virtual void GetHeightfieldRepresentation(UTexture2D*& OutHeightmapTexture, UTexture2D*& OutDiffuseColorTexture, UTexture2D*& OutVisibilityTexture, FHeightfieldComponentDescription& OutDescription) const
+	virtual void GetHeightfieldRepresentation(UTexture2D*& OutHeightmapTexture, UTexture2D*& OutVisibilityTexture, FHeightfieldComponentDescription& OutDescription) const
 	{
 		OutHeightmapTexture = nullptr;
-		OutDiffuseColorTexture = nullptr;
 		OutVisibilityTexture = nullptr;
 	}
 
diff --git a/Engine/Source/Runtime/Landscape/Classes/Landscape.h b/Engine/Source/Runtime/Landscape/Classes/Landscape.h
--- a/Engine/Source/Runtime/Landscape/Classes/Landscape.h
+++ b/Engine/Source/Runtime/Landscape/Classes/Landscape.h
@@ -557,7 +557,6 @@ private:
 	TSharedPtr<FLandscapeNotification> WaitingForTexturesNotification;
 	TSharedPtr<FLandscapeNotification> WaitingForBrushesNotification;
 	TSharedPtr<FLandscapeNotification> InvalidShadingModelNotification;
-	TSharedPtr<FLandscapeNotification> TextureBakingNotification;
 	TSharedPtr<FLandscapeNotification> GrassRenderingNotification;
 
 	// Represent all the resolved paint layer, from all layers blended together (size of the landscape x material layer count)
diff --git a/Engine/Source/Runtime/Landscape/Classes/LandscapeComponent.h b/Engine/Source/Runtime/Landscape/Classes/LandscapeComponent.h
--- a/Engine/Source/Runtime/Landscape/Classes/LandscapeComponent.h
+++ b/Engine/Source/Runtime/Landscape/Classes/LandscapeComponent.h
@@ -548,21 +548,19 @@ public:
 	int32 LODBias;
 
 	UPROPERTY()
+	// TODO [jonathan.bard] : remove unused : 
 	FGuid StateId;
 
-	/** The Material Guid that used when baking, to detect material recompilations */
-	UPROPERTY()
+	UE_DEPRECATED(5.3, "BakedTextureMaterialGuid is officially deprecated now and nothing updates it anymore")
 	FGuid BakedTextureMaterialGuid;
 
-	/** The Material Guid that last saved, to notify manual build operation to bake textures */
-	UPROPERTY(Transient)
+	UE_DEPRECATED(5.3, "LastBakedTextureMaterialGuid is officially deprecated now and nothing updates it anymore")
 	FGuid LastBakedTextureMaterialGuid;
 
-	/** Pre-baked Base Color texture for use by distance field GI */
-	UPROPERTY(VisibleAnywhere, BlueprintReadOnly, Category = BakedTextures)
+#if WITH_EDITORONLY_DATA
+	UE_DEPRECATED(5.3, "GIBakedBaseColorTexture is officially deprecated now and nothing updates it anymore")
 	TObjectPtr<UTexture2D> GIBakedBaseColorTexture;
 
-#if WITH_EDITORONLY_DATA
 	/**	Legacy irrelevant lights */
 	UPROPERTY()
 	TArray<FGuid> IrrelevantLights_DEPRECATED;
diff --git a/Engine/Source/Runtime/Landscape/Classes/LandscapeInfo.h b/Engine/Source/Runtime/Landscape/Classes/LandscapeInfo.h
--- a/Engine/Source/Runtime/Landscape/Classes/LandscapeInfo.h
+++ b/Engine/Source/Runtime/Landscape/Classes/LandscapeInfo.h
@@ -303,8 +303,8 @@ public:
 
 	LANDSCAPE_API void RemoveXYOffsets();
 
-	/** Postpones landscape textures baking, usually used during landscape painting to avoid hitches */
-	LANDSCAPE_API void PostponeTextureBaking();
+	UE_DEPRECATED(5.3, "Texture Baking is officially deprecated now and nothing updates it anymore")
+	LANDSCAPE_API void PostponeTextureBaking() {}
 
 	/** Will tell if the landscape actor can have some content related to the layer system */
 	LANDSCAPE_API bool CanHaveLayersContent() const;
diff --git a/Engine/Source/Runtime/Landscape/Classes/LandscapeProxy.h b/Engine/Source/Runtime/Landscape/Classes/LandscapeProxy.h
--- a/Engine/Source/Runtime/Landscape/Classes/LandscapeProxy.h
+++ b/Engine/Source/Runtime/Landscape/Classes/LandscapeProxy.h
@@ -927,26 +927,18 @@ public:
 	/** Render grass maps for the specified components */
 	void RenderGrassMaps(const TArray<ULandscapeComponent*>& LandscapeComponents, const TArray<ULandscapeGrassType*>& GrassTypes);
 
-	/** Stores Components and their latest BakedTextureMaterial hash*/
+	UE_DEPRECATED(5.3, "FGIBakedTextureState is officially deprecated now and nothing updates it anymore")
 	struct FGIBakedTextureState
 	{
 		FGuid CombinedStateId;
 		TArray<ULandscapeComponent*> Components;
 	};
 		
-	/** Update any GI baked textures from the landscape as necessary */
-	void UpdateGIBakedTextures(bool bBakeAllGITextures = false);
-	/** Update the landscape GI baked textures without deferring*/
-	void UpdateGIBakedTextureData(bool bInShouldMarkDirty = false);
-	/** Update the status of landscape GI baked textures*/
-	void UpdateGIBakedTextureStatus(bool* bOutGenerateLandscapeGIData, TMap<UTexture2D*, FGIBakedTextureState>* OutComponentsNeedBakingByHeightmap, int32* OutdatedComponentsCount=nullptr) const;
-
 	/** Update the landscape physical material render tasks */
 	void UpdatePhysicalMaterialTasks(bool bInShouldMarkDirty = false);
-
 	void UpdatePhysicalMaterialTasksStatus(TSet<ULandscapeComponent*>* OutdatedComponents, int32* OutdatedComponentsCount) const;
 
-	/** Frame counter to count down to the next time we check to update baked textures, so we don't check every frame */
+	UE_DEPRECATED(5.3, "UpdateBakedTexturesCountdown is officially deprecated now and nothing updates it anymore")
 	int32 UpdateBakedTexturesCountdown;
 
 	/** Editor notification when changing feature level */
@@ -1008,8 +1000,10 @@ public:
 
 	LANDSCAPE_API int32 GetOutdatedGrassMapCount() const;
 	LANDSCAPE_API void BuildGrassMaps(struct FScopedSlowTask* InSlowTask = nullptr);
-	LANDSCAPE_API void BuildGIBakedTextures(struct FScopedSlowTask* InSlowTask = nullptr);
-	LANDSCAPE_API int32 GetOutdatedGIBakedTextureComponentsCount() const;
+	UE_DEPRECATED(5.3, "BuildGIBakedTextures is officially deprecated now")
+	LANDSCAPE_API void BuildGIBakedTextures(struct FScopedSlowTask* InSlowTask = nullptr) {}
+	UE_DEPRECATED(5.3, "GetOutdatedGIBakedTextureComponentsCount is officially deprecated now returns 0")
+	LANDSCAPE_API int32 GetOutdatedGIBakedTextureComponentsCount() const { return 0; }
 	LANDSCAPE_API void BuildPhysicalMaterial(struct FScopedSlowTask* InSlowTask = nullptr);
 	LANDSCAPE_API int32 GetOudatedPhysicalMaterialComponentsCount() const;
 	LANDSCAPE_API virtual void CreateSplineComponent() override;
@@ -1246,17 +1240,13 @@ public:
 
 	/* For the grassmap rendering notification */
 	int32 NumComponentsNeedingGrassMapRender;
-	UE_DEPRECATED(5.1, "No longer used for the grassmap rendering notification.")
-	LANDSCAPE_API static int32 TotalComponentsNeedingGrassMapRender;
 
 	/* To throttle texture streaming when we're trying to render a grassmap */
 	int32 NumTexturesToStreamForVisibleGrassMapRender;
 	LANDSCAPE_API static int32 TotalTexturesToStreamForVisibleGrassMapRender;
 
-	/* For the texture baking notification */
+	UE_DEPRECATED(5.3, "NumComponentsNeedingTextureBaking is officially deprecated now and nothing updates it anymore")
 	int32 NumComponentsNeedingTextureBaking;
-	UE_DEPRECATED(5.1, "No longer used for the texture baking notification.")
-	LANDSCAPE_API static int32 TotalComponentsNeedingTextureBaking;
 
 	/** remove an overlapping component. Called from MapCheck. */
 	LANDSCAPE_API void RemoveOverlappingComponent(ULandscapeComponent* Component);
@@ -1378,16 +1368,19 @@ private:
 /**
  * Helper class used to Build or monitor Landscape GI Textures
  */
-class LANDSCAPE_API FLandscapeGIBakedTextureBuilder
+class LANDSCAPE_API UE_DEPRECATED(5.3, "FLandscapeGIBakedTextureBuilder is officially deprecated now and nothing updates it anymore") FLandscapeGIBakedTextureBuilder
 {
 public:
-	FLandscapeGIBakedTextureBuilder(UWorld* InWorld);
-	void Build();
-	int32 GetOutdatedGIBakedTextureComponentsCount(bool bInForceUpdate = true) const;
+	FLandscapeGIBakedTextureBuilder(UWorld* InWorld)
+		:World(InWorld)
+	{}
+	void Build(){}
+	int32 GetOutdatedGIBakedTextureComponentsCount(bool bInForceUpdate = true) const { return 0; }
+
 private:
-	UWorld* World;
-	mutable int32 OutdatedGIBakedTextureComponentsCount;
-	mutable double GIBakedTexturesLastCheckTime;
+	UWorld* World = nullptr;
+	mutable int32 OutdatedGIBakedTextureComponentsCount = 0;
+	mutable double GIBakedTexturesLastCheckTime = 0;
 };
 
 
diff --git a/Engine/Source/Runtime/Landscape/Private/Landscape.cpp b/Engine/Source/Runtime/Landscape/Private/Landscape.cpp
--- a/Engine/Source/Runtime/Landscape/Private/Landscape.cpp
+++ b/Engine/Source/Runtime/Landscape/Private/Landscape.cpp
@@ -760,8 +760,6 @@ void ULandscapeComponent::Serialize(FArchive& Ar)
 #if WITH_EDITOR
 	if (Ar.IsSaving() && Ar.IsPersistent())
 	{
-		//Update the last saved Guid for GI texture
-		LastBakedTextureMaterialGuid = BakedTextureMaterialGuid;
 		//Update the last saved Hash for physical material
 		LastSavedPhysicalMaterialHash = PhysicalMaterialHash;
 	}
@@ -1028,26 +1026,6 @@ void ULandscapeComponent::PostLoad()
 			}
 		}
 
-		if (GIBakedBaseColorTexture)
-		{
-			if (GIBakedBaseColorTexture->GetOutermost() != GetOutermost())
-			{
-				// The GIBakedBaseColorTexture property was never intended to be reassigned, but it was previously editable so we need to null any invalid values
-				// it will get recreated by ALandscapeProxy::UpdateBakedTextures()
-				GIBakedBaseColorTexture = nullptr;
-				BakedTextureMaterialGuid = FGuid();
-			}
-			else
-			{
-				// Remove public flag from GI textures to stop them being visible in the content browser.
-				// Previous version of landscape set these flags on creation.
-				if (GIBakedBaseColorTexture->HasAnyFlags(RF_Public))
-				{
-					GIBakedBaseColorTexture->ClearFlags(RF_Public);
-				}
-			}
-		}
-		LastBakedTextureMaterialGuid = BakedTextureMaterialGuid;
 		LastSavedPhysicalMaterialHash = PhysicalMaterialHash;
 
 		PRAGMA_DISABLE_DEPRECATION_WARNINGS;
@@ -1319,7 +1297,6 @@ ALandscapeProxy::ALandscapeProxy(const FObjectInitializer& ObjectInitializer)
 #if WITH_EDITOR
 	NumComponentsNeedingGrassMapRender = 0;
 	NumTexturesToStreamForVisibleGrassMapRender = 0;
-	NumComponentsNeedingTextureBaking = 0;
 
 	if (VisibilityLayer == nullptr)
 	{
@@ -2515,23 +2492,6 @@ void ALandscape::PostInitProperties()
 			{
 				if (FLandscapeNotificationManager* LandscapeNotificationManager = LandscapeSubSystem->GetNotificationManager())
 				{
-					// Create conditional TextureBakingNotification
-					TextureBakingNotification = MakeShared<FLandscapeNotification>(this, FLandscapeNotification::EType::TextureBaking,
-						[this]() { return NumComponentsNeedingTextureBaking > 0; },
-						[this](FText& InText)
-						{
-							// Accumulating all outstanding textures needing baking for ALL landscapes because only one such notification can be displayed at a time
-							int ComponentsNeedingTextureBakingAllLandscapes = 0;
-							for (const ALandscape* Landscape : TObjectRange<ALandscape>(RF_ClassDefaultObject | RF_ArchetypeObject, true, EInternalObjectFlags::Garbage))
-							{
-								ComponentsNeedingTextureBakingAllLandscapes += Landscape->NumComponentsNeedingTextureBaking;
-							}
-							
-							FFormatNamedArguments Args;
-							Args.Add(TEXT("OutstandingTextures"), FText::AsNumber(ComponentsNeedingTextureBakingAllLandscapes));
-							InText = FText::Format(NSLOCTEXT("TextureBaking", "TextureBakingFormat", "Baking Landscape Textures ({OutstandingTextures})"), Args);
-						});
-			
 					// Create conditional GrassRenderingNotification
 					GrassRenderingNotification = MakeShared<FLandscapeNotification>(this, FLandscapeNotification::EType::GrassRendering,
 						[this]() { return NumComponentsNeedingGrassMapRender > 0 && !(FSlateApplication::Get().HasAnyMouseCaptor() || GUnrealEd->IsUserInteracting()); },
@@ -2549,7 +2509,6 @@ void ALandscape::PostInitProperties()
 							InText = FText::Format(NSLOCTEXT("GrassMapRender", "GrassMapRenderFormat", "Building Grass Maps ({OutstandingGrassMaps})"), Args);
 						});
 			
-					LandscapeNotificationManager->RegisterNotification(TextureBakingNotification);
 					LandscapeNotificationManager->RegisterNotification(GrassRenderingNotification);
 				}
 			}
@@ -3045,8 +3004,6 @@ void ALandscapeProxy::PreSave(FObjectPreSaveContext ObjectSaveContext)
 #if WITH_EDITOR
 	// Work out whether we have grass or not for the next game run
 	BuildGrassMaps();
-	//Update the baked textures before saving
-	BuildGIBakedTextures();
 
 	for (ULandscapeComponent* Component : LandscapeComponents)
 	{
@@ -5099,304 +5056,6 @@ void ULandscapeComponent::SerializeStateHashes(FArchive& Ar)
 	Proxy->SerializeStateHashes(Ar);
 }
 
-FLandscapeGIBakedTextureBuilder::FLandscapeGIBakedTextureBuilder(UWorld* InWorld)
-	:World(InWorld)
-	,OutdatedGIBakedTextureComponentsCount(0)
-	,GIBakedTexturesLastCheckTime(0)
-{
-
-}
-
-void FLandscapeGIBakedTextureBuilder::Build()
-{
-	if (World)
-	{
-		for (TActorIterator<ALandscapeProxy> ProxyIt(World); ProxyIt; ++ProxyIt)
-		{
-			ProxyIt->BuildGIBakedTextures();
-		}
-		//Force update the outdated count when using the build menu option.
-		OutdatedGIBakedTextureComponentsCount = 0;
-		GIBakedTexturesLastCheckTime = FPlatformTime::Seconds();
-
-	}
-}
-
-int32 FLandscapeGIBakedTextureBuilder::GetOutdatedGIBakedTextureComponentsCount(bool bInForceUpdate) const
-{
-	if (World)
-	{
-		bool bUpdate = bInForceUpdate;
-		double GIBakedTexturesTimeNow = FPlatformTime::Seconds();
-		if (!bUpdate)
-		{
-			// Recheck every 20 secs
-			if ((GIBakedTexturesTimeNow - GIBakedTexturesLastCheckTime) > 20)
-			{
-				bUpdate = true;
-			}
-		}
-		if (bUpdate)
-		{
-			GIBakedTexturesLastCheckTime = GIBakedTexturesTimeNow;
-			OutdatedGIBakedTextureComponentsCount = 0;
-			for (TActorIterator<ALandscapeProxy> ProxyIt(World); ProxyIt; ++ProxyIt)
-			{
-				OutdatedGIBakedTextureComponentsCount += ProxyIt->GetOutdatedGIBakedTextureComponentsCount();
-			}
-		}
-	}
-	return OutdatedGIBakedTextureComponentsCount;
-}
-
-void ALandscapeProxy::BuildGIBakedTextures(struct FScopedSlowTask* InSlowTask /* = nullptr */)
-{
-	if (!HasAnyFlags(RF_ClassDefaultObject))
-	{
-		const bool bShouldMarkDirty = true;
-		UpdateGIBakedTextureData(bShouldMarkDirty);
-	}
-	}
-
-int32 ALandscapeProxy::GetOutdatedGIBakedTextureComponentsCount() const
-	{
-	int32 OutdatedGITextureComponentsCount = 0;
-	UpdateGIBakedTextureStatus(nullptr, nullptr, &OutdatedGITextureComponentsCount);
-	return OutdatedGITextureComponentsCount;
-	}
-
-void ALandscapeProxy::UpdateGIBakedTextureStatus(bool* bOutGenerateLandscapeGIData, TMap<UTexture2D*, FGIBakedTextureState>* OutComponentsNeedBakingByHeightmap, int32* OutdatedComponentsCount) const
-{
-	int32 OutdatedComponents = 0;
-	int32 ComponentsNeedToBeCleared = 0;
-	int32 ComponentsNeedToBeBaked = 0;
-
-	//@todo - remove Landscape GI Data
-	if (true)
-	{
-		if (bOutGenerateLandscapeGIData)
-		{
-			*bOutGenerateLandscapeGIData = false;
-		}
-
-		for (ULandscapeComponent* Component : LandscapeComponents)
-		{
-			if (Component != nullptr && Component->GIBakedBaseColorTexture != nullptr)
-			{
-				ComponentsNeedToBeCleared++;
-			}
-		}
-
-		OutdatedComponents += ComponentsNeedToBeCleared;
-	}
-	else
-	{
-		// Stores the components and their state hash data for a single atlas
-		struct FGIBakeTextureStateBuilder
-		{
-			// pointer as FMemoryWriter caches the address of the FBufferArchive, and this struct could be relocated on a realloc.
-			TUniquePtr<FBufferArchive> ComponentStateAr;
-			TArray<ULandscapeComponent*> Components;
-
-			FGIBakeTextureStateBuilder()
-			{
-				ComponentStateAr = MakeUnique<FBufferArchive>();
-			}
-		};
-
-		TMap<UTexture2D*, FGIBakeTextureStateBuilder> ComponentsByHeightmap;
-		for (ULandscapeComponent* Component : LandscapeComponents)
-		{
-			if (Component == nullptr)
-			{
-				continue;
-			}
-
-			FGIBakeTextureStateBuilder& Info = ComponentsByHeightmap.FindOrAdd(Component->GetHeightmap());
-			Info.Components.Add(Component);
-			Component->SerializeStateHashes(*Info.ComponentStateAr);
-		}
-
-		for (auto It = ComponentsByHeightmap.CreateIterator(); It; ++It)
-		{
-			FGIBakeTextureStateBuilder& Info = It.Value();
-
-			// Calculate a combined Guid-like ID we can use for this component
-			uint32 Hash[5];
-			FSHA1::HashBuffer(Info.ComponentStateAr->GetData(), Info.ComponentStateAr->Num(), (uint8*)Hash);
-			FGuid CombinedStateId = FGuid(Hash[0] ^ Hash[4], Hash[1], Hash[2], Hash[3]);
-
-			if (Info.Components[0]->BakedTextureMaterialGuid != CombinedStateId)
-			{
-				ComponentsNeedToBeBaked += Info.Components.Num();
-				if (OutComponentsNeedBakingByHeightmap)
-				{
-					FGIBakedTextureState& GIBakedTextureState = OutComponentsNeedBakingByHeightmap->FindOrAdd(It.Key());
-					GIBakedTextureState.Components = MoveTemp(Info.Components);
-					GIBakedTextureState.CombinedStateId = CombinedStateId;
-				}
-			}
-		}
-		OutdatedComponents += ComponentsNeedToBeBaked;
-	}
-
-	if (OutdatedComponentsCount)
-	{
-		if (OutdatedComponents == 0)
-		{
-			for (auto Component : LandscapeComponents)
-			{
-				const bool bIsDirty = Component->GetPackage()->IsDirty();
-				if (Component->LastBakedTextureMaterialGuid != Component->BakedTextureMaterialGuid && !bIsDirty)
-				{
-					OutdatedComponents++;
-				}
-			}
-		}
-		*OutdatedComponentsCount = OutdatedComponents;
-	}
-}
-
-void ALandscapeProxy::UpdateGIBakedTextureData(bool bInShouldMarkDirty)
-{
-	const bool bBakeAllGITextures = true;
-	UpdateGIBakedTextures(bBakeAllGITextures);
-	if (bInShouldMarkDirty && GetOutdatedGIBakedTextureComponentsCount() > 0)
-	{
-		MarkPackageDirty();
-	}
-}
-
-void ALandscapeProxy::UpdateGIBakedTextures(bool bBakeAllGITextures)
-{
-	// See if we can render
-	UWorld* World = GetWorld();
-	if (!GIsEditor || GUsingNullRHI || !World || World->IsGameWorld() || World->GetFeatureLevel() < ERHIFeatureLevel::SM5)
-	{
-		return;
-	}
-
-	if (!bBakeAllGITextures && UpdateBakedTexturesCountdown-- > 0)
-	{
-		return;
-	}
-
-	bool bGenerateLandscapeGIData = true;
-	TMap<UTexture2D*, FGIBakedTextureState> ComponentsToBeBakedByHeightmap;
-	UpdateGIBakedTextureStatus(&bGenerateLandscapeGIData, &ComponentsToBeBakedByHeightmap);
-	
-	if (!bGenerateLandscapeGIData)
-	{
-		// Clear out any existing GI textures
-		for (ULandscapeComponent* Component : LandscapeComponents)
-		{
-			if (Component != nullptr && Component->GIBakedBaseColorTexture != nullptr)
-			{
-				Component->BakedTextureMaterialGuid.Invalidate();
-				Component->GIBakedBaseColorTexture = nullptr;
-				Component->MarkRenderStateDirty();
-			}
-		}
-
-		// Don't check if we need to update anything for another 60 frames
-		UpdateBakedTexturesCountdown = 60;
-
-		return;
-	}
-
-	NumComponentsNeedingTextureBaking = 0;
-	int32 NumGenerated = 0;
-
-	for (auto It = ComponentsToBeBakedByHeightmap.CreateConstIterator(); It; ++It)
-	{
-		const FGIBakedTextureState& Info = It.Value();
-
-		bool bCanBake = true;
-		for (ULandscapeComponent* Component : Info.Components)
-		{
-			// not registered; ignore this component
-			if (!Component->SceneProxy)
-			{
-				continue;
-			}
-
-			// Check we can render the material
-			UMaterialInstance* MaterialInstance = Component->GetMaterialInstance(0, false);
-			if (!MaterialInstance)
-			{
-				// Cannot render this component yet as it doesn't have a material; abandon the atlas for this heightmap
-				bCanBake = false;
-				break;
-			}
-
-			FMaterialResource* MaterialResource = MaterialInstance->GetMaterialResource(World->GetFeatureLevel());
-			if (!MaterialResource || !MaterialResource->HasValidGameThreadShaderMap())
-			{
-				// Cannot render this component yet as its shaders aren't compiled; abandon the atlas for this heightmap
-				bCanBake = false;
-				break;
-			}
-		}
-
-		if (bCanBake)
-		{
-			// We throttle, baking only one atlas per frame if bBakeAllGITextures is false.
-			if (!bBakeAllGITextures && NumGenerated > 0)
-				{
-					NumComponentsNeedingTextureBaking += Info.Components.Num();
-				}
-				else
-				{
-					UTexture2D* HeightmapTexture = It.Key();
-					// 1/8 the res of the heightmap
-				FIntPoint AtlasSize(HeightmapTexture->Source.GetSizeX() >> 3, HeightmapTexture->Source.GetSizeY() >> 3);
-
-					TArray<FColor> AtlasSamples;
-					AtlasSamples.AddZeroed(AtlasSize.X * AtlasSize.Y);
-
-					for (ULandscapeComponent* Component : Info.Components)
-					{
-						// not registered; ignore this component
-						if (!Component->SceneProxy)
-						{
-							continue;
-						}
-
-						int32 ComponentSamples = (SubsectionSizeQuads + 1) * NumSubsections;
-						check(FMath::IsPowerOfTwo(ComponentSamples));
-
-						int32 BakeSize = ComponentSamples >> 3;
-						TArray<FColor> Samples;
-						if (FMaterialUtilities::ExportBaseColor(Component, BakeSize, Samples))
-						{
-							const int32 AtlasOffsetX = FMath::RoundToInt32(Component->HeightmapScaleBias.Z * HeightmapTexture->Source.GetSizeX()) >> 3;
-							const int32 AtlasOffsetY = FMath::RoundToInt32(Component->HeightmapScaleBias.W * HeightmapTexture->Source.GetSizeY()) >> 3;
-							for (int32 y = 0; y < BakeSize; y++)
-							{
-								FMemory::Memcpy(&AtlasSamples[(y + AtlasOffsetY) * AtlasSize.X + AtlasOffsetX], &Samples[y * BakeSize], sizeof(FColor) * BakeSize);
-							}
-							NumGenerated++;
-						}
-					}
-				UTexture2D* AtlasTexture = FMaterialUtilities::CreateTexture(GetOutermost(), HeightmapTexture->GetName() + TEXT("_BaseColor"), AtlasSize, AtlasSamples, TC_Default, TEXTUREGROUP_World, RF_NoFlags, true, Info.CombinedStateId);
-
-					for (ULandscapeComponent* Component : Info.Components)
-					{
-						Component->BakedTextureMaterialGuid = Info.CombinedStateId;
-						Component->GIBakedBaseColorTexture = AtlasTexture;
-						Component->MarkRenderStateDirty();
-					}
-				}
-			}
-		}
-	
-	if (NumGenerated == 0)
-	{
-		// Don't check if we need to update anything for another 60 frames
-		UpdateBakedTexturesCountdown = 60;
-	}
-}
-
 FLandscapePhysicalMaterialBuilder::FLandscapePhysicalMaterialBuilder(UWorld* InWorld)
 	:World(InWorld)
 	,OudatedPhysicalMaterialComponentsCount(0)
@@ -5508,7 +5167,6 @@ void ALandscapeProxy::InvalidateGeneratedComponentData(const TArray<ULandscapeCo
 		{
 			Component->InvalidateLightingCache();
 		}
-		Component->BakedTextureMaterialGuid.Invalidate();
 		ByProxy.FindOrAdd(Component->GetLandscapeProxy()).Add(Component);
 	}
 
diff --git a/Engine/Source/Runtime/Landscape/Private/LandscapeEdit.cpp b/Engine/Source/Runtime/Landscape/Private/LandscapeEdit.cpp
--- a/Engine/Source/Runtime/Landscape/Private/LandscapeEdit.cpp
+++ b/Engine/Source/Runtime/Landscape/Private/LandscapeEdit.cpp
@@ -4792,16 +4792,6 @@ void ULandscapeInfo::RemoveXYOffsets()
 	});
 }
 
-void ULandscapeInfo::PostponeTextureBaking()
-{
-	static const int32 PostponeValue = 60; //frames
-	
-	ForAllLandscapeProxies([](ALandscapeProxy* Proxy)
-	{
-		Proxy->UpdateBakedTexturesCountdown = PostponeValue;
-	});
-}
-
 bool ULandscapeInfo::CanHaveLayersContent() const
 {
 	if (ALandscape* Landscape = LandscapeActor.Get())
diff --git a/Engine/Source/Runtime/Landscape/Private/LandscapeGrass.cpp b/Engine/Source/Runtime/Landscape/Private/LandscapeGrass.cpp
--- a/Engine/Source/Runtime/Landscape/Private/LandscapeGrass.cpp
+++ b/Engine/Source/Runtime/Landscape/Private/LandscapeGrass.cpp
@@ -3159,9 +3159,7 @@ int32 ALandscapeProxy::GetOutdatedGrassMapCount() const
 	return OutdatedGrassMaps;
 }
 
-int32 ALandscapeProxy::TotalComponentsNeedingGrassMapRender = 0;
 int32 ALandscapeProxy::TotalTexturesToStreamForVisibleGrassMapRender = 0;
-int32 ALandscapeProxy::TotalComponentsNeedingTextureBaking = 0;
 
 void ALandscapeProxy::UpdateGrassDataStatus(TSet<UTexture2D*>* OutCurrentForcedStreamedTextures, TSet<UTexture2D*>* OutDesiredForcedStreamedTextures, TSet<ULandscapeComponent*>* OutComponentsNeedingGrassMapRender, TSet<ULandscapeComponent*>* OutOutdatedComponents, bool bInEnableForceResidentFlag, int32* OutOutdatedGrassMaps) const
 {
diff --git a/Engine/Source/Runtime/Landscape/Private/LandscapeRender.cpp b/Engine/Source/Runtime/Landscape/Private/LandscapeRender.cpp
--- a/Engine/Source/Runtime/Landscape/Private/LandscapeRender.cpp
+++ b/Engine/Source/Runtime/Landscape/Private/LandscapeRender.cpp
@@ -890,7 +890,6 @@ FLandscapeComponentSceneProxy::FLandscapeComponentSceneProxy(ULandscapeComponent
 	, VisibilityWeightmapTexture(nullptr)
 	, VisibilityWeightmapChannel(-1)
 	, HeightmapTexture(InComponent->GetHeightmap())
-	, BaseColorForGITexture(InComponent->GIBakedBaseColorTexture)
 	, HeightmapScaleBias(InComponent->HeightmapScaleBias)
 	, XYOffsetmapTexture(InComponent->XYOffsetmapTexture)
 	, SharedBuffersKey(0)
@@ -3810,10 +3809,9 @@ bool FLandscapeComponentSceneProxy::HeightfieldHasPendingStreaming() const
 	return bHeightmapTextureStreaming || bVisibilityTextureStreaming;
 }
 
-void FLandscapeComponentSceneProxy::GetHeightfieldRepresentation(UTexture2D*& OutHeightmapTexture, UTexture2D*& OutDiffuseColorTexture, UTexture2D*& OutVisibilityTexture, FHeightfieldComponentDescription& OutDescription) const
+void FLandscapeComponentSceneProxy::GetHeightfieldRepresentation(UTexture2D*& OutHeightmapTexture, UTexture2D*& OutVisibilityTexture, FHeightfieldComponentDescription& OutDescription) const
 {
 	OutHeightmapTexture = HeightmapTexture;
-	OutDiffuseColorTexture = BaseColorForGITexture;
 	OutVisibilityTexture = VisibilityWeightmapTexture;
 
 	OutDescription.HeightfieldScaleBias = HeightmapScaleBias;
diff --git a/Engine/Source/Runtime/Landscape/Private/LandscapeSubsystem.cpp b/Engine/Source/Runtime/Landscape/Private/LandscapeSubsystem.cpp
--- a/Engine/Source/Runtime/Landscape/Private/LandscapeSubsystem.cpp
+++ b/Engine/Source/Runtime/Landscape/Private/LandscapeSubsystem.cpp
@@ -49,13 +49,6 @@ DECLARE_CYCLE_STAT(TEXT("LandscapeSubsystem Tick"), STAT_LandscapeSubsystemTick,
 #define LOCTEXT_NAMESPACE "LandscapeSubsystem"
 
 ULandscapeSubsystem::ULandscapeSubsystem()
-	: bIsGrassCreationPrioritized(false)
-#if WITH_EDITOR
-	, GrassMapsBuilder(nullptr)
-	, GIBakedTextureBuilder(nullptr)
-	, PhysicalMaterialBuilder(nullptr)
-	, NotificationManager(nullptr)
-#endif
 {
 }
 
@@ -87,7 +80,6 @@ void ULandscapeSubsystem::Initialize(FSubsystemCollectionBase& Collection)
 
 #if WITH_EDITOR
 	GrassMapsBuilder = new FLandscapeGrassMapsBuilder(GetWorld());
-	GIBakedTextureBuilder = new FLandscapeGIBakedTextureBuilder(GetWorld());
 	PhysicalMaterialBuilder = new FLandscapePhysicalMaterialBuilder(GetWorld());
 
 	if (!IsRunningCommandlet())
@@ -113,7 +105,6 @@ void ULandscapeSubsystem::Deinitialize()
 	
 #if WITH_EDITOR
 	delete GrassMapsBuilder;
-	delete GIBakedTextureBuilder;
 	delete PhysicalMaterialBuilder;
 	delete NotificationManager;
 #endif
@@ -275,7 +266,6 @@ void ULandscapeSubsystem::Tick(float DeltaTime)
 				// editor-only
 				if (!World->IsPlayInEditor())
 				{
-					Proxy->UpdateGIBakedTextures();
 					Proxy->UpdatePhysicalMaterialTasks();
 				}
 			}
@@ -309,7 +299,6 @@ void ULandscapeSubsystem::BuildAll()
 	MarkModifiedLandscapesAsDirty();
 
 	BuildGrassMaps();
-	BuildGIBakedTextures();
 	BuildPhysicalMaterial();
 	BuildNanite();
 }
@@ -324,16 +313,6 @@ int32 ULandscapeSubsystem::GetOutdatedGrassMapCount()
 	return GrassMapsBuilder->GetOutdatedGrassMapCount(/*bInForceUpdate*/false);
 }
 
-void ULandscapeSubsystem::BuildGIBakedTextures()
-{
-	GIBakedTextureBuilder->Build();
-}
-
-int32 ULandscapeSubsystem::GetOutdatedGIBakedTextureComponentsCount()
-{
-	return GIBakedTextureBuilder->GetOutdatedGIBakedTextureComponentsCount(/*bInForceUpdate*/false);
-}
-
 void ULandscapeSubsystem::BuildPhysicalMaterial()
 {
 	PhysicalMaterialBuilder->Build();
@@ -549,14 +528,6 @@ void ULandscapeSubsystem::DisplayMessages(FCanvas* Canvas, float& XPos, float& Y
 		YPos += FontSizeY;
 	}
 
-	if (int32 ComponentsNeedingGITextureBaking = GetOutdatedGIBakedTextureComponentsCount())
-	{
-		SmallTextItem.SetColor(FLinearColor::Red);
-		SmallTextItem.Text = FText::Format(LOCTEXT("LANDSCAPE_TEXTURES_NEED_TO_BE_REBUILT_FMT", "LANDSCAPE: BAKED TEXTURES NEEDS TO BE REBUILT ({0} {0}|plural(one=object,other=objects))"), ComponentsNeedingGITextureBaking);
-		Canvas->DrawItem(SmallTextItem, FVector2D(XPos, YPos));
-		YPos += FontSizeY;
-	}
-
 	if (int32 ComponentsWithOudatedPhysicalMaterial = GetOudatedPhysicalMaterialComponentsCount())
 	{
 		SmallTextItem.SetColor(FLinearColor::Red);
diff --git a/Engine/Source/Runtime/Landscape/Public/LandscapeNotification.h b/Engine/Source/Runtime/Landscape/Public/LandscapeNotification.h
--- a/Engine/Source/Runtime/Landscape/Public/LandscapeNotification.h
+++ b/Engine/Source/Runtime/Landscape/Public/LandscapeNotification.h
@@ -15,7 +15,6 @@ struct FLandscapeNotification
 		ShadingModelInvalid = 0,
 		LandscapeTextureResourcesNotReady,
 		LandscapeBrushResourcesNotReady,
-		TextureBaking,
 		GrassRendering,
 	};
 
diff --git a/Engine/Source/Runtime/Landscape/Public/LandscapeRender.h b/Engine/Source/Runtime/Landscape/Public/LandscapeRender.h
--- a/Engine/Source/Runtime/Landscape/Public/LandscapeRender.h
+++ b/Engine/Source/Runtime/Landscape/Public/LandscapeRender.h
@@ -719,7 +719,6 @@ protected:
 #endif
 	// Heightmap in RG and Normalmap in BA
 	UTexture2D* HeightmapTexture; 
-	UTexture2D* BaseColorForGITexture;
 	FVector4f HeightmapScaleBias;
 	float HeightmapSubsectionOffsetU;
 	float HeightmapSubsectionOffsetV;
@@ -819,7 +818,7 @@ public:
 
 	virtual bool HeightfieldHasPendingStreaming() const override;
 
-	virtual void GetHeightfieldRepresentation(UTexture2D*& OutHeightmapTexture, UTexture2D*& OutDiffuseColorTexture, UTexture2D*& OutVisibilityTexture, FHeightfieldComponentDescription& OutDescription) const override;
+	virtual void GetHeightfieldRepresentation(UTexture2D*& OutHeightmapTexture, UTexture2D*& OutVisibilityTexture, FHeightfieldComponentDescription& OutDescription) const override;
 
 	virtual void GetLCIs(FLCIArray& LCIs) override;
 
diff --git a/Engine/Source/Runtime/Landscape/Public/LandscapeSubsystem.h b/Engine/Source/Runtime/Landscape/Public/LandscapeSubsystem.h
--- a/Engine/Source/Runtime/Landscape/Public/LandscapeSubsystem.h
+++ b/Engine/Source/Runtime/Landscape/Public/LandscapeSubsystem.h
@@ -65,12 +65,16 @@ public:
 
 #if WITH_EDITOR
 	LANDSCAPE_API void BuildAll();
-	LANDSCAPE_API void BuildGrassMaps();
+
+	UE_DEPRECATED(5.3, "BuildGIBakedTextures is officially deprecated nowe")
+	LANDSCAPE_API void BuildGIBakedTextures() {}
+	UE_DEPRECATED(5.3, "GetOutdatedGIBakedTextureComponentsCount is officially deprecated now")
+	LANDSCAPE_API int32 GetOutdatedGIBakedTextureComponentsCount() { return 0; }
+
 	LANDSCAPE_API int32 GetOutdatedGrassMapCount();
-	LANDSCAPE_API void BuildGIBakedTextures();
-	LANDSCAPE_API int32 GetOutdatedGIBakedTextureComponentsCount();
-	LANDSCAPE_API void BuildPhysicalMaterial();
+	LANDSCAPE_API void BuildGrassMaps();
 	LANDSCAPE_API int32 GetOudatedPhysicalMaterialComponentsCount();
+	LANDSCAPE_API void BuildPhysicalMaterial();
 	/**
 	 * Updates the Nanite mesh on all landscape actors whose mesh is not up to date.
 	 * @param InProxiesToBuild - If specified, only the Nanite meshes of the specified landscape actors (recursively for all streaming proxies, in the case of a 1 ALandscape / N ALandscapeStreamingProxy setup) will be built
@@ -103,16 +107,15 @@ private:
 
 	void OnNaniteWorldSettingsChanged(AWorldSettings* WorldSettings) { RegenerateGrass(true, true); }
 
-	bool bIsGrassCreationPrioritized;
+	bool bIsGrassCreationPrioritized = false;
 	TArray<TWeakObjectPtr<ALandscapeProxy>> Proxies;
 	FDelegateHandle OnNaniteWorldSettingsChangedHandle;
 
 #if WITH_EDITOR
-	class FLandscapeGrassMapsBuilder* GrassMapsBuilder;
-	class FLandscapeGIBakedTextureBuilder* GIBakedTextureBuilder;
-	class FLandscapePhysicalMaterialBuilder* PhysicalMaterialBuilder;
+	class FLandscapeGrassMapsBuilder* GrassMapsBuilder = nullptr;
+	class FLandscapePhysicalMaterialBuilder* PhysicalMaterialBuilder = nullptr;
 	
-	FLandscapeNotificationManager* NotificationManager;
+	FLandscapeNotificationManager* NotificationManager = nullptr;
 	FOnHeightmapStreamedDelegate OnHeightmapStreamed;
 	bool bAnyViewShowCollisions = false;
 	FDateTime AppCurrentDateTime; // Represents FDateTime::Now(), at the beginning of the frame (useful to get a human-readable date/time that is fixed during the frame)
diff --git a/Engine/Source/Runtime/Renderer/Private/DistanceFieldObjectManagement.cpp b/Engine/Source/Runtime/Renderer/Private/DistanceFieldObjectManagement.cpp
--- a/Engine/Source/Runtime/Renderer/Private/DistanceFieldObjectManagement.cpp
+++ b/Engine/Source/Runtime/Renderer/Private/DistanceFieldObjectManagement.cpp
@@ -764,10 +764,9 @@ void FSceneRenderer::UpdateGlobalHeightFieldObjectBuffers(FRDGBuilder& GraphBuil
 			FVector4f* UploadObjectBounds = (FVector4f*)DistanceFieldSceneData.UploadHeightFieldBoundsBuffer.Add_GetRef(PrimitiveIndex);
 
 			UTexture2D* HeightNormalTexture;
-			UTexture2D* DiffuseColorTexture;
 			UTexture2D* VisibilityTexture;
 			FHeightfieldComponentDescription HeightFieldCompDesc(Primitive->Proxy->GetLocalToWorld(), Primitive->GetInstanceSceneDataOffset());
-			Primitive->Proxy->GetHeightfieldRepresentation(HeightNormalTexture, DiffuseColorTexture, VisibilityTexture, HeightFieldCompDesc);
+			Primitive->Proxy->GetHeightfieldRepresentation(HeightNormalTexture, VisibilityTexture, HeightFieldCompDesc);
 
 			const bool bInAtlas = HeightNormalTexture && (GHeightFieldTextureAtlas.GetAllocationHandle(HeightNormalTexture) != INDEX_NONE);
 
diff --git a/Engine/Source/Runtime/Renderer/Private/GlobalDistanceField.cpp b/Engine/Source/Runtime/Renderer/Private/GlobalDistanceField.cpp
--- a/Engine/Source/Runtime/Renderer/Private/GlobalDistanceField.cpp
+++ b/Engine/Source/Runtime/Renderer/Private/GlobalDistanceField.cpp
@@ -2282,14 +2282,13 @@ void UpdateGlobalDistanceFieldVolume(
 								if (Clipmap.Bounds.Intersect(HeightfieldInfluenceBox))
 								{
 									UTexture2D* HeightfieldTexture = nullptr;
-									UTexture2D* DiffuseColorTexture = nullptr;
 									UTexture2D* VisibilityTexture = nullptr;
 									FHeightfieldComponentDescription NewComponentDescription(HeightfieldPrimitiveProxy->GetLocalToWorld(), GPUSceneInstanceIndex);
-									HeightfieldPrimitiveProxy->GetHeightfieldRepresentation(HeightfieldTexture, DiffuseColorTexture, VisibilityTexture, NewComponentDescription);
+									HeightfieldPrimitiveProxy->GetHeightfieldRepresentation(HeightfieldTexture, VisibilityTexture, NewComponentDescription);
 
 									if (HeightfieldTexture && HeightfieldTexture->GetResource() && HeightfieldTexture->GetResource()->TextureRHI)
 									{
-										TArray<FHeightfieldComponentDescription>& ComponentDescriptions = UpdateRegionHeightfield.ComponentDescriptions.FindOrAdd(FHeightfieldComponentTextures(HeightfieldTexture, DiffuseColorTexture, VisibilityTexture));
+										TArray<FHeightfieldComponentDescription>& ComponentDescriptions = UpdateRegionHeightfield.ComponentDescriptions.FindOrAdd(FHeightfieldComponentTextures(HeightfieldTexture, VisibilityTexture));
 										ComponentDescriptions.Add(NewComponentDescription);
 									}
 								}
diff --git a/Engine/Source/Runtime/Renderer/Private/HeightfieldLighting.h b/Engine/Source/Runtime/Renderer/Private/HeightfieldLighting.h
--- a/Engine/Source/Runtime/Renderer/Private/HeightfieldLighting.h
+++ b/Engine/Source/Runtime/Renderer/Private/HeightfieldLighting.h
@@ -19,15 +19,14 @@ class FHeightfieldComponentTextures
 {
 public:
 
-	FHeightfieldComponentTextures(UTexture2D* InHeightAndNormal, UTexture2D* InDiffuseColor, UTexture2D* InVisibility) :
+	FHeightfieldComponentTextures(UTexture2D* InHeightAndNormal, UTexture2D* InVisibility) :
 		HeightAndNormal(InHeightAndNormal),
-		DiffuseColor(InDiffuseColor),
 		Visibility(InVisibility)
 	{}
 
 	FORCEINLINE bool operator==(FHeightfieldComponentTextures Other) const
 	{
-		return HeightAndNormal == Other.HeightAndNormal && DiffuseColor == Other.DiffuseColor && Visibility == Other.Visibility;
+		return HeightAndNormal == Other.HeightAndNormal && Visibility == Other.Visibility;
 	}
 
 	FORCEINLINE friend uint32 GetTypeHash(FHeightfieldComponentTextures ComponentTextures)
@@ -36,7 +35,6 @@ public:
 	}
 
 	UTexture2D* HeightAndNormal;
-	UTexture2D* DiffuseColor;
 	UTexture2D* Visibility;
 };
 
diff --git a/Engine/Source/Runtime/Renderer/Private/RendererScene.cpp b/Engine/Source/Runtime/Renderer/Private/RendererScene.cpp
--- a/Engine/Source/Runtime/Renderer/Private/RendererScene.cpp
+++ b/Engine/Source/Runtime/Renderer/Private/RendererScene.cpp
@@ -722,10 +722,9 @@ void FDistanceFieldSceneData::AddPrimitive(FPrimitiveSceneInfo* InPrimitive)
 		if (Proxy->SupportsHeightfieldRepresentation())
 		{
 			UTexture2D* HeightAndNormal;
-			UTexture2D* DiffuseColor;
 			UTexture2D* Visibility;
 			FHeightfieldComponentDescription Desc(FMatrix::Identity, InPrimitive->GetInstanceSceneDataOffset());
-			Proxy->GetHeightfieldRepresentation(HeightAndNormal, DiffuseColor, Visibility, Desc);
+			Proxy->GetHeightfieldRepresentation(HeightAndNormal, Visibility, Desc);
 			GHeightFieldTextureAtlas.AddAllocation(HeightAndNormal);
 
 			if (Visibility)
@@ -786,10 +785,9 @@ void FDistanceFieldSceneData::RemovePrimitive(FPrimitiveSceneInfo* InPrimitive)
 		if (Proxy->SupportsHeightfieldRepresentation())
 		{
 			UTexture2D* HeightAndNormal;
-			UTexture2D* DiffuseColor;
 			UTexture2D* Visibility;
 			FHeightfieldComponentDescription Desc(FMatrix::Identity, InPrimitive->GetInstanceSceneDataOffset());
-			Proxy->GetHeightfieldRepresentation(HeightAndNormal, DiffuseColor, Visibility, Desc);
+			Proxy->GetHeightfieldRepresentation(HeightAndNormal, Visibility, Desc);
 			GHeightFieldTextureAtlas.RemoveAllocation(HeightAndNormal);
 
 			if (Visibility)
