Fixes update grass method not using GetLandscapeMaterial getter.

diff --git a/Engine/Source/Runtime/Landscape/Private/LandscapeGrass.cpp b/Engine/Source/Runtime/Landscape/Private/LandscapeGrass.cpp
--- a/Engine/Source/Runtime/Landscape/Private/LandscapeGrass.cpp
+++ b/Engine/Source/Runtime/Landscape/Private/LandscapeGrass.cpp
@@ -2043,7 +2043,7 @@ void ULandscapeGrassType::PostEditChangeProperty(FPropertyChangedEvent& Property
 			ALandscapeProxy* Proxy = *It;
 			if (Proxy->GetWorld() && !Proxy->GetWorld()->IsPlayInEditor())
 			{
-				const UMaterialInterface* MaterialInterface = Proxy->LandscapeMaterial;
+				const UMaterialInterface* MaterialInterface = Proxy->GetLandscapeMaterial();
 				if (MaterialInterface)
 				{
 					TArray<const UMaterialExpressionLandscapeGrassOutput*> GrassExpressions;
@@ -3198,7 +3198,7 @@ void ALandscapeProxy::UpdateGrassDataStatus(TSet<UTexture2D*>* OutCurrentForcedS
 
 	TArray<ULandscapeGrassType*> GrassTypes;
 	float OutMaxDiscardDistance = 0.0f;
-	GetGrassTypes(World, LandscapeMaterial, GrassTypes, OutMaxDiscardDistance);
+	GetGrassTypes(World, GetLandscapeMaterial(), GrassTypes, OutMaxDiscardDistance);
 	const bool bHasGrassTypes = GrassTypes.Num() > 0;
 
 	const bool bIsOutermostPackageDirty = GetOutermost()->IsDirty();
@@ -3415,18 +3415,19 @@ void ALandscapeProxy::UpdateGrass(const TArray<FVector>& Cameras, int32& InOutNu
 
 		if (World)
 		{
+			UMaterialInterface* LandscapeMaterialInterface = GetLandscapeMaterial();
 #if WITH_EDITOR
 			TArray<ULandscapeGrassType*> LandscapeGrassTypes;
 			float GrassMaxDiscardDistance = 0.0f;
-			GetGrassTypes(World, LandscapeMaterial, LandscapeGrassTypes, GrassMaxDiscardDistance);
+			GetGrassTypes(World, LandscapeMaterialInterface, LandscapeGrassTypes, GrassMaxDiscardDistance);
 #else
 			// In non editor builds, cache grass types for performance.
-			if (LandscapeMaterial != LandscapeMaterialCached)
+			if (LandscapeMaterialInterface != LandscapeMaterialCached)
 			{
-				LandscapeMaterialCached = LandscapeMaterial;
+				LandscapeMaterialCached = LandscapeMaterialInterface;
 				LandscapeGrassTypes.Reset();
 				GrassMaxDiscardDistance = 0.0f;
-				GetGrassTypes(World, LandscapeMaterial, LandscapeGrassTypes, GrassMaxDiscardDistance);
+				GetGrassTypes(World, LandscapeMaterialInterface, LandscapeGrassTypes, GrassMaxDiscardDistance);
 			}
 #endif
 			// Cull grass max distance based on Cull Distance scale factor and on max GuardBand factor.
