Clang build fix

#rnx

diff --git a/Engine/Source/Runtime/Landscape/Classes/LandscapeProxy.h b/Engine/Source/Runtime/Landscape/Classes/LandscapeProxy.h
--- a/Engine/Source/Runtime/Landscape/Classes/LandscapeProxy.h
+++ b/Engine/Source/Runtime/Landscape/Classes/LandscapeProxy.h
@@ -927,8 +927,7 @@ public:
 	/** Render grass maps for the specified components */
 	void RenderGrassMaps(const TArray<ULandscapeComponent*>& LandscapeComponents, const TArray<ULandscapeGrassType*>& GrassTypes);
 
-	UE_DEPRECATED(5.3, "FGIBakedTextureState is officially deprecated now and nothing updates it anymore")
-	struct FGIBakedTextureState
+	struct UE_DEPRECATED(5.3, "FGIBakedTextureState is officially deprecated now and nothing updates it anymore") FGIBakedTextureState
 	{
 		FGuid CombinedStateId;
 		TArray<ULandscapeComponent*> Components;
