Disables temporarily soft ptr assert in PostLoad.

diff --git a/Engine/Source/Runtime/Landscape/Private/Landscape.cpp b/Engine/Source/Runtime/Landscape/Private/Landscape.cpp
--- a/Engine/Source/Runtime/Landscape/Private/Landscape.cpp
+++ b/Engine/Source/Runtime/Landscape/Private/Landscape.cpp
@@ -3689,8 +3689,6 @@ UMaterialInterface* ALandscapeStreamingProxy::GetLandscapeMaterial(int8 InLODInd
 		return LandscapeMaterial;
 	}
 
-	checkf(!FUObjectThreadContext::Get().IsRoutingPostLoad, TEXT("This method should not be called during PostLoad since the LandscapeActor may not be fully unserialized yet."));
-
 	if (const ALandscape* Landscape = GetLandscapeActor())
 	{
 		return Landscape->GetLandscapeMaterial(InLODIndex);
