Implements delayed post load mechanism for landscape actors and component relying on parent actor soft ptr link.

diff --git a/Engine/Source/Runtime/Landscape/Classes/LandscapeComponent.h b/Engine/Source/Runtime/Landscape/Classes/LandscapeComponent.h
--- a/Engine/Source/Runtime/Landscape/Classes/LandscapeComponent.h
+++ b/Engine/Source/Runtime/Landscape/Classes/LandscapeComponent.h
@@ -12,6 +12,7 @@
 #include "Serialization/BulkData.h"
 #include "LandscapePhysicalMaterial.h"
 #include "LandscapeInfo.h"
+#include "LandscapeUtils.h"
 #include "LandscapeWeightmapUsage.h"
 #include "Containers/ArrayView.h"
 #include "Engine/StreamableRenderAsset.h"
@@ -457,6 +458,9 @@ class ULandscapeComponent : public UPrimitiveComponent
 	TLazyObjectPtr<ULandscapeHeightfieldCollisionComponent> CollisionComponent_DEPRECATED;
 #endif // !WITH_EDITORONLY_DATA
 
+protected:
+	UE::Landscape::EDelayedPostLoadState DelayedPostLoadState = UE::Landscape::EDelayedPostLoadState::Default;
+
 private:
 	/** Reference to associated collision component */
 	UPROPERTY()
@@ -1136,6 +1140,9 @@ protected:
 		return true;
 	}
 
+	/** Used to simulate a delayed PostLoad when registering component. This is needed because at this time we have no guarantee the parent ALandscape is fully loaded. */
+	virtual void TryDelayedPostLoad();
+
 #if WITH_EDITOR
 public:
 	/** Records the ULandscapeComponents that are modified in any undo/redo operation that is being applied currently */
diff --git a/Engine/Source/Runtime/Landscape/Classes/LandscapeProxy.h b/Engine/Source/Runtime/Landscape/Classes/LandscapeProxy.h
--- a/Engine/Source/Runtime/Landscape/Classes/LandscapeProxy.h
+++ b/Engine/Source/Runtime/Landscape/Classes/LandscapeProxy.h
@@ -17,6 +17,7 @@
 #include "VT/RuntimeVirtualTextureEnum.h"
 #include "ActorPartition/PartitionActor.h"
 #include "ILandscapeSplineInterface.h"
+#include "LandscapeUtils.h"
 
 #if UE_ENABLE_INCLUDE_ORDER_DEPRECATED_IN_5_2
 #include "LandscapeComponent.h"
@@ -387,6 +388,8 @@ protected:
 	TArray<FLandscapePerLODMaterialOverride> PreEditPerLODOverrideMaterials;
 #endif // WITH_EDITORONLY_DATA
 
+	UE::Landscape::EDelayedPostLoadState DelayedPostLoadState = UE::Landscape::EDelayedPostLoadState::Default;
+
 public:
 	LANDSCAPE_API TOptional<float> GetHeightAtLocation(FVector Location, EHeightfieldSource HeightFieldSource = EHeightfieldSource::Complex) const;
 
@@ -1312,9 +1315,14 @@ protected:
 	/** Initialize Layer with empty content if it hasn't been initialized yet. */
 	LANDSCAPE_API void InitializeLayerWithEmptyContent(const FGuid& InLayerGuid);
 
+#endif // WITH_EDITOR
+
+	/** Used to simulate a delayed PostLoad when PostRegisteringAllComponents. This is needed because at this time we have no guarantee the parent ALandscape is fully loaded. */
+	virtual void TryDelayedPostLoad();
+
 protected:
+#if WITH_EDITOR
 	FLandscapeMaterialChangedDelegate LandscapeMaterialChangedDelegate;
-
 #endif // WITH_EDITOR
 private:
 	/** Returns Grass Update interval */
diff --git a/Engine/Source/Runtime/Landscape/Private/Landscape.cpp b/Engine/Source/Runtime/Landscape/Private/Landscape.cpp
--- a/Engine/Source/Runtime/Landscape/Private/Landscape.cpp
+++ b/Engine/Source/Runtime/Landscape/Private/Landscape.cpp
@@ -15,6 +15,7 @@ Landscape.cpp: Terrain rendering
 #include "UObject/ConstructorHelpers.h"
 #include "UObject/DevObjectVersion.h"
 #include "UObject/LinkerLoad.h"
+#include "UObject/UObjectThreadContext.h"
 #include "Framework/Application/SlateApplication.h"
 #include "LandscapePrivate.h"
 #include "LandscapeStreamingProxy.h"
@@ -92,7 +93,6 @@ Landscape.cpp: Terrain rendering
 #include "LandscapeHeightfieldCollisionComponent.h"
 #include "UObject/EditorObjectVersion.h"
 #include "Algo/BinarySearch.h"
-#include "UObject/UObjectThreadContext.h"
 #include "WorldPartition/WorldPartition.h"
 #include "WorldPartition/WorldPartitionHelpers.h"
 #include "WorldPartition/WorldPartitionHandle.h"
@@ -936,290 +936,7 @@ void ULandscapeComponent::PostLoad()
 {
 	Super::PostLoad();
 
-	if (IsComponentPSOPrecachingEnabled())
-	{
-		TArray<UMaterialInterface*> Materials;
-		bool bGetDebugMaterials = false;
-		GetUsedMaterials(Materials, bGetDebugMaterials);
-
-		FPSOPrecacheParams PrecachePSOParams;
-		SetupPrecachePSOParams(PrecachePSOParams);
-
-		FPSOPrecacheVertexFactoryDataList VertexFactoryDataList;
-		if (!XYOffsetmapTexture)
-		{
-			VertexFactoryDataList.Add(FPSOPrecacheVertexFactoryData(&FLandscapeVertexFactory::StaticType));
-		}
-		else
-		{
-			VertexFactoryDataList.Add(FPSOPrecacheVertexFactoryData(&FLandscapeXYOffsetVertexFactory::StaticType));
-		}
-
-		// we need the fixed grid vertex factory for both virtual texturing and grass
-		ERHIFeatureLevel::Type FeatureLevel = GMaxRHIFeatureLevel;
-		bool bNeedsFixedGridVertexFactory = UseVirtualTexturing(FeatureLevel);
-		if (bNeedsFixedGridVertexFactory)
-		{
-			VertexFactoryDataList.Add(FPSOPrecacheVertexFactoryData(&FLandscapeFixedGridVertexFactory::StaticType));
-		}
-
-		TArray<FMaterialPSOPrecacheRequestID> MaterialPrecacheRequestIDs;
-		for (UMaterialInterface* MaterialInterface : Materials)
-		{
-			if (MaterialInterface)
-			{
-				MaterialInterface->ConditionalPostLoad();
-				MaterialInterface->PrecachePSOs(VertexFactoryDataList, PrecachePSOParams, EPSOPrecachePriority::High, MaterialPrecacheRequestIDs);
-			}
-		}
-	}
-
-#if WITH_EDITOR
-	ALandscapeProxy* LandscapeProxy = GetLandscapeProxy();
-	if (ensure(LandscapeProxy))
-	{
-		// Ensure that the component's lighting settings matches the actor's.
-		UpdatedSharedPropertiesFromActor();
-
-		// check SectionBaseX/Y are correct
-		const FVector LocalRelativeLocation = GetRelativeLocation();
-		int32 CheckSectionBaseX = FMath::RoundToInt32(LocalRelativeLocation.X) + LandscapeProxy->LandscapeSectionOffset.X;
-		int32 CheckSectionBaseY = FMath::RoundToInt32(LocalRelativeLocation.Y) + LandscapeProxy->LandscapeSectionOffset.Y;
-		if (CheckSectionBaseX != SectionBaseX ||
-			CheckSectionBaseY != SectionBaseY)
-		{
-			UE_LOG(LogLandscape, Warning, TEXT("LandscapeComponent SectionBaseX disagrees with its location, attempted automated fix: '%s', %d,%d vs %d,%d."),
-				*GetFullName(), SectionBaseX, SectionBaseY, CheckSectionBaseX, CheckSectionBaseY);
-			SectionBaseX = CheckSectionBaseX;
-			SectionBaseY = CheckSectionBaseY;
-		}
-	}
-
-	if (GIsEditor && !HasAnyFlags(RF_ClassDefaultObject))
-	{
-		// This is to ensure that component relative location is exact section base offset value
-		FVector LocalRelativeLocation = GetRelativeLocation();
-		float CheckRelativeLocationX = float(SectionBaseX - LandscapeProxy->LandscapeSectionOffset.X);
-		float CheckRelativeLocationY = float(SectionBaseY - LandscapeProxy->LandscapeSectionOffset.Y);
-		if (!FMath::IsNearlyEqual(CheckRelativeLocationX, LocalRelativeLocation.X, UE_DOUBLE_KINDA_SMALL_NUMBER) ||
-			!FMath::IsNearlyEqual(CheckRelativeLocationY, LocalRelativeLocation.Y, UE_DOUBLE_KINDA_SMALL_NUMBER))
-		{
-			UE_LOG(LogLandscape, Warning, TEXT("LandscapeComponent RelativeLocation disagrees with its section base, attempted automated fix: '%s', %f,%f vs %f,%f."),
-				*GetFullName(), LocalRelativeLocation.X, LocalRelativeLocation.Y, CheckRelativeLocationX, CheckRelativeLocationY);
-			LocalRelativeLocation.X = CheckRelativeLocationX;
-			LocalRelativeLocation.Y = CheckRelativeLocationY;
-
-			SetRelativeLocation_Direct(LocalRelativeLocation);
-		}
-
-		// Remove standalone flags from data textures to ensure data is unloaded in the editor when reverting an unsaved level.
-		// Previous version of landscape set these flags on creation.
-		if (HeightmapTexture && HeightmapTexture->HasAnyFlags(RF_Standalone))
-		{
-			HeightmapTexture->ClearFlags(RF_Standalone);
-		}
-		for (int32 Idx = 0; Idx < WeightmapTextures.Num(); Idx++)
-		{
-			if (WeightmapTextures[Idx] && WeightmapTextures[Idx]->HasAnyFlags(RF_Standalone))
-			{
-				WeightmapTextures[Idx]->ClearFlags(RF_Standalone);
-			}
-		}
-
-		LastSavedPhysicalMaterialHash = PhysicalMaterialHash;
-
-		PRAGMA_DISABLE_DEPRECATION_WARNINGS;
-		if (!OverrideMaterials_DEPRECATED.IsEmpty())
-		{
-			PerLODOverrideMaterials.Reserve(OverrideMaterials_DEPRECATED.Num());
-			for (const FLandscapeComponentMaterialOverride& LocalMaterialOverride : OverrideMaterials_DEPRECATED)
-			{
-				PerLODOverrideMaterials.Add({ LocalMaterialOverride.LODIndex.Default, LocalMaterialOverride.Material });
-			}
-			OverrideMaterials_DEPRECATED.Reset();
-		}
-		PRAGMA_ENABLE_DEPRECATION_WARNINGS;
-	}
-
-#if WITH_EDITORONLY_DATA
-	// Handle old MaterialInstance
-	if (MaterialInstance_DEPRECATED)
-	{
-		MaterialInstances.Empty(1);
-		MaterialInstances.Add(MaterialInstance_DEPRECATED);
-		MaterialInstance_DEPRECATED = nullptr;
-
-		if (GIsEditor && MaterialInstances.Num() > 0 && MaterialInstances[0] != nullptr)
-		{
-			MaterialInstances[0]->ConditionalPostLoad();
-			UpdateMaterialInstances();
-		}
-	}
-#endif
-
-	auto ReparentObject = [this](UObject* Object)
-	{
-		if (Object && !Object->HasAllFlags(RF_Public | RF_Standalone) && (Object->GetOuter() != GetOuter()) && (Object->GetOutermost() == GetOutermost()))
-		{
-			Object->Rename(nullptr, GetOuter(), REN_ForceNoResetLoaders);
-			return true;
-		}
-		return false;
-	};
-
-	ReparentObject(HeightmapTexture);
-	ReparentObject(XYOffsetmapTexture);
-
-	for (UTexture2D* WeightmapTexture : WeightmapTextures)
-	{
-		ReparentObject(WeightmapTexture);
-	}
-
-	for (UTexture2D* MobileWeightmapTexture : MobileWeightmapTextures)
-	{
-		ReparentObject(MobileWeightmapTexture);
-	}
-
-	for (auto& ItPair : LayersData)
-	{
-		FLandscapeLayerComponentData& LayerComponentData = ItPair.Value;
-		ReparentObject(LayerComponentData.HeightmapData.Texture);
-		for (UTexture2D* WeightmapTexture : LayerComponentData.WeightmapData.Textures)
-		{
-			ReparentObject(WeightmapTexture);
-		}
-
-		// Fixup missing/mismatching edit layer names :
-		if (const FLandscapeLayer* EditLayer = GetLandscapeActor() ? GetLandscapeActor()->GetLayer(ItPair.Key) : nullptr)
-		{
-			if (LayerComponentData.DebugName != EditLayer->Name)
-			{
-				LayerComponentData.DebugName = EditLayer->Name;
-			}
-		}
-	}
-
-	for (UMaterialInstance* MaterialInstance : MaterialInstances)
-	{
-		ULandscapeMaterialInstanceConstant* CurrentMIC = Cast<ULandscapeMaterialInstanceConstant>(MaterialInstance);
-		while (ReparentObject(CurrentMIC))
-		{
-			CurrentMIC = Cast<ULandscapeMaterialInstanceConstant>(MaterialInstance->Parent);
-		}
-	}
-
-	for (UMaterialInterface* MobileMaterialInterface : MobileMaterialInterfaces)
-	{
-		while (ReparentObject(MobileMaterialInterface))
-		{
-			MobileMaterialInterface = Cast<UMaterialInstance>(MobileMaterialInterface) ? Cast<UMaterialInstance>(((UMaterialInstance*)MobileMaterialInterface)->Parent) : nullptr;
-		}
-	}
-
-	for (UMaterialInstance* MobileCombinationMaterialInstance : MobileCombinationMaterialInstances)
-	{
-		while (ReparentObject(MobileCombinationMaterialInstance))
-		{
-			MobileCombinationMaterialInstance = Cast<UMaterialInstance>(MobileCombinationMaterialInstance->Parent);
-		}
-	}
-
-#if !UE_BUILD_SHIPPING
-	// This will fix the data in case there is mismatch between save of asset/maps
-	const int8 MaxLOD = static_cast<int8>(FMath::CeilLogTwo(SubsectionSizeQuads + 1) - 1);
-
-	TArray<ULandscapeMaterialInstanceConstant*> ResolvedMaterials;
-
-	if (LODIndexToMaterialIndex.Num() != MaxLOD+1)
-	{
-		if (GIsEditor)
-		{
-			UpdateMaterialInstances();
-		}
-		else
-		{
-			// Correct in-place differences by applying the highest LOD value we have to the newly added items as most case will be missing items added at the end
-			LODIndexToMaterialIndex.SetNumZeroed(MaxLOD + 1);
-
-			int8 LastLODIndex = 0;
-
-			for (int32 i = 0; i < LODIndexToMaterialIndex.Num(); ++i)
-			{
-				if (LODIndexToMaterialIndex[i] > LastLODIndex)
-				{
-					LastLODIndex = LODIndexToMaterialIndex[i];
-				}
-
-				if (LODIndexToMaterialIndex[i] == 0 && LastLODIndex != 0)
-				{
-					LODIndexToMaterialIndex[i] = LastLODIndex;
-				}
-			}
-		}
-	}
-#endif // UE_BUILD_SHIPPING
-
-	if (GIsEditor && !HasAnyFlags(RF_ClassDefaultObject))
-	{
-		// Move the MICs and Textures back to the Package if they're currently in the level
-		// Moving them into the level caused them to be duplicated when running PIE, which is *very very slow*, so we've reverted that change
-		// Also clear the public flag to avoid various issues, e.g. generating and saving thumbnails that can never be seen
-		if (ULevel* Level = GetLevel())
-		{
-			TArray<UObject*> ObjectsToMoveFromLevelToPackage;
-			GetGeneratedTexturesAndMaterialInstances(ObjectsToMoveFromLevelToPackage);
-
-			UPackage* MyPackage = GetOutermost();
-			for (auto* Obj : ObjectsToMoveFromLevelToPackage)
-			{
-				Obj->ClearFlags(RF_Public);
-				if (Obj->GetOuter() == Level)
-				{
-					Obj->Rename(nullptr, MyPackage, REN_DoNotDirty | REN_DontCreateRedirectors | REN_ForceNoResetLoaders | REN_NonTransactional);
-				}
-			}
-		}
-	}
-
-#if !UE_BUILD_SHIPPING
-	if (MobileCombinationMaterialInstances.Num() == 0)
-	{
-		if (GIsEditor)
-		{
-			UpdateMaterialInstances();
-		}
-		else
-		{
-			UE_LOG(LogLandscape, Error, TEXT("Landscape component (%d, %d) Does not have a valid mobile combination material. To correct this issue, open the map in the editor and resave the map."), SectionBaseX, SectionBaseY);
-		}
-	}
-#endif // UE_BUILD_SHIPPING
-
-	if (!HasAnyFlags(RF_ClassDefaultObject))
-	{
-		UWorld* World = GetWorld();
-		ERHIFeatureLevel::Type FeatureLevel = ((GEngine->GetDefaultWorldFeatureLevel() == ERHIFeatureLevel::ES3_1) || (World && (World->GetFeatureLevel() <= ERHIFeatureLevel::ES3_1))) 
-			? ERHIFeatureLevel::ES3_1 : GMaxRHIFeatureLevel;
-
-		// If we're loading on a platform that doesn't require cooked data, but defaults to a mobile feature level, generate or preload data from the DDC
-		if (!FPlatformProperties::RequiresCookedData() && FeatureLevel == ERHIFeatureLevel::ES3_1)
-		{
-			CheckGenerateMobilePlatformData(/*bIsCooking = */ false, /*TargetPlatform = */ nullptr);
-		}
-	}
-
-	GrassData->ConditionalDiscardDataOnLoad();
-
-#if WITH_EDITORONLY_DATA
-	// If the Collision Component is not set yet and we're transferring the property from the lazy object pointer it was previously stored as to the soft object ptr it is now stored as :
-	if (!CollisionComponentRef && CollisionComponent_DEPRECATED.IsValid())
-	{
-		CollisionComponentRef = CollisionComponent_DEPRECATED.Get();
-		CollisionComponent_DEPRECATED = nullptr;
-	}
-#endif // !WITH_EDITORONLY_DATA
-
-#endif // WITH_EDITOR
+	DelayedPostLoadState = UE::Landscape::EDelayedPostLoadState::Requested;
 }
 
 #if WITH_EDITORONLY_DATA
@@ -1386,11 +1103,13 @@ ALandscape* ALandscape::GetLandscapeActor()
 
 const ALandscape* ALandscapeStreamingProxy::GetLandscapeActor() const
 {
+	checkf(!FUObjectThreadContext::Get().IsRoutingPostLoad, TEXT("This method should not be called during PostLoad since the LandscapeActor may not be fully unserialized yet."));
 	return LandscapeActorRef.Get();
 }
 
 ALandscape* ALandscapeStreamingProxy::GetLandscapeActor()
 {
+	checkf(!FUObjectThreadContext::Get().IsRoutingPostLoad, TEXT("This method should not be called during PostLoad since the LandscapeActor may not be fully unserialized yet."));
 	return LandscapeActorRef.Get();
 }
 
@@ -1726,6 +1445,8 @@ FAutoConsoleVariableSink OnStaticMeshLODDistanceScaleChangedSink(FConsoleCommand
 
 void ULandscapeComponent::OnRegister()
 {
+	TryDelayedPostLoad();
+
 	Super::OnRegister();
 
 	if (GetLandscapeProxy())
@@ -2300,6 +2021,8 @@ void ULandscapeComponent::DeleteLayerAllocation(const FGuid& InEditLayerGuid, in
 
 void ALandscapeProxy::PostRegisterAllComponents()
 {
+	TryDelayedPostLoad();
+
 	Super::PostRegisterAllComponents();
 
 	ULandscapeInfo* LandscapeInfo = nullptr;
@@ -2361,55 +2084,6 @@ void ALandscapeProxy::PostRegisterAllComponents()
 	{
 		LandscapeInfo->FixupProxiesTransform();
 	}
-
-	bool bFixedUpInvalidMaterialInstances = false;
-
-	for (ULandscapeComponent* Comp : LandscapeComponents)
-	{
-		if (Comp == nullptr)
-		{
-			continue;
-		}
-
-		// Validate the layer combination and store it in the MaterialInstanceConstantMap
-		UMaterialInstance* MaterialInstance = Comp->GetMaterialInstance(0, false);
-
-		if (MaterialInstance == nullptr)
-		{
-			continue;
-		}
-
-		UMaterialInstanceConstant* CombinationMaterialInstance = Cast<UMaterialInstanceConstant>(MaterialInstance->Parent);
-		// Only validate if uncooked and in the editor/commandlet mode (we cannot re-build material instance constants if this is not the case : see UMaterialInstance::CacheResourceShadersForRendering, which is only called if FApp::CanEverRender() returns true) 
-		if (!Comp->GetOutermost()->HasAnyPackageFlags(PKG_FilterEditorOnly) && (GIsEditor && FApp::CanEverRender()))
-		{
-			if (Comp->ValidateCombinationMaterial(CombinationMaterialInstance))
-			{
-				MaterialInstanceConstantMap.Add(*ULandscapeComponent::GetLayerAllocationKey(Comp->GetWeightmapLayerAllocations(), CombinationMaterialInstance->Parent), CombinationMaterialInstance);
-			}
-			else
-			{
-				// There was a problem with the loaded material : it doesn't match the expected material combination, we need to regenerate the material instances : 
-				Comp->UpdateMaterialInstances();
-				bFixedUpInvalidMaterialInstances = true;
-			}
-		}
-		else if (CombinationMaterialInstance)
-		{
-			// Skip ValidateCombinationMaterial
-			MaterialInstanceConstantMap.Add(*ULandscapeComponent::GetLayerAllocationKey(Comp->GetWeightmapLayerAllocations(), CombinationMaterialInstance->Parent), CombinationMaterialInstance);
-		}
-	}
-
-	if (bFixedUpInvalidMaterialInstances)
-	{
-		FFormatNamedArguments Arguments;
-		Arguments.Add(TEXT("LandscapeName"), FText::FromString(GetPathName()));
-		Arguments.Add(TEXT("ProxyPackage"), FText::FromString(GetOutermost()->GetName()));
-		FMessageLog("MapCheck").Info()
-			->AddToken(FTextToken::Create(FText::Format(LOCTEXT("MapCheck_Message_FixedUpInvalidLandscapeMaterialInstances", "{LandscapeName} : Fixed up invalid landscape material instances. Please re-save {ProxyPackage}."), Arguments)))
-			->AddToken(FMapErrorToken::Create(FMapErrors::FixedUpInvalidLandscapeMaterialInstances));
-	}
 #endif // WITH_EDITOR
 }
 
@@ -3435,6 +3109,16 @@ void ALandscapeProxy::PostLoad()
 {
 	Super::PostLoad();
 
+	DelayedPostLoadState = UE::Landscape::EDelayedPostLoadState::Requested;
+}
+
+void ALandscapeProxy::TryDelayedPostLoad()
+{
+	if (DelayedPostLoadState != UE::Landscape::EDelayedPostLoadState::Requested)
+	{
+		return;
+	}
+
 	LandscapeInstancing::FixupLandscapeGuidsIfInstanced(this);
 
 	// Temporary
@@ -3467,25 +3151,10 @@ void ALandscapeProxy::PostLoad()
 	}
 	PRAGMA_ENABLE_DEPRECATION_WARNINGS;
 
+	// At this point, our component should have been loaded and registered (was previously enforced by calling ConditionalPostLoad on the components).
 	if (GIsEditor)
 	{
-		// We may not have run PostLoad on LandscapeComponents yet
-		for (TObjectPtr<ULandscapeComponent>& LandscapeComponent : LandscapeComponents)
-		{
-			if (LandscapeComponent)
-			{
-				LandscapeComponent->ConditionalPostLoad();
-			}
-		}
-		
-		// We may not have run PostLoad on CollisionComponent yet
-		for (TObjectPtr<ULandscapeHeightfieldCollisionComponent>& CollisionComponent : CollisionComponents)
-		{
-			if (CollisionComponent)
-			{
-				CollisionComponent->ConditionalPostLoad();
-			}
-		}
+
 
 		if ((GetLinker() && (GetLinker()->UEVer() < VER_UE4_LANDSCAPE_COMPONENT_LAZY_REFERENCES)) ||
 			LandscapeComponents.Num() != CollisionComponents.Num() ||
@@ -3508,6 +3177,55 @@ void ALandscapeProxy::PostLoad()
 		EditorCachedLayerInfos_DEPRECATED.Empty();
 	}
 
+	bool bFixedUpInvalidMaterialInstances = false;
+
+	for (ULandscapeComponent* Comp : LandscapeComponents)
+	{
+		if (Comp == nullptr)
+		{
+			continue;
+		}
+
+		// Validate the layer combination and store it in the MaterialInstanceConstantMap
+		UMaterialInstance* MaterialInstance = Comp->GetMaterialInstance(0, false);
+
+		if (MaterialInstance == nullptr)
+		{
+			continue;
+		}
+
+		UMaterialInstanceConstant* CombinationMaterialInstance = Cast<UMaterialInstanceConstant>(MaterialInstance->Parent);
+		// Only validate if uncooked and in the editor/commandlet mode (we cannot re-build material instance constants if this is not the case : see UMaterialInstance::CacheResourceShadersForRendering, which is only called if FApp::CanEverRender() returns true) 
+		if (!Comp->GetOutermost()->HasAnyPackageFlags(PKG_FilterEditorOnly) && (GIsEditor && FApp::CanEverRender()))
+		{
+			if (Comp->ValidateCombinationMaterial(CombinationMaterialInstance))
+			{
+				MaterialInstanceConstantMap.Add(*ULandscapeComponent::GetLayerAllocationKey(Comp->GetWeightmapLayerAllocations(), CombinationMaterialInstance->Parent), CombinationMaterialInstance);
+			}
+			else
+			{
+				// There was a problem with the loaded material : it doesn't match the expected material combination, we need to regenerate the material instances : 
+				Comp->UpdateMaterialInstances();
+				bFixedUpInvalidMaterialInstances = true;
+			}
+		}
+		else if (CombinationMaterialInstance)
+		{
+			// Skip ValidateCombinationMaterial
+			MaterialInstanceConstantMap.Add(*ULandscapeComponent::GetLayerAllocationKey(Comp->GetWeightmapLayerAllocations(), CombinationMaterialInstance->Parent), CombinationMaterialInstance);
+		}
+	}
+
+	if (bFixedUpInvalidMaterialInstances)
+	{
+		FFormatNamedArguments Arguments;
+		Arguments.Add(TEXT("LandscapeName"), FText::FromString(GetPathName()));
+		Arguments.Add(TEXT("ProxyPackage"), FText::FromString(GetOutermost()->GetName()));
+		FMessageLog("MapCheck").Info()
+			->AddToken(FTextToken::Create(FText::Format(LOCTEXT("MapCheck_Message_FixedUpInvalidLandscapeMaterialInstances", "{LandscapeName} : Fixed up invalid landscape material instances. Please re-save {ProxyPackage}."), Arguments)))
+			->AddToken(FMapErrorToken::Create(FMapErrors::FixedUpInvalidLandscapeMaterialInstances));
+	}
+
 	UWorld* World = GetWorld();
 
 	// track feature level change to flush grass cache
@@ -3532,6 +3250,8 @@ void ALandscapeProxy::PostLoad()
 #endif // WITH_EDITOR
 
 	UpdateRenderingMethod();
+
+	DelayedPostLoadState = UE::Landscape::EDelayedPostLoadState::Processed;
 }
 
 FIntPoint ALandscapeProxy::GetSectionBaseOffset() const
diff --git a/Engine/Source/Runtime/Landscape/Private/LandscapeEdit.cpp b/Engine/Source/Runtime/Landscape/Private/LandscapeEdit.cpp
--- a/Engine/Source/Runtime/Landscape/Private/LandscapeEdit.cpp
+++ b/Engine/Source/Runtime/Landscape/Private/LandscapeEdit.cpp
@@ -7324,8 +7324,303 @@ bool ALandscapeProxy::LandscapeImportHeightmapFromRenderTarget(UTextureRenderTar
 
 	return true;
 }
+#endif // WITH_EDITOR
+
+void ULandscapeComponent::TryDelayedPostLoad()
+{
+	if (DelayedPostLoadState != UE::Landscape::EDelayedPostLoadState::Requested)
+	{
+		return;
+	}
+
+	if (IsComponentPSOPrecachingEnabled())
+	{
+		TArray<UMaterialInterface*> Materials;
+		bool bGetDebugMaterials = false;
+		GetUsedMaterials(Materials, bGetDebugMaterials);
+
+		FPSOPrecacheParams PrecachePSOParams;
+		SetupPrecachePSOParams(PrecachePSOParams);
+
+		FPSOPrecacheVertexFactoryDataList VertexFactoryDataList;
+		if (!XYOffsetmapTexture)
+		{
+			VertexFactoryDataList.Add(FPSOPrecacheVertexFactoryData(&FLandscapeVertexFactory::StaticType));
+		}
+		else
+		{
+			VertexFactoryDataList.Add(FPSOPrecacheVertexFactoryData(&FLandscapeXYOffsetVertexFactory::StaticType));
+		}
+
+		// we need the fixed grid vertex factory for both virtual texturing and grass
+		ERHIFeatureLevel::Type FeatureLevel = GMaxRHIFeatureLevel;
+		bool bNeedsFixedGridVertexFactory = UseVirtualTexturing(FeatureLevel);
+		if (bNeedsFixedGridVertexFactory)
+		{
+			VertexFactoryDataList.Add(FPSOPrecacheVertexFactoryData(&FLandscapeFixedGridVertexFactory::StaticType));
+		}
+
+		TArray<FMaterialPSOPrecacheRequestID> MaterialPrecacheRequestIDs;
+		for (UMaterialInterface* MaterialInterface : Materials)
+		{
+			if (MaterialInterface)
+			{
+				check(!MaterialInterface->HasAnyFlags(RF_NeedPostLoad));
+				MaterialInterface->PrecachePSOs(VertexFactoryDataList, PrecachePSOParams, EPSOPrecachePriority::High, MaterialPrecacheRequestIDs);
+			}
+		}
+	}
+
+#if WITH_EDITOR
+	ALandscapeProxy* LandscapeProxy = GetLandscapeProxy();
+	if (ensure(LandscapeProxy))
+	{
+		// Ensure that the component's lighting settings matches the actor's.
+		UpdatedSharedPropertiesFromActor();
+
+		// check SectionBaseX/Y are correct
+		const FVector LocalRelativeLocation = GetRelativeLocation();
+		int32 CheckSectionBaseX = FMath::RoundToInt32(LocalRelativeLocation.X) + LandscapeProxy->LandscapeSectionOffset.X;
+		int32 CheckSectionBaseY = FMath::RoundToInt32(LocalRelativeLocation.Y) + LandscapeProxy->LandscapeSectionOffset.Y;
+		if (CheckSectionBaseX != SectionBaseX ||
+			CheckSectionBaseY != SectionBaseY)
+		{
+			UE_LOG(LogLandscape, Warning, TEXT("LandscapeComponent SectionBaseX disagrees with its location, attempted automated fix: '%s', %d,%d vs %d,%d."),
+				*GetFullName(), SectionBaseX, SectionBaseY, CheckSectionBaseX, CheckSectionBaseY);
+			SectionBaseX = CheckSectionBaseX;
+			SectionBaseY = CheckSectionBaseY;
+		}
+	}
+
+	if (GIsEditor && !HasAnyFlags(RF_ClassDefaultObject))
+	{
+		// This is to ensure that component relative location is exact section base offset value
+		FVector LocalRelativeLocation = GetRelativeLocation();
+		float CheckRelativeLocationX = float(SectionBaseX - LandscapeProxy->LandscapeSectionOffset.X);
+		float CheckRelativeLocationY = float(SectionBaseY - LandscapeProxy->LandscapeSectionOffset.Y);
+		if (!FMath::IsNearlyEqual(CheckRelativeLocationX, LocalRelativeLocation.X, UE_DOUBLE_KINDA_SMALL_NUMBER) ||
+			!FMath::IsNearlyEqual(CheckRelativeLocationY, LocalRelativeLocation.Y, UE_DOUBLE_KINDA_SMALL_NUMBER))
+		{
+			UE_LOG(LogLandscape, Warning, TEXT("LandscapeComponent RelativeLocation disagrees with its section base, attempted automated fix: '%s', %f,%f vs %f,%f."),
+				*GetFullName(), LocalRelativeLocation.X, LocalRelativeLocation.Y, CheckRelativeLocationX, CheckRelativeLocationY);
+			LocalRelativeLocation.X = CheckRelativeLocationX;
+			LocalRelativeLocation.Y = CheckRelativeLocationY;
+
+			SetRelativeLocation_Direct(LocalRelativeLocation);
+		}
+
+		// Remove standalone flags from data textures to ensure data is unloaded in the editor when reverting an unsaved level.
+		// Previous version of landscape set these flags on creation.
+		if (HeightmapTexture && HeightmapTexture->HasAnyFlags(RF_Standalone))
+		{
+			HeightmapTexture->ClearFlags(RF_Standalone);
+		}
+		for (int32 Idx = 0; Idx < WeightmapTextures.Num(); Idx++)
+		{
+			if (WeightmapTextures[Idx] && WeightmapTextures[Idx]->HasAnyFlags(RF_Standalone))
+			{
+				WeightmapTextures[Idx]->ClearFlags(RF_Standalone);
+			}
+		}
+
+		LastSavedPhysicalMaterialHash = PhysicalMaterialHash;
+
+		PRAGMA_DISABLE_DEPRECATION_WARNINGS;
+		if (!OverrideMaterials_DEPRECATED.IsEmpty())
+		{
+			PerLODOverrideMaterials.Reserve(OverrideMaterials_DEPRECATED.Num());
+			for (const FLandscapeComponentMaterialOverride& LocalMaterialOverride : OverrideMaterials_DEPRECATED)
+			{
+				PerLODOverrideMaterials.Add({ LocalMaterialOverride.LODIndex.Default, LocalMaterialOverride.Material });
+			}
+			OverrideMaterials_DEPRECATED.Reset();
+		}
+		PRAGMA_ENABLE_DEPRECATION_WARNINGS;
+	}
+
+#if WITH_EDITORONLY_DATA
+	// Handle old MaterialInstance
+	if (MaterialInstance_DEPRECATED)
+	{
+		MaterialInstances.Empty(1);
+		MaterialInstances.Add(MaterialInstance_DEPRECATED);
+		MaterialInstance_DEPRECATED = nullptr;
+
+		if (GIsEditor && MaterialInstances.Num() > 0 && MaterialInstances[0] != nullptr)
+		{
+			MaterialInstances[0]->ConditionalPostLoad();
+			UpdateMaterialInstances();
+		}
+	}
 #endif
 
+	auto ReparentObject = [this](UObject* Object)
+	{
+		if (Object && !Object->HasAllFlags(RF_Public | RF_Standalone) && (Object->GetOuter() != GetOuter()) && (Object->GetOutermost() == GetOutermost()))
+		{
+			Object->Rename(nullptr, GetOuter(), REN_ForceNoResetLoaders);
+			return true;
+		}
+		return false;
+	};
+
+	ReparentObject(HeightmapTexture);
+	ReparentObject(XYOffsetmapTexture);
+
+	for (UTexture2D* WeightmapTexture : WeightmapTextures)
+	{
+		ReparentObject(WeightmapTexture);
+	}
+
+	for (UTexture2D* MobileWeightmapTexture : MobileWeightmapTextures)
+	{
+		ReparentObject(MobileWeightmapTexture);
+	}
+
+	for (auto& ItPair : LayersData)
+	{
+		FLandscapeLayerComponentData& LayerComponentData = ItPair.Value;
+		ReparentObject(LayerComponentData.HeightmapData.Texture);
+		for (UTexture2D* WeightmapTexture : LayerComponentData.WeightmapData.Textures)
+		{
+			ReparentObject(WeightmapTexture);
+		}
+
+		// Fixup missing/mismatching edit layer names :
+		if (const FLandscapeLayer* EditLayer = GetLandscapeActor() ? GetLandscapeActor()->GetLayer(ItPair.Key) : nullptr)
+		{
+			if (LayerComponentData.DebugName != EditLayer->Name)
+			{
+				LayerComponentData.DebugName = EditLayer->Name;
+			}
+		}
+	}
+
+	for (UMaterialInstance* MaterialInstance : MaterialInstances)
+	{
+		ULandscapeMaterialInstanceConstant* CurrentMIC = Cast<ULandscapeMaterialInstanceConstant>(MaterialInstance);
+		while (ReparentObject(CurrentMIC))
+		{
+			CurrentMIC = Cast<ULandscapeMaterialInstanceConstant>(MaterialInstance->Parent);
+		}
+	}
+
+	for (UMaterialInterface* MobileMaterialInterface : MobileMaterialInterfaces)
+	{
+		while (ReparentObject(MobileMaterialInterface))
+		{
+			MobileMaterialInterface = Cast<UMaterialInstance>(MobileMaterialInterface) ? Cast<UMaterialInstance>(((UMaterialInstance*)MobileMaterialInterface)->Parent) : nullptr;
+		}
+	}
+
+	for (UMaterialInstance* MobileCombinationMaterialInstance : MobileCombinationMaterialInstances)
+	{
+		while (ReparentObject(MobileCombinationMaterialInstance))
+		{
+			MobileCombinationMaterialInstance = Cast<UMaterialInstance>(MobileCombinationMaterialInstance->Parent);
+		}
+	}
+
+#if !UE_BUILD_SHIPPING
+	// This will fix the data in case there is mismatch between save of asset/maps
+	const int8 MaxLOD = static_cast<int8>(FMath::CeilLogTwo(SubsectionSizeQuads + 1) - 1);
+
+	TArray<ULandscapeMaterialInstanceConstant*> ResolvedMaterials;
+
+	if (LODIndexToMaterialIndex.Num() != MaxLOD + 1)
+	{
+		if (GIsEditor)
+		{
+			UpdateMaterialInstances();
+		}
+		else
+		{
+			// Correct in-place differences by applying the highest LOD value we have to the newly added items as most case will be missing items added at the end
+			LODIndexToMaterialIndex.SetNumZeroed(MaxLOD + 1);
+
+			int8 LastLODIndex = 0;
+
+			for (int32 i = 0; i < LODIndexToMaterialIndex.Num(); ++i)
+			{
+				if (LODIndexToMaterialIndex[i] > LastLODIndex)
+				{
+					LastLODIndex = LODIndexToMaterialIndex[i];
+				}
+
+				if (LODIndexToMaterialIndex[i] == 0 && LastLODIndex != 0)
+				{
+					LODIndexToMaterialIndex[i] = LastLODIndex;
+				}
+			}
+		}
+	}
+#endif // UE_BUILD_SHIPPING
+
+	if (GIsEditor && !HasAnyFlags(RF_ClassDefaultObject))
+	{
+		// Move the MICs and Textures back to the Package if they're currently in the level
+		// Moving them into the level caused them to be duplicated when running PIE, which is *very very slow*, so we've reverted that change
+		// Also clear the public flag to avoid various issues, e.g. generating and saving thumbnails that can never be seen
+		if (ULevel* Level = GetLevel())
+		{
+			TArray<UObject*> ObjectsToMoveFromLevelToPackage;
+			GetGeneratedTexturesAndMaterialInstances(ObjectsToMoveFromLevelToPackage);
+
+			UPackage* MyPackage = GetOutermost();
+			for (auto* Obj : ObjectsToMoveFromLevelToPackage)
+			{
+				Obj->ClearFlags(RF_Public);
+				if (Obj->GetOuter() == Level)
+				{
+					Obj->Rename(nullptr, MyPackage, REN_DoNotDirty | REN_DontCreateRedirectors | REN_ForceNoResetLoaders | REN_NonTransactional);
+				}
+			}
+		}
+	}
+
+#if !UE_BUILD_SHIPPING
+	if (MobileCombinationMaterialInstances.Num() == 0)
+	{
+		if (GIsEditor)
+		{
+			UpdateMaterialInstances();
+		}
+		else
+		{
+			UE_LOG(LogLandscape, Error, TEXT("Landscape component (%d, %d) Does not have a valid mobile combination material. To correct this issue, open the map in the editor and resave the map."), SectionBaseX, SectionBaseY);
+		}
+	}
+#endif // UE_BUILD_SHIPPING
+
+	if (!HasAnyFlags(RF_ClassDefaultObject))
+	{
+		UWorld* World = GetWorld();
+		ERHIFeatureLevel::Type FeatureLevel = ((GEngine->GetDefaultWorldFeatureLevel() == ERHIFeatureLevel::ES3_1) || (World && (World->FeatureLevel <= ERHIFeatureLevel::ES3_1)))
+			? ERHIFeatureLevel::ES3_1 : GMaxRHIFeatureLevel;
+
+		// If we're loading on a platform that doesn't require cooked data, but defaults to a mobile feature level, generate or preload data from the DDC
+		if (!FPlatformProperties::RequiresCookedData() && FeatureLevel == ERHIFeatureLevel::ES3_1)
+		{
+			CheckGenerateMobilePlatformData(/*bIsCooking = */ false, /*TargetPlatform = */ nullptr);
+		}
+	}
+
+	GrassData->ConditionalDiscardDataOnLoad();
+
+#if WITH_EDITORONLY_DATA
+	// If the Collision Component is not set yet and we're transferring the property from the lazy object pointer it was previously stored as to the soft object ptr it is now stored as :
+	if (!CollisionComponentRef && CollisionComponent_DEPRECATED.IsValid())
+	{
+		CollisionComponentRef = CollisionComponent_DEPRECATED.Get();
+		CollisionComponent_DEPRECATED = nullptr;
+	}
+#endif // !WITH_EDITORONLY_DATA
+
+#endif // WITH_EDITOR
+
+	DelayedPostLoadState = UE::Landscape::EDelayedPostLoadState::Processed;
+}
+
 bool ALandscapeProxy::LandscapeExportHeightmapToRenderTarget(UTextureRenderTarget2D* InRenderTarget, bool bInExportHeightIntoRGChannel, bool InExportLandscapeProxies)
 {
 #if WITH_EDITOR
diff --git a/Engine/Source/Runtime/Landscape/Public/LandscapeUtils.h b/Engine/Source/Runtime/Landscape/Public/LandscapeUtils.h
--- a/Engine/Source/Runtime/Landscape/Public/LandscapeUtils.h
+++ b/Engine/Source/Runtime/Landscape/Public/LandscapeUtils.h
@@ -18,6 +18,13 @@ class UTexture2D;
 namespace UE::Landscape
 {
 
+enum class EDelayedPostLoadState : uint8
+{
+	Default = 0,
+	Requested,
+	Processed
+};
+
 /**
 * Returns true if edit layers (GPU landscape tools) are enabled on this platform :
 * Note: this is intended for the editor but is in runtime code since global shaders need to exist in runtime modules 
