Fixes deprecation error for World->FeatureLevel

diff --git a/Engine/Source/Runtime/Landscape/Private/LandscapeEdit.cpp b/Engine/Source/Runtime/Landscape/Private/LandscapeEdit.cpp
--- a/Engine/Source/Runtime/Landscape/Private/LandscapeEdit.cpp
+++ b/Engine/Source/Runtime/Landscape/Private/LandscapeEdit.cpp
@@ -7594,8 +7594,9 @@ void ULandscapeComponent::TryDelayedPostLoad()
 
 	if (!HasAnyFlags(RF_ClassDefaultObject))
 	{
-		UWorld* World = GetWorld();
-		ERHIFeatureLevel::Type FeatureLevel = ((GEngine->GetDefaultWorldFeatureLevel() == ERHIFeatureLevel::ES3_1) || (World && (World->FeatureLevel <= ERHIFeatureLevel::ES3_1)))
+		FSceneInterface* SceneInterface = GetScene();
+
+		ERHIFeatureLevel::Type FeatureLevel = ((GEngine->GetDefaultWorldFeatureLevel() == ERHIFeatureLevel::ES3_1) || (SceneInterface && (SceneInterface->GetFeatureLevel() <= ERHIFeatureLevel::ES3_1)))
 			? ERHIFeatureLevel::ES3_1 : GMaxRHIFeatureLevel;
 
 		// If we're loading on a platform that doesn't require cooked data, but defaults to a mobile feature level, generate or preload data from the DDC
