#UE Finish backout of 24823203 in Release-Engine-Staging

diff --git a/Engine/Source/Runtime/Landscape/Private/Landscape.cpp b/Engine/Source/Runtime/Landscape/Private/Landscape.cpp
--- a/Engine/Source/Runtime/Landscape/Private/Landscape.cpp
+++ b/Engine/Source/Runtime/Landscape/Private/Landscape.cpp
@@ -1198,7 +1198,7 @@ void ULandscapeComponent::PostLoad()
 	if (!HasAnyFlags(RF_ClassDefaultObject))
 	{
 		UWorld* World = GetWorld();
-		ERHIFeatureLevel::Type FeatureLevel = ((GEngine->GetDefaultWorldFeatureLevel() == ERHIFeatureLevel::ES3_1) || (World && (World->FeatureLevel <= ERHIFeatureLevel::ES3_1))) 
+		ERHIFeatureLevel::Type FeatureLevel = ((GEngine->GetDefaultWorldFeatureLevel() == ERHIFeatureLevel::ES3_1) || (World && (World->GetFeatureLevel() <= ERHIFeatureLevel::ES3_1))) 
 			? ERHIFeatureLevel::ES3_1 : GMaxRHIFeatureLevel;
 
 		// If we're loading on a platform that doesn't require cooked data, but defaults to a mobile feature level, generate or preload data from the DDC
