#UE Finish resolved backout of 24823203 in Dev-EngineMerge

diff --git a/Engine/Source/Runtime/Landscape/Private/Landscape.cpp b/Engine/Source/Runtime/Landscape/Private/Landscape.cpp
--- a/Engine/Source/Runtime/Landscape/Private/Landscape.cpp
+++ b/Engine/Source/Runtime/Landscape/Private/Landscape.cpp
@@ -983,8 +983,8 @@ void ULandscapeComponent::PostLoad()
 
 		// check SectionBaseX/Y are correct
 		const FVector LocalRelativeLocation = GetRelativeLocation();
-		int32 CheckSectionBaseX = FMath::RoundToInt(LocalRelativeLocation.X) + LandscapeProxy->LandscapeSectionOffset.X;
-		int32 CheckSectionBaseY = FMath::RoundToInt(LocalRelativeLocation.Y) + LandscapeProxy->LandscapeSectionOffset.Y;
+		int32 CheckSectionBaseX = FMath::RoundToInt32(LocalRelativeLocation.X) + LandscapeProxy->LandscapeSectionOffset.X;
+		int32 CheckSectionBaseY = FMath::RoundToInt32(LocalRelativeLocation.Y) + LandscapeProxy->LandscapeSectionOffset.Y;
 		if (CheckSectionBaseX != SectionBaseX ||
 			CheckSectionBaseY != SectionBaseY)
 		{
@@ -1126,7 +1126,7 @@ void ULandscapeComponent::PostLoad()
 
 #if !UE_BUILD_SHIPPING
 	// This will fix the data in case there is mismatch between save of asset/maps
-	int8 MaxLOD = FMath::CeilLogTwo(SubsectionSizeQuads + 1) - 1;
+	const int8 MaxLOD = static_cast<int8>(FMath::CeilLogTwo(SubsectionSizeQuads + 1) - 1);
 
 	TArray<ULandscapeMaterialInstanceConstant*> ResolvedMaterials;
 
