Fixed landscape potentially dispatching a newly created physics actor and then subsequently destroying the geometry it uses before the next physics update leading to a crash due to use-after-free.
Dispatch a deferred reference counter to extend the lifetime of a geometry beyond the point at which the physics thread has finished using it.

diff --git a/Engine/Source/Runtime/Landscape/Private/LandscapeCollision.cpp b/Engine/Source/Runtime/Landscape/Private/LandscapeCollision.cpp
--- a/Engine/Source/Runtime/Landscape/Private/LandscapeCollision.cpp
+++ b/Engine/Source/Runtime/Landscape/Private/LandscapeCollision.cpp
@@ -61,6 +61,8 @@
 #include "PhysicsEngine/Experimental/ChaosCooking.h"
 #include "Chaos/ChaosArchive.h"
 #include "PhysicsProxy/SingleParticlePhysicsProxy.h"
+#include "Chaos/Framework/PhysicsSolverBase.h"
+#include "PBDRigidsSolver.h"
 
 using namespace PhysicsInterfaceTypes;
 
@@ -1764,6 +1766,36 @@ bool ULandscapeHeightfieldCollisionComponent::RecreateCollision()
 		RecreatePhysicsState();
 
 		MarkRenderStateDirty();
+
+		// If we have a world, then our physics state will be queued for processing on the physics thread
+		if(UWorld* World = GetWorld())
+		{
+			if(FPhysScene* PhysScene = World->GetPhysicsScene(); PhysScene && PhysScene->GetSolver())
+			{
+				// We could potentially call RecreateCollision multiple times before a physics update happens, especially
+				// if we're using the async tick mode for physics. In this case we would have a pending actor in the
+				// dirty proxy list on the physics thread with a geometry that has been destructed by the lifetime
+				// extender falling out of scope.
+				// To avoid this we dispatch an empty callable with the unique geometries which runs after the
+				// proxy queue will have been cleared, avoiding a use-after-free.
+				// #TODO auto ref counted user objects for Chaos.
+				PhysScene->GetSolver()->EnqueueCommandImmediate(
+					[ComplexHeightfield = MoveTemp(HeightfieldRefLifetimeExtender->Heightfield)
+					, SimpleHeightfield = MoveTemp(HeightfieldRefLifetimeExtender->HeightfieldSimple)
+#if WITH_EDITORONLY_DATA
+					, EditorHeightfield = MoveTemp(HeightfieldRefLifetimeExtender->EditorHeightfield)
+#endif
+					]
+				() mutable
+				{
+					ComplexHeightfield = nullptr;
+					SimpleHeightfield = nullptr;
+#if WITH_EDITORONLY_DATA
+					EditorHeightfield = nullptr;
+#endif
+				});
+			}
+		}
 	}
 	return true;
 }
