Moved navigation-related landscape parameters to their own category

#rnx

diff --git a/Engine/Source/Runtime/Landscape/Classes/LandscapeProxy.h b/Engine/Source/Runtime/Landscape/Classes/LandscapeProxy.h
--- a/Engine/Source/Runtime/Landscape/Classes/LandscapeProxy.h
+++ b/Engine/Source/Runtime/Landscape/Classes/LandscapeProxy.h
@@ -730,20 +730,20 @@ public:
 	int32 NumSubsections;    // Number of subsections in X and Y axis
 
 	/** Hints navigation system whether this landscape will ever be navigated on. true by default, but make sure to set it to false for faraway, background landscapes */
-	UPROPERTY(EditAnywhere, Category=Landscape)
+	UPROPERTY(EditAnywhere, Category=Navigation)
 	uint32 bUsedForNavigation:1;
 
 	/** Set to true to prevent navmesh generation under the terrain geometry */
-	UPROPERTY(EditAnywhere, Category = Landscape)
+	UPROPERTY(EditAnywhere, Category = Navigation)
 	uint32 bFillCollisionUnderLandscapeForNavmesh:1;
 
+	UPROPERTY(EditAnywhere, Category = Navigation, AdvancedDisplay)
+	ENavDataGatheringMode NavigationGeometryGatheringMode;
+
 	/** When set to true it will generate MaterialInstanceDynamic for each components, so material can be changed at runtime */
 	UPROPERTY(EditAnywhere, Category = Landscape)
 	bool bUseDynamicMaterialInstance;
 
-	UPROPERTY(EditAnywhere, Category = Landscape, AdvancedDisplay)
-	ENavDataGatheringMode NavigationGeometryGatheringMode;
-
 #if WITH_EDITORONLY_DATA
 	UPROPERTY(EditAnywhere, Category=Landscape)
 	int32 MaxPaintedLayersPerComponent; // 0 = disabled
