Deprecated unused landscape enum

#rnx

diff --git a/Engine/Source/Runtime/Landscape/Classes/Landscape.h b/Engine/Source/Runtime/Landscape/Classes/Landscape.h
--- a/Engine/Source/Runtime/Landscape/Classes/Landscape.h
+++ b/Engine/Source/Runtime/Landscape/Classes/Landscape.h
@@ -40,7 +40,7 @@ extern LANDSCAPE_API TAutoConsoleVariable<int32> CVarLandscapeSplineFalloffModul
 #endif
 
 UENUM()
-enum ELandscapeSetupErrors : int
+enum UE_DEPRECATED(5.3, "ELandscapeSetupErrors is now unused and deprecated") ELandscapeSetupErrors : int
 {
 	LSE_None,
 	/** No Landscape Info available. */
