Add null check on heightfield ref in landscape recreation

diff --git a/Engine/Source/Runtime/Landscape/Private/LandscapeCollision.cpp b/Engine/Source/Runtime/Landscape/Private/LandscapeCollision.cpp
--- a/Engine/Source/Runtime/Landscape/Private/LandscapeCollision.cpp
+++ b/Engine/Source/Runtime/Landscape/Private/LandscapeCollision.cpp
@@ -1770,7 +1770,7 @@ bool ULandscapeHeightfieldCollisionComponent::RecreateCollision()
 		// If we have a world, then our physics state will be queued for processing on the physics thread
 		if(UWorld* World = GetWorld())
 		{
-			if(FPhysScene* PhysScene = World->GetPhysicsScene(); PhysScene && PhysScene->GetSolver())
+			if(FPhysScene* PhysScene = World->GetPhysicsScene(); PhysScene && PhysScene->GetSolver() && HeightfieldRefLifetimeExtender)
 			{
 				// We could potentially call RecreateCollision multiple times before a physics update happens, especially
 				// if we're using the async tick mode for physics. In this case we would have a pending actor in the
