Expose building a landscape collision component to the runtime.

diff --git a/Engine/Source/Runtime/Experimental/Chaos/Private/Chaos/HeightField.cpp b/Engine/Source/Runtime/Experimental/Chaos/Private/Chaos/HeightField.cpp
--- a/Engine/Source/Runtime/Experimental/Chaos/Private/Chaos/HeightField.cpp
+++ b/Engine/Source/Runtime/Experimental/Chaos/Private/Chaos/HeightField.cpp
@@ -595,7 +595,7 @@ namespace Chaos
 
 
 	template<typename BufferType>
-	void BuildGeomData(TArrayView<BufferType> BufferView, TArrayView<uint8> MaterialIndexView, int32 NumRows, int32 NumCols, const FVec3& InScale, TUniqueFunction<FReal(const BufferType)> ToRealFunc, typename FHeightField::FDataType& OutData, FAABB3& OutBounds)
+	void BuildGeomData(TArrayView<BufferType> BufferView, TArrayView<const uint8> MaterialIndexView, int32 NumRows, int32 NumCols, const FVec3& InScale, TUniqueFunction<FReal(const BufferType)> ToRealFunc, typename FHeightField::FDataType& OutData, FAABB3& OutBounds)
 	{
 		using FDataType = typename FHeightField::FDataType;
 
@@ -770,7 +770,7 @@ namespace Chaos
 		SetScale(InScale);
 	}
 
-	Chaos::FHeightField::FHeightField(TArrayView<const uint16> InHeights, TArrayView<uint8> InMaterialIndices, int32 InNumRows, int32 InNumCols, const FVec3& InScale)
+	Chaos::FHeightField::FHeightField(TArrayView<const uint16> InHeights, TArrayView<const uint8> InMaterialIndices, int32 InNumRows, int32 InNumCols, const FVec3& InScale)
 		: FImplicitObject(EImplicitObject::HasBoundingBox, ImplicitObjectType::HeightField)
 	{
 		TUniqueFunction<FReal(const uint16)> ConversionFunc = [](const uint16 InVal) -> FReal
diff --git a/Engine/Source/Runtime/Experimental/Chaos/Public/Chaos/HeightField.h b/Engine/Source/Runtime/Experimental/Chaos/Public/Chaos/HeightField.h
--- a/Engine/Source/Runtime/Experimental/Chaos/Public/Chaos/HeightField.h
+++ b/Engine/Source/Runtime/Experimental/Chaos/Public/Chaos/HeightField.h
@@ -28,7 +28,7 @@ namespace Chaos
 		using FImplicitObject::GetTypeName;
 
 		FHeightField(TArray<FReal>&& Height, TArray<uint8>&& InMaterialIndices, int32 InNumRows, int32 InNumCols, const FVec3& InScale);
-		FHeightField(TArrayView<const uint16> InHeights, TArrayView<uint8> InMaterialIndices, int32 InNumRows, int32 InNumCols, const FVec3& InScale);
+		FHeightField(TArrayView<const uint16> InHeights, TArrayView<const uint8> InMaterialIndices, int32 InNumRows, int32 InNumCols, const FVec3& InScale);
 		FHeightField(const FHeightField& Other) = delete;
 		
 		// Not required as long as FImplicitObject also has deleted move constructor (adding this causes an error on Linux build)
diff --git a/Engine/Source/Runtime/Landscape/Classes/LandscapeHeightfieldCollisionComponent.h b/Engine/Source/Runtime/Landscape/Classes/LandscapeHeightfieldCollisionComponent.h
--- a/Engine/Source/Runtime/Landscape/Classes/LandscapeHeightfieldCollisionComponent.h
+++ b/Engine/Source/Runtime/Landscape/Classes/LandscapeHeightfieldCollisionComponent.h
@@ -279,6 +279,13 @@ public:
 	/** Creates collision object from a cooked collision data */
 	virtual void CreateCollisionObject();
 
+	/** Creates collision object from raw runtime data.  Data is assumed at this point to contain valid physical material indices with visiblility layer set to 0xFF*/
+	LANDSCAPE_API void CreateCollisionObject(
+		bool bUseDefaultMaterialOnly,
+		TArrayView<const uint16> Heights, TArrayView<const uint16> SimpleHeights,
+		TArrayView<const uint8> PhysicalMaterialIds, TArrayView<const uint8> SimplePhysicalMaterialIds,
+		TArrayView<const TObjectPtr<UPhysicalMaterial>> PhysicalMaterialObjects);
+
 	/** Return the landscape actor associated with this component. */
 	LANDSCAPE_API ALandscapeProxy* GetLandscapeProxy() const;
 
diff --git a/Engine/Source/Runtime/Landscape/Private/LandscapeCollision.cpp b/Engine/Source/Runtime/Landscape/Private/LandscapeCollision.cpp
--- a/Engine/Source/Runtime/Landscape/Private/LandscapeCollision.cpp
+++ b/Engine/Source/Runtime/Landscape/Private/LandscapeCollision.cpp
@@ -798,6 +798,91 @@ void ULandscapeHeightfieldCollisionComponent::CreateCollisionObject()
 	}
 }
 
+void ULandscapeHeightfieldCollisionComponent::CreateCollisionObject(
+	bool bUseDefaultMaterialOnly, 
+	TArrayView<const uint16> Heights, TArrayView<const uint16> SimpleHeights, 
+	TArrayView<const uint8> PhysicalMaterialIds, TArrayView<const uint8> SimplePhysicalMaterialIds, 
+	TArrayView<const TObjectPtr<UPhysicalMaterial>> PhysicalMaterialObjects)
+{
+	TArrayView<const uint8> ComplexMaterialIndicesView;
+	TArrayView<const uint8> SimpleMaterialIndicesView;
+
+	bool bGenerateSimpleCollision = SimpleCollisionSizeQuads > 0 && !bUseDefaultMaterialOnly;
+
+	if(!ensureMsgf(!HeightfieldGuid.IsValid(), TEXT("Attempting to create a runtime collision object, but one already exists")))
+	{
+		return;
+	}
+
+	int32 CollisionSizeVerts;
+	int32 SimpleCollisionSizeVerts;
+	int32 NumSamples;
+	int32 NumSimpleSamples;
+	GetCollisionSampleInfo(CollisionSizeVerts, SimpleCollisionSizeVerts, NumSamples, NumSimpleSamples);
+
+	if (!ensure(Heights.Num() == NumSamples))
+	{
+		return;
+	}
+
+	int32 NumQuads = (CollisionSizeVerts - 1) * (CollisionSizeVerts - 1);
+	if (!ensure(PhysicalMaterialIds.Num() == NumQuads))
+	{
+		return;
+	}
+
+	if(bGenerateSimpleCollision)
+	{
+		if (!ensure(SimpleHeights.Num() == NumSimpleSamples))
+		{
+			return;
+		}
+
+		int32 NumSimpleQuads = (SimpleCollisionSizeVerts - 1) * (SimpleCollisionSizeVerts - 1);
+		if (!ensure(SimplePhysicalMaterialIds.Num() == NumSimpleQuads))
+		{
+			return;
+		}
+	}
+
+	// In non performant builds, validate that the incoming data's indices are all valid
+#if !UE_BUILD_TEST && !UE_BUILD_SHIPPING
+	for (uint8 Sample : PhysicalMaterialIds)
+	{
+		if (!ensure(Sample == 0xFF || PhysicalMaterialObjects.IsValidIndex(Sample)))
+		{
+			return;
+		}
+	}
+
+	if (bGenerateSimpleCollision)
+	{
+		for (uint8 Sample : SimplePhysicalMaterialIds)
+		{
+			if (!ensure(Sample == 0xFF || PhysicalMaterialObjects.IsValidIndex(Sample)))
+			{
+				return;
+			}
+		}
+	}
+#endif
+
+	HeightfieldGuid = FGuid::NewGuid();
+
+	HeightfieldRef = GSharedHeightfieldRefs.Add(HeightfieldGuid, new FHeightfieldGeometryRef(HeightfieldGuid));
+	HeightfieldRef->Heightfield = MakeUnique<Chaos::FHeightField>(Heights, PhysicalMaterialIds, CollisionSizeVerts, CollisionSizeVerts, Chaos::FVec3(1));
+	
+	if (bGenerateSimpleCollision)
+	{
+		HeightfieldRef->HeightfieldSimple = MakeUnique<Chaos::FHeightField>(SimpleHeights, SimplePhysicalMaterialIds, SimpleCollisionSizeVerts, SimpleCollisionSizeVerts, Chaos::FVec3(1));
+	}
+
+	for (UPhysicalMaterial* PhysicalMaterial : PhysicalMaterialObjects)
+	{
+		HeightfieldRef->UsedChaosMaterials.Add(PhysicalMaterial->GetPhysicsMaterial());
+	}
+}
+
 #if WITH_EDITOR
 void ULandscapeHeightfieldCollisionComponent::SpeculativelyLoadAsyncDDCCollsionData()
 {
