#UE Fix crash during load when LandscapeCompnents contains null entries

diff --git a/Engine/Source/Runtime/Landscape/Private/Landscape.cpp b/Engine/Source/Runtime/Landscape/Private/Landscape.cpp
--- a/Engine/Source/Runtime/Landscape/Private/Landscape.cpp
+++ b/Engine/Source/Runtime/Landscape/Private/Landscape.cpp
@@ -397,14 +397,22 @@ FGuid ALandscapeProxy::GetNaniteContentId() const
 
 	struct FCompareULandscapeComponentBySectionBase
 	{
-		FORCEINLINE bool operator()(const ULandscapeComponent& A, const ULandscapeComponent& B) const
+		FORCEINLINE bool operator()(const ULandscapeComponent* A, const ULandscapeComponent* B) const
 		{
+			if (!A)
+			{
+				return true;
+			}
+			if (!B)
+			{
+				return false;
+			}
 			// Sort components based on their SectionBase (i.e. 2D index relative to the entire landscape) to ensure stable ID generation
-			return (A.GetSectionBase().X == B.GetSectionBase().X) ? (A.GetSectionBase().Y < B.GetSectionBase().Y) : (A.GetSectionBase().X < B.GetSectionBase().X);
+			return (A->GetSectionBase().X == B->GetSectionBase().X) ? (A->GetSectionBase().Y < B->GetSectionBase().Y) : (A->GetSectionBase().X < B->GetSectionBase().X);
 		}
 	};
 	TArray<ULandscapeComponent*> StableOrderComponents(LandscapeComponents);
-	StableOrderComponents.Sort(FCompareULandscapeComponentBySectionBase());
+	Algo::Sort(StableOrderComponents, FCompareULandscapeComponentBySectionBase());
 
 	for (ULandscapeComponent* Component : StableOrderComponents)
 	{
