Disabled Save Modified Landscapes hyperlink in the map check after the action has been executed by the user

diff --git a/Engine/Source/Runtime/Landscape/Private/Landscape.cpp b/Engine/Source/Runtime/Landscape/Private/Landscape.cpp
--- a/Engine/Source/Runtime/Landscape/Private/Landscape.cpp
+++ b/Engine/Source/Runtime/Landscape/Private/Landscape.cpp
@@ -3817,7 +3817,7 @@ void ALandscapeProxy::FixupSharedData(ALandscape* Landscape, bool bMapCheck)
 				->AddToken(FTextToken::Create(LOCTEXT("MapCheck_Message_LandscapeProxy_FixupSharedData", "had some shared properties not in sync with its parent landscape actor. This has been fixed but the proxy needs to be saved in order to ensure cooking behaves as expected. ")))
 				->AddToken(FActionToken::Create(LOCTEXT("MapCheck_SaveFixedUpData", "Save Modified Landscapes"), LOCTEXT("MapCheck_SaveFixedUpData_Desc", "Saves the modified landscape proxy actors"), 
 					FOnActionTokenExecuted::CreateUObject(LandscapeSubsystem, &ULandscapeSubsystem::SaveModifiedLandscapes), 
-					// TODO [jonathan.bard] Enable after the FActionToken change has been submitted : FCanExecuteActionToken::CreateUObject(LandscapeSubsystem, &ULandscapeSubsystem::HasModifiedLandscapes),
+					FCanExecuteActionToken::CreateUObject(LandscapeSubsystem, &ULandscapeSubsystem::HasModifiedLandscapes),
 					/*bInSingleUse = */false))
 				->AddToken(FMapErrorToken::Create(FMapErrors::LandscapeComponentPostLoad_Warning));
 
