Landscape: Fix truncation warnings, part 4

diff --git a/Engine/Source/Runtime/Landscape/Private/Landscape.cpp b/Engine/Source/Runtime/Landscape/Private/Landscape.cpp
--- a/Engine/Source/Runtime/Landscape/Private/Landscape.cpp
+++ b/Engine/Source/Runtime/Landscape/Private/Landscape.cpp
@@ -993,8 +993,8 @@ void ULandscapeComponent::PostLoad()
 
 		// check SectionBaseX/Y are correct
 		const FVector LocalRelativeLocation = GetRelativeLocation();
-		int32 CheckSectionBaseX = FMath::RoundToInt(LocalRelativeLocation.X) + LandscapeProxy->LandscapeSectionOffset.X;
-		int32 CheckSectionBaseY = FMath::RoundToInt(LocalRelativeLocation.Y) + LandscapeProxy->LandscapeSectionOffset.Y;
+		int32 CheckSectionBaseX = FMath::RoundToInt32(LocalRelativeLocation.X) + LandscapeProxy->LandscapeSectionOffset.X;
+		int32 CheckSectionBaseY = FMath::RoundToInt32(LocalRelativeLocation.Y) + LandscapeProxy->LandscapeSectionOffset.Y;
 		if (CheckSectionBaseX != SectionBaseX ||
 			CheckSectionBaseY != SectionBaseY)
 		{
@@ -1136,7 +1136,7 @@ void ULandscapeComponent::PostLoad()
 
 #if !UE_BUILD_SHIPPING
 	// This will fix the data in case there is mismatch between save of asset/maps
-	int8 MaxLOD = FMath::CeilLogTwo(SubsectionSizeQuads + 1) - 1;
+	const int8 MaxLOD = static_cast<int8>(FMath::CeilLogTwo(SubsectionSizeQuads + 1) - 1);
 
 	TArray<ULandscapeMaterialInstanceConstant*> ResolvedMaterials;
 
@@ -3952,7 +3952,7 @@ void ULandscapeInfo::MarkModifiedPackagesAsDirty()
 
 int32 ULandscapeInfo::GetModifiedPackageCount() const
 {
-	return Algo::CountIf(ModifiedPackages, [](const TWeakObjectPtr<UPackage>& InWeakPackagePtr) { return InWeakPackagePtr.IsValid(); });
+	return IntCastChecked<int32>(Algo::CountIf(ModifiedPackages, [](const TWeakObjectPtr<UPackage>& InWeakPackagePtr) { return InWeakPackagePtr.IsValid(); }));
 }
 
 bool ULandscapeInfo::TryAddToModifiedPackages(UPackage* InPackage)
@@ -5067,6 +5067,7 @@ void ULandscapeMeshProxyComponent::InitializeForLandscape(ALandscapeProxy* Lands
 	if (InProxyLOD != INDEX_NONE)
 	{
 		ProxyLOD = static_cast<int8>(FMath::Clamp<int32>(InProxyLOD, 0, FMath::CeilLogTwo(Landscape->SubsectionSizeQuads + 1) - 1));
+		ProxyLOD = FMath::Clamp<int8>(InProxyLOD, 0, static_cast<int8>(FMath::CeilLogTwo(Landscape->SubsectionSizeQuads + 1) - 1));
 	}
 }
 
diff --git a/Engine/Source/Runtime/Landscape/Private/LandscapeCollision.cpp b/Engine/Source/Runtime/Landscape/Private/LandscapeCollision.cpp
--- a/Engine/Source/Runtime/Landscape/Private/LandscapeCollision.cpp
+++ b/Engine/Source/Runtime/Landscape/Private/LandscapeCollision.cpp
@@ -332,7 +332,8 @@ void ULandscapeHeightfieldCollisionComponent::OnCreatePhysicsState()
 
 				// Setup filtering
 				FCollisionFilterData QueryFilterData, SimFilterData;
-				CreateShapeFilterData(GetCollisionObjectType(), FMaskFilter(0), GetOwner()->GetUniqueID(), GetCollisionResponseToChannels(), GetUniqueID(), 0, QueryFilterData, SimFilterData, true, false, true);
+				CreateShapeFilterData(static_cast<uint8>(GetCollisionObjectType()), FMaskFilter(0), GetOwner()->GetUniqueID(), GetCollisionResponseToChannels(), 
+				GetUniqueID(), 0, QueryFilterData, SimFilterData, true, false, true);
 
 				// Heightfield is used for simple and complex collision
 				QueryFilterData.Word3 |= bCreateSimpleCollision ? EPDF_ComplexCollision : (EPDF_SimpleCollision | EPDF_ComplexCollision);
@@ -513,8 +514,8 @@ FPrimitiveSceneProxy* ULandscapeHeightfieldCollisionComponent::CreateSceneProxy(
 			Vertices.SetNumUninitialized(NumVerts);
 			for (int32 I = 0; I < NumVerts; I++)
 			{
-				Chaos::FVec3 Point = GeomData.GetPointScaled(I);
-				Vertices[I].Position = FVector3f(Point.X, Point.Y, Point.Z);
+				const Chaos::FVec3 Point = GeomData.GetPointScaled(I);
+				Vertices[I].Position = FVector3f(static_cast<float>(Point.X), static_cast<float>(Point.Y), static_cast<float>(Point.Z));
 			}
 			Indices.SetNumUninitialized(NumTris * 3);
 
@@ -615,7 +616,7 @@ FPrimitiveSceneProxy* ULandscapeHeightfieldCollisionComponent::CreateSceneProxy(
 		}
 
 		virtual uint32 GetMemoryFootprint(void) const override { return(sizeof(*this) + GetAllocatedSize()); }
-		uint32 GetAllocatedSize(void) const { return FPrimitiveSceneProxy::GetAllocatedSize(); }
+		uint32 GetAllocatedSize(void) const { return static_cast<uint32>(FPrimitiveSceneProxy::GetAllocatedSize()); }
 
 	private:
 		TArray<FDynamicMeshVertex> Vertices;
@@ -978,7 +979,7 @@ bool ULandscapeHeightfieldCollisionComponent::WriteRuntimeData(const FWriteRunti
 							MaterialIndex = InOutMaterials.AddUnique(DominantMaterial);
 						}
 					}
-					MaterialIndices.Add(MaterialIndex);
+					MaterialIndices.Add(IntCastChecked<uint8>(MaterialIndex));
 				}
 			}
 		}
@@ -1296,7 +1297,7 @@ bool ULandscapeMeshCollisionComponent::CookCollisionData(const FName& Format, bo
 
 			if (DominantLayers)
 			{
-				MaterialIndices[TriangleIdx] = MaterialIndex;
+				MaterialIndices[TriangleIdx] = static_cast<uint16>(MaterialIndex);
 			}
 			TriangleIdx++;
 
@@ -1316,7 +1317,7 @@ bool ULandscapeMeshCollisionComponent::CookCollisionData(const FName& Format, bo
 
 			if (DominantLayers)
 			{
-				MaterialIndices[TriangleIdx] = MaterialIndex;
+				MaterialIndices[TriangleIdx] = static_cast<uint16>(MaterialIndex);
 			}
 			TriangleIdx++;
 		}
@@ -1650,7 +1651,7 @@ void ULandscapeMeshCollisionComponent::OnCreatePhysicsState()
 			{
 				// Setup filtering
 				FCollisionFilterData QueryFilterData, SimFilterData;
-				CreateShapeFilterData(GetCollisionObjectType(), FMaskFilter(0), GetOwner()->GetUniqueID(), GetCollisionResponseToChannels(), GetUniqueID(), 0, QueryFilterData, SimFilterData, false, false, true);
+				CreateShapeFilterData(static_cast<uint8>(GetCollisionObjectType()), FMaskFilter(0), GetOwner()->GetUniqueID(), GetCollisionResponseToChannels(), GetUniqueID(), 0, QueryFilterData, SimFilterData, false, false, true);
 				QueryFilterData.Word3 |= (EPDF_SimpleCollision | EPDF_ComplexCollision);
 				SimFilterData.Word3 |= (EPDF_SimpleCollision | EPDF_ComplexCollision);
 
@@ -1708,19 +1709,19 @@ uint32 ULandscapeHeightfieldCollisionComponent::ComputeCollisionHash() const
 	Hash = HashCombine(GetTypeHash(CollisionSizeQuads), Hash);
 	Hash = HashCombine(GetTypeHash(CollisionScale), Hash);
 
-	FTransform ComponentTransform = GetComponentToWorld();
+	const FTransform ComponentTransform = GetComponentToWorld();
 	Hash = FCrc::MemCrc32(&ComponentTransform, sizeof(ComponentTransform), Hash);
 
 	const void* HeightBuffer = CollisionHeightData.LockReadOnly();
-	Hash = FCrc::MemCrc32(HeightBuffer, CollisionHeightData.GetBulkDataSize(), Hash);
+	Hash = FCrc::MemCrc32(HeightBuffer, static_cast<int32>(CollisionHeightData.GetBulkDataSize()), Hash);
 	CollisionHeightData.Unlock();
 
 	const void* DominantBuffer = DominantLayerData.LockReadOnly();
-	Hash = FCrc::MemCrc32(DominantBuffer, DominantLayerData.GetBulkDataSize(), Hash);
+	Hash = FCrc::MemCrc32(DominantBuffer, static_cast<int32>(DominantLayerData.GetBulkDataSize()), Hash);
 	DominantLayerData.Unlock();
 
 	const void* PhysicalMaterialBuffer = PhysicalMaterialRenderData.LockReadOnly();
-	Hash = FCrc::MemCrc32(PhysicalMaterialBuffer, PhysicalMaterialRenderData.GetBulkDataSize(), Hash);
+	Hash = FCrc::MemCrc32(PhysicalMaterialBuffer, static_cast<int32>(PhysicalMaterialRenderData.GetBulkDataSize()), Hash);
 	PhysicalMaterialRenderData.Unlock();
 
 	return Hash;
@@ -1910,8 +1911,8 @@ void ULandscapeHeightfieldCollisionComponent::SnapFoliageInstances(const FBox& I
 			const auto* InstanceSet = MeshInfo.ComponentHash.Find(BaseId);
 			if (InstanceSet)
 			{
-				float TraceExtentSize = Bounds.SphereRadius * 2.f + 10.f; // extend a little
-				FVector TraceVector = GetOwner()->GetRootComponent()->GetComponentTransform().GetUnitAxis(EAxis::Z) * TraceExtentSize;
+				const float TraceExtentSize = static_cast<float>(Bounds.SphereRadius) * 2.f + 10.f; // extend a little
+				const FVector TraceVector = GetOwner()->GetRootComponent()->GetComponentTransform().GetUnitAxis(EAxis::Z) * TraceExtentSize;
 
 				TArray<int32> InstancesToRemove;
 				TSet<UHierarchicalInstancedStaticMeshComponent*> AffectedFoliageComponents;
@@ -2213,7 +2214,7 @@ void ULandscapeHeightfieldCollisionComponent::PrepareGeometryExportSync()
 				CachedHeightFieldSamples.Heights.SetNumUninitialized(HeightsCount);
 				for(int32 Index = 0; Index < HeightsCount; ++Index)
 				{
-					CachedHeightFieldSamples.Heights[Index] = HeightfieldRef->Heightfield->GetHeight(Index);
+					CachedHeightFieldSamples.Heights[Index] = static_cast<int16>(HeightfieldRef->Heightfield->GetHeight(Index));
 				}
 
 				const int32 HolesCount = (HeightfieldRowsCount-1) * (HeightfieldColumnsCount-1);
@@ -2534,12 +2535,12 @@ void ULandscapeHeightfieldCollisionComponent::ExportCustomProperties(FOutputDevi
 	CollisionHeightData.Unlock();
 	Out.Logf(TEXT("\r\n"));
 
-	int32 NumDominantLayerSamples = DominantLayerData.GetElementCount();
+	const int32 NumDominantLayerSamples = static_cast<int32>(DominantLayerData.GetElementCount());
 	check(NumDominantLayerSamples == 0 || NumDominantLayerSamples == NumHeights);
 
 	if (NumDominantLayerSamples > 0)
 	{
-		uint8* DominantLayerSamples = (uint8*)DominantLayerData.Lock(LOCK_READ_ONLY);
+		const uint8* DominantLayerSamples = (uint8*)DominantLayerData.Lock(LOCK_READ_ONLY);
 
 		Out.Logf(TEXT("%sCustomProperties DominantLayerData "), FCString::Spc(Indent));
 		for (int32 i = 0; i < NumDominantLayerSamples; i++)
@@ -2570,7 +2571,7 @@ void ULandscapeHeightfieldCollisionComponent::ImportCustomProperties(const TCHAR
 		{
 			if (i < NumHeights)
 			{
-				Heights[i++] = FCString::Atoi(SourceText);
+				Heights[i++] = static_cast<uint16>(FCString::Atoi(SourceText));
 				while (FChar::IsDigit(*SourceText))
 				{
 					SourceText++;
@@ -2601,7 +2602,7 @@ void ULandscapeHeightfieldCollisionComponent::ImportCustomProperties(const TCHAR
 		{
 			if (i < NumDominantLayerSamples)
 			{
-				DominantLayerSamples[i++] = FParse::HexDigit(SourceText[0]) * 16 + FParse::HexDigit(SourceText[1]);
+				DominantLayerSamples[i++] = static_cast<uint8>(FParse::HexDigit(SourceText[0]) * 16 + FParse::HexDigit(SourceText[1]));
 			}
 			SourceText += 2;
 		}
@@ -2656,7 +2657,7 @@ void ULandscapeMeshCollisionComponent::ImportCustomProperties(const TCHAR* Sourc
 		{
 			if (i < NumHeights)
 			{
-				Heights[i++] = FCString::Atoi(SourceText);
+				Heights[i++] = static_cast<uint16>(FCString::Atoi(SourceText));
 				while (FChar::IsDigit(*SourceText))
 				{
 					SourceText++;
@@ -2687,7 +2688,7 @@ void ULandscapeMeshCollisionComponent::ImportCustomProperties(const TCHAR* Sourc
 		{
 			if (i < NumDominantLayerSamples)
 			{
-				DominantLayerSamples[i++] = FParse::HexDigit(SourceText[0]) * 16 + FParse::HexDigit(SourceText[1]);
+				DominantLayerSamples[i++] = static_cast<uint8>(FParse::HexDigit(SourceText[0]) * 16 + FParse::HexDigit(SourceText[1]));
 			}
 			SourceText += 2;
 		}
@@ -2713,7 +2714,7 @@ void ULandscapeMeshCollisionComponent::ImportCustomProperties(const TCHAR* Sourc
 		{
 			if (i < NumOffsets)
 			{
-				Offsets[i++] = FCString::Atoi(SourceText);
+				Offsets[i++] = static_cast<uint16>(FCString::Atoi(SourceText));
 				while (FChar::IsDigit(*SourceText))
 				{
 					SourceText++;
@@ -2790,7 +2791,7 @@ ULandscapeComponent* ULandscapeHeightfieldCollisionComponent::GetRenderComponent
 TOptional<float> ULandscapeHeightfieldCollisionComponent::GetHeight(float X, float Y, EHeightfieldSource HeightFieldSource)
 {
 	TOptional<float> Height;
-	const float ZScale = GetComponentTransform().GetScale3D().Z * LANDSCAPE_ZSCALE;
+	const float ZScale = static_cast<float>(GetComponentTransform().GetScale3D().Z * LANDSCAPE_ZSCALE); // TODO michael.balzer: Is it okay that ZScale is not used in this function?
 
 	if (!IsValidRef(HeightfieldRef))
 	{
@@ -2818,7 +2819,7 @@ TOptional<float> ULandscapeHeightfieldCollisionComponent::GetHeight(float X, flo
 	
 	if (HeightField)
 	{
-		Height = HeightField->GetHeightAt({ X, Y });
+		Height = static_cast<float>(HeightField->GetHeightAt({ X, Y }));
 	}
 
 	return Height;
@@ -2835,7 +2836,7 @@ struct FHeightFieldAccessor
 
 	float GetUnscaledHeight(int32 X, int32 Y) const
 	{
-		return GeometryRef.Heightfield->GetHeight(X, Y);
+		return static_cast<float>(GeometryRef.Heightfield->GetHeight(X, Y));
 	}
 
 	uint8 GetMaterialIndex(int32 X, int32 Y) const
@@ -2864,7 +2865,7 @@ bool ULandscapeHeightfieldCollisionComponent::FillHeightTile(TArrayView<float> H
 	}
 
 	const FTransform& WorldTransform = GetComponentToWorld();
-	const float ZScale = WorldTransform.GetScale3D().Z * LANDSCAPE_ZSCALE;
+	const float ZScale = static_cast<float>(WorldTransform.GetScale3D().Z * LANDSCAPE_ZSCALE);
 
 	// Write all values to output array
 	for (int32 y = 0; y < Accessor.NumY; ++y)
@@ -2872,7 +2873,7 @@ bool ULandscapeHeightfieldCollisionComponent::FillHeightTile(TArrayView<float> H
 		for (int32 x = 0; x < Accessor.NumX; ++x)
 		{
 			const float CurrHeight = Accessor.GetUnscaledHeight(x, y);
-			const float WorldHeight = WorldTransform.TransformPositionNoScale(FVector(0, 0, CurrHeight * ZScale)).Z;
+			const float WorldHeight = static_cast<float>(WorldTransform.TransformPositionNoScale(FVector(0, 0, CurrHeight * ZScale)).Z);
 
 			// write output
 			const int32 WriteIndex = Offset + y * Stride + x;
@@ -2943,15 +2944,15 @@ TOptional<float> ALandscapeProxy::GetHeightAtLocation(FVector Location, EHeightf
 	if (ULandscapeInfo* Info = GetLandscapeInfo())
 	{
 		const FVector ActorSpaceLocation = LandscapeActorToWorld().InverseTransformPosition(Location);
-		const FIntPoint Key = FIntPoint(FMath::FloorToInt(ActorSpaceLocation.X / ComponentSizeQuads), FMath::FloorToInt(ActorSpaceLocation.Y / ComponentSizeQuads));
+		const FIntPoint Key = FIntPoint(FMath::FloorToInt32(ActorSpaceLocation.X / ComponentSizeQuads), FMath::FloorToInt32(ActorSpaceLocation.Y / ComponentSizeQuads));
 		ULandscapeHeightfieldCollisionComponent* Component = Info->XYtoCollisionComponentMap.FindRef(Key);
 		if (Component)
 		{
 			const FVector ComponentSpaceLocation = Component->GetComponentToWorld().InverseTransformPosition(Location);
-			const TOptional<float> LocalHeight = Component->GetHeight(ComponentSpaceLocation.X, ComponentSpaceLocation.Y, HeightFieldSource);
+			const TOptional<float> LocalHeight = Component->GetHeight(static_cast<float>(ComponentSpaceLocation.X), static_cast<float>(ComponentSpaceLocation.Y), HeightFieldSource);
 			if (LocalHeight.IsSet())
 			{
-				Height = Component->GetComponentToWorld().TransformPositionNoScale(FVector(0, 0, LocalHeight.GetValue())).Z;
+				Height = static_cast<float>(Component->GetComponentToWorld().TransformPositionNoScale(FVector(0, 0, LocalHeight.GetValue())).Z);
 			}
 		}
 	}
@@ -3029,15 +3030,15 @@ void ALandscapeProxy::GetHeightValues(int32& SizeX, int32& SizeY, TArray<float>
 		const int32 NumY = HeightFieldData->GetNumRows();
 
 		const FTransform& ComponentToWorld = CollisionComponent->GetComponentToWorld();
-		const float ZScale = ComponentToWorld.GetScale3D().Z * LANDSCAPE_ZSCALE;
+		const float ZScale = static_cast<float>(ComponentToWorld.GetScale3D().Z * LANDSCAPE_ZSCALE);
 
 		// Write all values to output array
 		for (int32 x = 0; x < NumX; ++x)
 		{
 			for (int32 y = 0; y < NumY; ++y)
 			{
-				const float CurrHeight = HeightFieldData->GetHeight(x, y) * ZScale;				
-				const float WorldHeight = ComponentToWorld.TransformPositionNoScale(FVector(0, 0, CurrHeight)).Z;
+				const float CurrHeight = static_cast<float>(HeightFieldData->GetHeight(x, y) * ZScale);				
+				const float WorldHeight = static_cast<float>(ComponentToWorld.TransformPositionNoScale(FVector(0, 0, CurrHeight)).Z);
 				
 				// write output
 				const int32 WriteX = BaseX + x;
diff --git a/Engine/Source/Runtime/Landscape/Private/LandscapeConfigHelper.cpp b/Engine/Source/Runtime/Landscape/Private/LandscapeConfigHelper.cpp
--- a/Engine/Source/Runtime/Landscape/Private/LandscapeConfigHelper.cpp
+++ b/Engine/Source/Runtime/Landscape/Private/LandscapeConfigHelper.cpp
@@ -136,9 +136,9 @@ bool FLandscapeConfigHelper::ChangeGridSize(ULandscapeInfo* InLandscapeInfo, uin
 				ComponentMaterials.FindOrAdd(LandscapeComponent, LandscapeComponent->GetLandscapeMaterial());
 				ComponentHoleMaterials.FindOrAdd(LandscapeComponent, LandscapeComponent->GetLandscapeHoleMaterial());
 				TMap<int32, UMaterialInterface*>& LODMaterials = ComponentLODMaterials.FindOrAdd(LandscapeComponent);
-				for (int32 LODIndex = 0; LODIndex <= 8; ++LODIndex)
+				for (int8 LODIndex = 0; LODIndex <= 8; ++LODIndex)
 				{
-					LODMaterials.Add(LODIndex, LandscapeComponent->GetLandscapeMaterial(static_cast<int8>(LODIndex)));
+					LODMaterials.Add(LODIndex, LandscapeComponent->GetLandscapeMaterial(LODIndex));
 				}
 
 				ComponentsToMove.Add(LandscapeComponent);
diff --git a/Engine/Source/Runtime/Landscape/Private/LandscapeEdit.cpp b/Engine/Source/Runtime/Landscape/Private/LandscapeEdit.cpp
--- a/Engine/Source/Runtime/Landscape/Private/LandscapeEdit.cpp
+++ b/Engine/Source/Runtime/Landscape/Private/LandscapeEdit.cpp
@@ -405,7 +405,7 @@ UMaterialInstanceConstant* ULandscapeComponent::GetCombinationMaterial(FMaterial
 
 void ULandscapeComponent::UpdateMaterialInstances_Internal(FMaterialUpdateContext& Context)
 {
-	int32 MaxLOD = FMath::CeilLogTwo(SubsectionSizeQuads + 1) - 1;
+	const int8 MaxLOD = IntCastChecked<int8>(FMath::CeilLogTwo(SubsectionSizeQuads + 1) - 1);
 	decltype(MaterialPerLOD) NewMaterialPerLOD;
 	LODIndexToMaterialIndex.SetNumUninitialized(MaxLOD+1);
 	int8 LastLODIndex = INDEX_NONE;
@@ -413,7 +413,7 @@ void ULandscapeComponent::UpdateMaterialInstances_Internal(FMaterialUpdateContex
 	UMaterialInterface* BaseMaterial = GetLandscapeMaterial();
 	UMaterialInterface* LOD0Material = GetLandscapeMaterial(0);
 
-	for (int32 LODIndex = 0; LODIndex <= MaxLOD; ++LODIndex)
+	for (int8 LODIndex = 0; LODIndex <= MaxLOD; ++LODIndex)
 	{
 		UMaterialInterface* CurrentMaterial = GetLandscapeMaterial(LODIndex);
 
@@ -431,7 +431,7 @@ void ULandscapeComponent::UpdateMaterialInstances_Internal(FMaterialUpdateContex
 		}
 		else
 		{
-			int32 AddedIndex = NewMaterialPerLOD.Num();
+			const int8 AddedIndex = IntCastChecked<int8>(NewMaterialPerLOD.Num());
 			NewMaterialPerLOD.Add(CurrentMaterial, LODIndex);
 			LODIndexToMaterialIndex[LODIndex] = AddedIndex;
 			LastLODIndex = AddedIndex;
@@ -636,7 +636,7 @@ class UMaterialInterface* ULandscapeComponent::GetMaterial(int32 ElementIndex) c
 {
 	if (ensure(ElementIndex == 0))
 	{
-		return GetLandscapeMaterial(ElementIndex);
+		return GetLandscapeMaterial(IntCastChecked<int8>(ElementIndex));
 	}
 
 	return nullptr;
@@ -979,7 +979,7 @@ void ULandscapeComponent::FixupWeightmaps(const FGuid& InEditLayerGuid)
 				// Fix up any problems caused by the layer deletion bug.
 				if (Allocation.WeightmapTextureIndex >= LocalWeightmapTextures.Num())
 				{
-					Allocation.WeightmapTextureIndex = LocalWeightmapTextures.Num() - 1;
+					Allocation.WeightmapTextureIndex = static_cast<uint8>(LocalWeightmapTextures.Num() - 1);
 					if (!bFixedWeightmapTextureIndex)
 					{
 						FFormatNamedArguments Arguments;
@@ -1442,8 +1442,8 @@ void ULandscapeComponent::UpdateCollisionHeightBuffer(	int32 InComponentX1, int3
 	const int32 MipSizeU = InHeightmapSizeU >> InCollisionMipLevel;
 	const int32 MipSizeV = InHeightmapSizeV >> InCollisionMipLevel;
 
-	const int32 HeightmapOffsetX = FMath::RoundToInt(HeightmapScaleBias.Z * (float)InHeightmapSizeU) >> InCollisionMipLevel;
-	const int32 HeightmapOffsetY = FMath::RoundToInt(HeightmapScaleBias.W * (float)InHeightmapSizeV) >> InCollisionMipLevel;
+	const int32 HeightmapOffsetX = FMath::RoundToInt32(HeightmapScaleBias.Z * (float)InHeightmapSizeU) >> InCollisionMipLevel;
+	const int32 HeightmapOffsetY = FMath::RoundToInt32(HeightmapScaleBias.W * (float)InHeightmapSizeV) >> InCollisionMipLevel;
 
 	const int32 XYMipSizeU = XYOffsetmapTexture ? XYOffsetmapTexture->Source.GetSizeX() >> InCollisionMipLevel : 0;
 
@@ -1563,7 +1563,8 @@ void ULandscapeComponent::UpdateDominantLayerBuffer(int32 InComponentX1, int32 I
 
 					uint8 DominantLayer = 255; // 255 as invalid value
 					int32 DominantWeight = 0;
-					for (int32 LayerIdx = 0; LayerIdx < InCollisionDataPtrs.Num(); LayerIdx++)
+					const uint8 NumLayers = IntCastChecked<uint8>(InCollisionDataPtrs.Num());
+					for (uint8 LayerIdx = 0; LayerIdx < NumLayers; LayerIdx++)
 					{
 						const uint8 LayerWeight = InCollisionDataPtrs[LayerIdx][DataOffset];
 						const uint8 LayerMinimumWeight = InLayerInfos[LayerIdx] ? (uint8)(InLayerInfos[LayerIdx]->MinimumCollisionRelevanceWeight * 255) :  0;
@@ -1948,8 +1949,8 @@ void ULandscapeComponent::GenerateHeightmapMips(TArray<FColor*>& HeightmapTextur
 	int32 HeightmapSizeU = GetHeightmap()->Source.GetSizeX();
 	int32 HeightmapSizeV = GetHeightmap()->Source.GetSizeY();
 
-	int32 HeightmapOffsetX = FMath::RoundToInt(HeightmapScaleBias.Z * (float)HeightmapSizeU);
-	int32 HeightmapOffsetY = FMath::RoundToInt(HeightmapScaleBias.W * (float)HeightmapSizeV);
+	const int32 HeightmapOffsetX = FMath::RoundToInt32(HeightmapScaleBias.Z * (float)HeightmapSizeU);
+	const int32 HeightmapOffsetY = FMath::RoundToInt32(HeightmapScaleBias.W * (float)HeightmapSizeV);
 
 	for (int32 SubsectionY = 0; SubsectionY < NumSubsections; SubsectionY++)
 	{
@@ -2086,27 +2087,27 @@ void ULandscapeComponent::GenerateHeightmapMips(TArray<FColor*>& HeightmapTextur
 						uint16 PrevMipHeightValue01 = PreMipTexData01->R << 8 | PreMipTexData01->G;
 						uint16 PrevMipHeightValue10 = PreMipTexData10->R << 8 | PreMipTexData10->G;
 						uint16 PrevMipHeightValue11 = PreMipTexData11->R << 8 | PreMipTexData11->G;
-						uint16 HeightValue = FMath::RoundToInt(
+						uint16 HeightValue = static_cast<uint16>(FMath::RoundToInt(
 							FMath::Lerp(
 							FMath::Lerp((float)PrevMipHeightValue00, (float)PrevMipHeightValue10, fPrevMipTexFracX),
 							FMath::Lerp((float)PrevMipHeightValue01, (float)PrevMipHeightValue11, fPrevMipTexFracX),
-							fPrevMipTexFracY));
+							fPrevMipTexFracY)));
 
 						TexData->R = HeightValue >> 8;
 						TexData->G = HeightValue & 255;
 
 						// Lerp tangents
-						TexData->B = FMath::RoundToInt(
+						TexData->B = static_cast<uint8>(FMath::RoundToInt(
 							FMath::Lerp(
 							FMath::Lerp((float)PreMipTexData00->B, (float)PreMipTexData10->B, fPrevMipTexFracX),
 							FMath::Lerp((float)PreMipTexData01->B, (float)PreMipTexData11->B, fPrevMipTexFracX),
-							fPrevMipTexFracY));
+							fPrevMipTexFracY)));
 
-						TexData->A = FMath::RoundToInt(
+						TexData->A = static_cast<uint8>(FMath::RoundToInt(
 							FMath::Lerp(
 							FMath::Lerp((float)PreMipTexData00->A, (float)PreMipTexData10->A, fPrevMipTexFracX),
 							FMath::Lerp((float)PreMipTexData01->A, (float)PreMipTexData11->A, fPrevMipTexFracX),
-							fPrevMipTexFracY));
+							fPrevMipTexFracY)));
 
 						// Padding for missing data
 						if (EndX && SubsectionX == NumSubsections - 1 && VertX == VertX2)
@@ -2224,36 +2225,36 @@ namespace
 	template<typename DataType>
 	void BiLerpTextureData(DataType* Output, const DataType* Data00, const DataType* Data10, const DataType* Data01, const DataType* Data11, float FracX, float FracY)
 	{
-		*Output = FMath::RoundToInt(
+		*Output = static_cast<DataType>(FMath::RoundToInt(
 			FMath::Lerp(
 			FMath::Lerp((float)*Data00, (float)*Data10, FracX),
 			FMath::Lerp((float)*Data01, (float)*Data11, FracX),
-			FracY));
+			FracY)));
 	}
 
 	template<>
 	void BiLerpTextureData(FColor* Output, const FColor* Data00, const FColor* Data10, const FColor* Data01, const FColor* Data11, float FracX, float FracY)
 	{
-		Output->R = FMath::RoundToInt(
+		Output->R = static_cast<uint8>(FMath::RoundToInt(
 			FMath::Lerp(
 			FMath::Lerp((float)Data00->R, (float)Data10->R, FracX),
 			FMath::Lerp((float)Data01->R, (float)Data11->R, FracX),
-			FracY));
-		Output->G = FMath::RoundToInt(
+			FracY)));
+		Output->G = static_cast<uint8>(FMath::RoundToInt(
 			FMath::Lerp(
 			FMath::Lerp((float)Data00->G, (float)Data10->G, FracX),
 			FMath::Lerp((float)Data01->G, (float)Data11->G, FracX),
-			FracY));
-		Output->B = FMath::RoundToInt(
+			FracY)));
+		Output->B = static_cast<uint8>(FMath::RoundToInt(
 			FMath::Lerp(
 			FMath::Lerp((float)Data00->B, (float)Data10->B, FracX),
 			FMath::Lerp((float)Data01->B, (float)Data11->B, FracX),
-			FracY));
-		Output->A = FMath::RoundToInt(
+			FracY)));
+		Output->A = static_cast<uint8>(FMath::RoundToInt(
 			FMath::Lerp(
 			FMath::Lerp((float)Data00->A, (float)Data10->A, FracX),
 			FMath::Lerp((float)Data01->A, (float)Data11->A, FracX),
-			FracY));
+			FracY)));
 	}
 
 	template<typename DataType>
@@ -2479,10 +2480,10 @@ float ULandscapeComponent::GetLayerWeightAtLocation(const FVector& InLocation, U
 	}
 
 	// Find data
-	int32 X1 = FMath::FloorToInt(TestLocation.X);
-	int32 Y1 = FMath::FloorToInt(TestLocation.Y);
-	int32 X2 = FMath::CeilToInt(TestLocation.X);
-	int32 Y2 = FMath::CeilToInt(TestLocation.Y);
+	int32 X1 = FMath::FloorToInt32(TestLocation.X);
+	int32 Y1 = FMath::FloorToInt32(TestLocation.Y);
+	int32 X2 = FMath::CeilToInt32(TestLocation.X);
+	int32 Y2 = FMath::CeilToInt32(TestLocation.Y);
 
 	int32 Stride = (SubsectionSizeQuads + 1) * NumSubsections;
 
@@ -2498,8 +2499,8 @@ float ULandscapeComponent::GetLayerWeightAtLocation(const FVector& InLocation, U
 	float Sample12 = (float)((*LayerCache)[IdxX1 + Stride * IdxY2]) / 255.0f;
 	float Sample22 = (float)((*LayerCache)[IdxX2 + Stride * IdxY2]) / 255.0f;
 
-	float LerpX = FMath::Fractional(TestLocation.X);
-	float LerpY = FMath::Fractional(TestLocation.Y);
+	float LerpX = FMath::Fractional(static_cast<float>(TestLocation.X));
+	float LerpY = FMath::Fractional(static_cast<float>(TestLocation.Y));
 
 	// Bilinear interpolate
 	return FMath::Lerp(
@@ -3027,7 +3028,7 @@ LANDSCAPE_API void ALandscapeProxy::Import(const FGuid& InGuid, int32 InMinX, in
 
 						if (255 - TotalWeight && MaxLayerIdx >= 0)
 						{
-							WeightValues[MaxLayerIdx][Idx] += 255 - TotalWeight;
+							WeightValues[MaxLayerIdx][Idx] += static_cast<uint8>(255 - TotalWeight);
 						}
 					}
 				}
@@ -3113,8 +3114,8 @@ LANDSCAPE_API void ALandscapeProxy::Import(const FGuid& InGuid, int32 InMinX, in
 
 					for (int32 i = 0; i < RemainingLayers; i++)
 					{
-						ComponentWeightmapLayerAllocations[LayerIndex + i].WeightmapTextureIndex = ComponentWeightmapTextures.Num();
-						ComponentWeightmapLayerAllocations[LayerIndex + i].WeightmapTextureChannel = Allocation.ChannelsInUse;
+						ComponentWeightmapLayerAllocations[LayerIndex + i].WeightmapTextureIndex = static_cast<uint8>(ComponentWeightmapTextures.Num());
+						ComponentWeightmapLayerAllocations[LayerIndex + i].WeightmapTextureChannel = static_cast<uint8>(Allocation.ChannelsInUse);
 						WeightmapUsage->ChannelUsage[Allocation.ChannelsInUse] = LandscapeComponent;
 						switch (Allocation.ChannelsInUse)
 						{
@@ -3152,28 +3153,28 @@ LANDSCAPE_API void ALandscapeProxy::Import(const FGuid& InGuid, int32 InMinX, in
 					UE_LOG(LogLandscape, VeryVerbose, TEXT("  ==> Storing %d channels in new texture %s"), ThisAllocationLayers, *WeightmapTexture->GetName());
 
 					WeightmapTextureDataPointers.Add((uint8*)&MipData->R);
-					ComponentWeightmapLayerAllocations[LayerIndex + 0].WeightmapTextureIndex = ComponentWeightmapTextures.Num();
+					ComponentWeightmapLayerAllocations[LayerIndex + 0].WeightmapTextureIndex = static_cast<uint8>(ComponentWeightmapTextures.Num());
 					ComponentWeightmapLayerAllocations[LayerIndex + 0].WeightmapTextureChannel = 0;
 					WeightmapUsage->ChannelUsage[0] = LandscapeComponent;
 
 					if (ThisAllocationLayers > 1)
 					{
 						WeightmapTextureDataPointers.Add((uint8*)&MipData->G);
-						ComponentWeightmapLayerAllocations[LayerIndex + 1].WeightmapTextureIndex = ComponentWeightmapTextures.Num();
+						ComponentWeightmapLayerAllocations[LayerIndex + 1].WeightmapTextureIndex = static_cast<uint8>(ComponentWeightmapTextures.Num());
 						ComponentWeightmapLayerAllocations[LayerIndex + 1].WeightmapTextureChannel = 1;
 						WeightmapUsage->ChannelUsage[1] = LandscapeComponent;
 
 						if (ThisAllocationLayers > 2)
 						{
 							WeightmapTextureDataPointers.Add((uint8*)&MipData->B);
-							ComponentWeightmapLayerAllocations[LayerIndex + 2].WeightmapTextureIndex = ComponentWeightmapTextures.Num();
+							ComponentWeightmapLayerAllocations[LayerIndex + 2].WeightmapTextureIndex = static_cast<uint8>(ComponentWeightmapTextures.Num());
 							ComponentWeightmapLayerAllocations[LayerIndex + 2].WeightmapTextureChannel = 2;
 							WeightmapUsage->ChannelUsage[2] = LandscapeComponent;
 
 							if (ThisAllocationLayers > 3)
 							{
 								WeightmapTextureDataPointers.Add((uint8*)&MipData->A);
-								ComponentWeightmapLayerAllocations[LayerIndex + 3].WeightmapTextureIndex = ComponentWeightmapTextures.Num();
+								ComponentWeightmapLayerAllocations[LayerIndex + 3].WeightmapTextureIndex = static_cast<uint8>(ComponentWeightmapTextures.Num());
 								ComponentWeightmapLayerAllocations[LayerIndex + 3].WeightmapTextureChannel = 3;
 								WeightmapUsage->ChannelUsage[3] = LandscapeComponent;
 							}
@@ -3214,8 +3215,8 @@ LANDSCAPE_API void ALandscapeProxy::Import(const FGuid& InGuid, int32 InMinX, in
 
 							HeightmapInfo.HeightmapTextureMipData[0][HeightTexDataIdx].R = HeightValue >> 8;
 							HeightmapInfo.HeightmapTextureMipData[0][HeightTexDataIdx].G = HeightValue & 255;
-							HeightmapInfo.HeightmapTextureMipData[0][HeightTexDataIdx].B = FMath::RoundToInt(127.5f * (Normal.X + 1.0f));
-							HeightmapInfo.HeightmapTextureMipData[0][HeightTexDataIdx].A = FMath::RoundToInt(127.5f * (Normal.Y + 1.0f));
+							HeightmapInfo.HeightmapTextureMipData[0][HeightTexDataIdx].B = static_cast<uint8>(FMath::RoundToInt32(127.5f * (Normal.X + 1.0f)));
+							HeightmapInfo.HeightmapTextureMipData[0][HeightTexDataIdx].A = static_cast<uint8>(FMath::RoundToInt32(127.5f * (Normal.Y + 1.0f)));
 
 							for (int32 WeightmapIndex = 0; WeightmapIndex < WeightValues.Num(); WeightmapIndex++)
 							{
@@ -3322,7 +3323,7 @@ LANDSCAPE_API void ALandscapeProxy::Import(const FGuid& InGuid, int32 InMinX, in
 
 	for (ULandscapeComponent* Component : LandscapeComponents)
 	{
-		int8 MaxLOD = FMath::CeilLogTwo(Component->SubsectionSizeQuads + 1) - 1;
+		const int8 MaxLOD = IntCastChecked<int8>(FMath::CeilLogTwo(Component->SubsectionSizeQuads + 1) - 1);
 
 		for (int8 LODIndex = 0; LODIndex < MaxLOD; ++LODIndex)
 		{
@@ -3696,11 +3697,11 @@ bool ALandscapeProxy::ExportToRawMesh(const FRawMeshExportParams& InExportParams
 		const FVector2f ComponentOffsetRelativeToProxyBoundsQuadsLOD = FVector2f(ComponentOffsetRelativeToProxyBoundsQuads) * LODScale;
 		const FVector2f ComponentUVScaleRelativeToProxyBoundsLOD = LandscapeProxyBoundsRectUVScale / LODScale;
 
-		const FVector2f ComponentHeightmapUVBias = FVector2f(Component->HeightmapScaleBias.Z, Component->HeightmapScaleBias.W);
-		const FVector2f ComponentHeightmapUVScale = FVector2f(Component->HeightmapScaleBias.X, Component->HeightmapScaleBias.Y);
+		const FVector2f ComponentHeightmapUVBias = FVector2f(static_cast<float>(Component->HeightmapScaleBias.Z), static_cast<float>(Component->HeightmapScaleBias.W));
+		const FVector2f ComponentHeightmapUVScale = FVector2f(static_cast<float>(Component->HeightmapScaleBias.X), static_cast<float>(Component->HeightmapScaleBias.Y));
 		const FVector2f ComponentHeightmapUVScaleLOD = ComponentHeightmapUVScale / LODScale;
 		const FVector2f ComponentHeightmapUVPixelOffset = ComponentHeightmapUVScale * 0.5f;
-		const FVector2f ComponentWeightmapUVScale = FVector2f(Component->WeightmapScaleBias.X, Component->WeightmapScaleBias.Y);
+		const FVector2f ComponentWeightmapUVScale = FVector2f(static_cast<float>(Component->WeightmapScaleBias.X), static_cast<float>(Component->WeightmapScaleBias.Y));
 		const FVector2f ComponentWeightmapUVScaleLOD = ComponentWeightmapUVScale / LODScale;
 		const FVector2f ComponentWeightmapUVPixelOffset = ComponentWeightmapUVScale * 0.5f; // I could have used Component->WeightmapScaleBias.ZW but then it would be confusing because it doesn't have the same signification as the heightmap UV bias
 
@@ -3745,7 +3746,7 @@ bool ALandscapeProxy::ExportToRawMesh(const FRawMeshExportParams& InExportParams
 			FIntPoint(1, 0),
 		};
 
-		const float SquaredSphereRadius = InExportParams.ExportBounds.IsSet() ? FMath::Square(InExportParams.ExportBounds->SphereRadius) : 0.0f;
+		const float SquaredSphereRadius = InExportParams.ExportBounds.IsSet() ? FMath::Square(static_cast<float>(InExportParams.ExportBounds->SphereRadius)) : 0.0f;
 
 		//We need to not duplicate the vertex position, so we use the FIndexAndZ to achieve fast result
 		TArray<FIndexAndZ> VertIndexAndZ;
@@ -3895,19 +3896,19 @@ bool ALandscapeProxy::ExportToRawMesh(const FRawMeshExportParams& InExportParams
 									{
 										case FRawMeshExportParams::EUVMappingType::RelativeToProxyBoundsUV:
 										{
-											FVector2f UV = (ComponentOffsetRelativeToProxyBoundsQuadsLOD + FVector2f(VertexX, VertexY)) * ComponentUVScaleRelativeToProxyBoundsLOD;
+											FVector2f UV = (ComponentOffsetRelativeToProxyBoundsQuadsLOD + FVector2f(static_cast<float>(VertexX), static_cast<float>(VertexY))) * ComponentUVScaleRelativeToProxyBoundsLOD;
 											VertexInstanceUVs.Set(VertexInstanceIDs[i], UVChannel, UV);
 											break;
 										}
 										case FRawMeshExportParams::EUVMappingType::HeightmapUV:
 										{
-											FVector2f UV = FVector2f(VertexX, VertexY) * ComponentHeightmapUVScaleLOD + ComponentHeightmapUVPixelOffset + ComponentHeightmapUVBias;
+											FVector2f UV = FVector2f(static_cast<float>(VertexX), static_cast<float>(VertexY)) * ComponentHeightmapUVScaleLOD + ComponentHeightmapUVPixelOffset + ComponentHeightmapUVBias;
 											VertexInstanceUVs.Set(VertexInstanceIDs[i], UVChannel, UV);
 											break;
 										}
 										case FRawMeshExportParams::EUVMappingType::WeightmapUV:
 										{
-											FVector2f UV = FVector2f(VertexX, VertexY) * ComponentWeightmapUVScaleLOD + ComponentWeightmapUVPixelOffset;
+											FVector2f UV = FVector2f(static_cast<float>(VertexX), static_cast<float>(VertexY)) * ComponentWeightmapUVScaleLOD + ComponentWeightmapUVPixelOffset;
 											VertexInstanceUVs.Set(VertexInstanceIDs[i], UVChannel, UV);
 											break;
 										}
@@ -3915,7 +3916,7 @@ bool ALandscapeProxy::ExportToRawMesh(const FRawMeshExportParams& InExportParams
 										case FRawMeshExportParams::EUVMappingType::TerrainCoordMapping_XZ:
 										case FRawMeshExportParams::EUVMappingType::TerrainCoordMapping_YZ:
 										{
-											FVector2f QuadCoords = (ComponentOffsetRelativeToProxyBoundsQuadsLOD + FVector2f(VertexX, VertexY) / LODScale);
+											FVector2f QuadCoords = (ComponentOffsetRelativeToProxyBoundsQuadsLOD + FVector2f(static_cast<float>(VertexX), static_cast<float>(VertexY)) / LODScale);
 											FVector2f UV;
 											if (UVMappingType == FRawMeshExportParams::EUVMappingType::TerrainCoordMapping_XY)
 											{
@@ -3924,7 +3925,7 @@ bool ALandscapeProxy::ExportToRawMesh(const FRawMeshExportParams& InExportParams
 											else
 											{
 												UV[0] = (UVMappingType == FRawMeshExportParams::EUVMappingType::TerrainCoordMapping_XZ) ? QuadCoords[0] : QuadCoords[1];
-												UV[1] = LocalPositions[i].Z;
+												UV[1] = static_cast<float>(LocalPositions[i].Z);
 											}
 											VertexInstanceUVs.Set(VertexInstanceIDs[i], UVChannel, UV);
 											break;
@@ -4092,7 +4093,7 @@ FVector ULandscapeInfo::GetLandscapeCenterPos(float& LengthZ, int32 MinX /*= MAX
 {
 	// MinZ, MaxZ is Local coordinate
 	float MaxZ = -UE_OLD_HALF_WORLD_MAX, MinZ = UE_OLD_HALF_WORLD_MAX;
-	const float ScaleZ = DrawScale.Z;
+	const float ScaleZ = static_cast<float>(DrawScale.Z);
 
 	if (MinX == MAX_int32)
 	{
@@ -4356,14 +4357,14 @@ void ALandscapeProxy::EditorApplyScale(const FVector& DeltaScale, const FVector*
 	FVector CurrentScale = GetRootComponent()->GetRelativeScale3D();
 	
 	// Lock X and Y scaling to the same value :
-	FVector2f XYDeltaScaleAbs(FMath::Abs(DeltaScale.X), FMath::Abs(DeltaScale.Y));
+	FVector2f XYDeltaScaleAbs(FMath::Abs(static_cast<float>(DeltaScale.X)), FMath::Abs(static_cast<float>(DeltaScale.Y)));
 	// Preserve the sign of the chosen delta :
 	bool bFavorX = (XYDeltaScaleAbs.X > XYDeltaScaleAbs.Y);
 
 	if (AActor::bUsePercentageBasedScaling)
 	{
 		// Correct for attempts to scale to 0 on any axis
-		float XYDeltaScale = bFavorX ? DeltaScale.X : DeltaScale.Y;
+		double XYDeltaScale = bFavorX ? DeltaScale.X : DeltaScale.Y;
 		if (XYDeltaScale == -1.0f)
 		{
 			XYDeltaScale = -(CurrentScale.X - 1) / CurrentScale.X;
@@ -4378,8 +4379,8 @@ void ALandscapeProxy::EditorApplyScale(const FVector& DeltaScale, const FVector*
 	else
 	{
 		// The absolute value of X and Y must be preserved so make sure they are preserved in case they flip from positive to negative (e.g.: a (-X, X) scale is accepted) : 
-		float SignMultiplier = FMath::Sign(CurrentScale.X) * FMath::Sign(CurrentScale.Y);
-		FVector2f NewScale(FVector2f::ZeroVector);
+		const float SignMultiplier = static_cast<float>(FMath::Sign(CurrentScale.X) * FMath::Sign(CurrentScale.Y));
+		FVector2d NewScale(FVector2f::ZeroVector);
 		if (bFavorX)
 		{
 			NewScale.X = CurrentScale.X + DeltaScale.X;
@@ -5175,7 +5176,7 @@ void ALandscape::SplitHeightmap(ULandscapeComponent* Comp, ALandscapeProxy* Targ
 		// Build a list of all unique materials the landscape uses
 		TArray<UMaterialInterface*> LandscapeMaterials;
 
-		int8 MaxLOD = FMath::CeilLogTwo(Comp->SubsectionSizeQuads + 1) - 1;
+		int8 MaxLOD = static_cast<int8>(FMath::CeilLogTwo(Comp->SubsectionSizeQuads + 1) - 1);
 
 		for (int8 LODIndex = 0; LODIndex < MaxLOD; ++LODIndex)
 		{
@@ -5287,7 +5288,7 @@ namespace
 		// Change Lighting resolution to proper one...
 		if (StaticLightingResolution > 1.0f)
 		{
-			StaticLightingResolution = (int32)StaticLightingResolution;
+			StaticLightingResolution = static_cast<float>(static_cast<int32>(StaticLightingResolution));
 		}
 		else if (StaticLightingResolution < 1.0f)
 		{
@@ -6441,8 +6442,8 @@ void ULandscapeComponent::ReallocateWeightmapsInternal(FLandscapeEditDataInterfa
 							DataInterface->ZeroTextureChannel(ComponentWeightmapTextures[TexIdx], ChanIdx);
 						}
 
-						AllocInfo.WeightmapTextureIndex = TexIdx;
-						AllocInfo.WeightmapTextureChannel = ChanIdx;
+						AllocInfo.WeightmapTextureIndex = IntCastChecked<uint8>(TexIdx);
+						AllocInfo.WeightmapTextureChannel = IntCastChecked<uint8>(ChanIdx);
 
 						Usage->ChannelUsage[ChanIdx] = this;
 
@@ -6581,8 +6582,8 @@ void ULandscapeComponent::ReallocateWeightmapsInternal(FLandscapeEditDataInterfa
 
 				// Assign the new allocation
 				CurrentWeightmapUsage->ChannelUsage[ChanIdx] = this;
-				AllocInfo.WeightmapTextureIndex = NewWeightmapTextures.Num() - 1;
-				AllocInfo.WeightmapTextureChannel = ChanIdx;
+				AllocInfo.WeightmapTextureIndex = IntCastChecked<uint8>(NewWeightmapTextures.Num() - 1);
+				AllocInfo.WeightmapTextureChannel = IntCastChecked<uint8>(ChanIdx);
 				CurrentLayer++;
 				TotalNeededChannels--;
 			}
@@ -7009,11 +7010,11 @@ void ULandscapeComponent::GenerateMobilePlatformPixelData(bool bIsCooking, const
 					MobileWeightmapTextures.Add(CurrentWeightmapTexture);
 				}
 
-				CopyRequests.AddWeightmapCopy(CurrentWeightmapTexture, CurrentChannel, this, Allocation.LayerInfo);
+				CopyRequests.AddWeightmapCopy(CurrentWeightmapTexture, IntCastChecked<int8>(CurrentChannel), this, Allocation.LayerInfo);
 
 				// update Allocation
-				Allocation.WeightmapTextureIndex = MobileWeightmapTextures.Num() - 1;
-				Allocation.WeightmapTextureChannel = CurrentChannel;
+				Allocation.WeightmapTextureIndex = IntCastChecked<uint8>(MobileWeightmapTextures.Num() - 1);
+				Allocation.WeightmapTextureChannel = IntCastChecked<uint8>(CurrentChannel);
 				CurrentChannel++;
 				RemainingChannels--;
 			}
@@ -7084,9 +7085,9 @@ void ULandscapeComponent::GenerateMobilePlatformPixelData(bool bIsCooking, const
 
 		if (MaterialPerLOD.Num() == 0)
 		{
-			int32 MaxLOD = FMath::CeilLogTwo(SubsectionSizeQuads + 1) - 1;
+			const int8 MaxLOD = static_cast<int8>(FMath::CeilLogTwo(SubsectionSizeQuads + 1) - 1);
 
-			for (int32 LODIndex = 0; LODIndex <= MaxLOD; ++LODIndex)
+			for (int8 LODIndex = 0; LODIndex <= MaxLOD; ++LODIndex)
 			{
 				UMaterialInterface* CurrentMaterial = GetLandscapeMaterial(LODIndex);
 
@@ -7233,11 +7234,11 @@ TArray<FLinearColor> ALandscapeProxy::SampleRTData(UTextureRenderTarget2D* InRen
 
 			FTextureRenderTargetResource* RTResource = InRenderTarget->GameThread_GetRenderTargetResource();
 
-			InRect.R = FMath::Clamp(int(InRect.R), 0, InRenderTarget->SizeX - 1);
-			InRect.G = FMath::Clamp(int(InRect.G), 0, InRenderTarget->SizeY - 1);
-			InRect.B = FMath::Clamp(int(InRect.B), int(InRect.R + 1), InRenderTarget->SizeX);
-			InRect.A = FMath::Clamp(int(InRect.A), int(InRect.G + 1), InRenderTarget->SizeY);
-			FIntRect Rect = FIntRect(InRect.R, InRect.G, InRect.B, InRect.A);
+			InRect.R = static_cast<float>(FMath::Clamp(int(InRect.R), 0, InRenderTarget->SizeX - 1));
+			InRect.G = static_cast<float>(FMath::Clamp(int(InRect.G), 0, InRenderTarget->SizeY - 1));
+			InRect.B = static_cast<float>(FMath::Clamp(int(InRect.B), int(InRect.R + 1), InRenderTarget->SizeX));
+			InRect.A = static_cast<float>(FMath::Clamp(int(InRect.A), int(InRect.G + 1), InRenderTarget->SizeY));
+			FIntRect Rect = FIntRect(static_cast<int32>(InRect.R), static_cast<int32>(InRect.G), static_cast<int32>(InRect.B), static_cast<int32>(InRect.A));
 
 			FReadSurfaceDataFlags ReadPixelFlags(RCM_MinMax);
 
@@ -7456,8 +7457,8 @@ bool ALandscapeProxy::LandscapeExportHeightmapToRenderTarget(UTextureRenderTarge
 				SubSectionSectionBase.Y += Component->SubsectionSizeQuads * SubY;
 
 				// Offset for this component's data in heightmap texture
-				float HeightmapOffsetU = Component->HeightmapScaleBias.Z + HeightmapSubsectionOffsetU * (float)SubX;
-				float HeightmapOffsetV = Component->HeightmapScaleBias.W + HeightmapSubsectionOffsetV * (float)SubY;
+				const float HeightmapOffsetU = static_cast<float>(Component->HeightmapScaleBias.Z) + HeightmapSubsectionOffsetU * SubX;
+				const float HeightmapOffsetV = static_cast<float>(Component->HeightmapScaleBias.W) + HeightmapSubsectionOffsetV * SubY;
 
 				FCanvasUVTri Tri1;
 				Tri1.V0_Pos = FVector2D(SubSectionSectionBase.X, SubSectionSectionBase.Y);
@@ -7535,7 +7536,7 @@ bool ALandscapeProxy::LandscapeImportWeightmapFromRenderTarget(UTextureRenderTar
 		{
 			const uint32 LandscapeWidth = (uint32)(1 + MaxX - MinX);
 			const uint32 LandscapeHeight = (uint32)(1 + MaxY - MinY);
-			FLinearColor SampleRect = FLinearColor(0, 0, LandscapeWidth, LandscapeHeight);
+			const FLinearColor SampleRect = FLinearColor(0.0f, 0.0f, static_cast<float>(LandscapeWidth), static_cast<float>(LandscapeHeight));
 
 			const uint32 RTWidth = InRenderTarget->SizeX;
 			const uint32 RTHeight = InRenderTarget->SizeY;
diff --git a/Engine/Source/Runtime/Landscape/Private/LandscapeEditResources.cpp b/Engine/Source/Runtime/Landscape/Private/LandscapeEditResources.cpp
--- a/Engine/Source/Runtime/Landscape/Private/LandscapeEditResources.cpp
+++ b/Engine/Source/Runtime/Landscape/Private/LandscapeEditResources.cpp
@@ -28,7 +28,7 @@ void FLandscapeTexture2DResource::InitRHI()
 
 	FRHITextureCreateDesc Desc =
 		FRHITextureCreateDesc::Create2D(TEXT("FLandscapeTexture2DResource"), SizeX, SizeY, Format)
-		.SetNumMips(NumMips);
+		.SetNumMips(static_cast<uint8>(NumMips));
 
 	if (bCreateUAVs)
 	{
@@ -53,7 +53,7 @@ void FLandscapeTexture2DResource::InitRHI()
 
 	if (bCreateSRV)
 	{
-		TextureSRV = RHICreateShaderResourceView(TextureRHI, /*MipLevel = */0, NumMips, Format);
+		TextureSRV = RHICreateShaderResourceView(TextureRHI, /*MipLevel = */0, static_cast<uint8>(NumMips), Format);
 	}
 }
 
diff --git a/Engine/Source/Runtime/Landscape/Private/LandscapeHLODBuilder.cpp b/Engine/Source/Runtime/Landscape/Private/LandscapeHLODBuilder.cpp
--- a/Engine/Source/Runtime/Landscape/Private/LandscapeHLODBuilder.cpp
+++ b/Engine/Source/Runtime/Landscape/Private/LandscapeHLODBuilder.cpp
@@ -138,7 +138,7 @@ static int32 GetMeshTextureSizeFromTargetTexelDensity(const FMeshDescription& In
 
 	// Compute the perfect texture size that would get us to our texture density
 	// Also compute the nearest power of two sizes (below and above our target)
-	const int32 SizePerfect = FMath::CeilToInt(InTargetTexelDensity / TexelRatio);
+	const int32 SizePerfect = FMath::CeilToInt32(InTargetTexelDensity / TexelRatio);
 	const int32 SizeHi = FMath::RoundUpToPowerOfTwo(SizePerfect);
 	const int32 SizeLo = SizeHi >> 1;
 
@@ -166,9 +166,9 @@ static int32 ComputeRequiredLandscapeLOD(const ALandscapeProxy* InLandscapeProxy
 
 	TArray<float> LODScreenSizes = InLandscapeProxy->GetLODScreenSizeArray();
 
-	ULandscapeComponent* LSComponent = InLandscapeProxy->LandscapeComponents[0];
-	float ComponentRadiusScaled = LSComponent->GetLocalBounds().SphereRadius * LSComponent->GetComponentTransform().GetScale3D().GetAbsMax();
-	float ExpectedScreenSize = ComputeBoundsScreenSize(FVector::ZeroVector, ComponentRadiusScaled, FVector(0.0f, 0.0f, InViewDistance), ProjMatrix);
+	const ULandscapeComponent* LSComponent = InLandscapeProxy->LandscapeComponents[0];
+	const float ComponentRadiusScaled = static_cast<float>(LSComponent->GetLocalBounds().SphereRadius * LSComponent->GetComponentTransform().GetScale3D().GetAbsMax());
+	const float ExpectedScreenSize = ComputeBoundsScreenSize(FVector::ZeroVector, ComponentRadiusScaled, FVector(0.0f, 0.0f, InViewDistance), ProjMatrix);
 
 	int32 RequiredLOD;
 	for (RequiredLOD = 0; RequiredLOD < LODScreenSizes.Num(); ++RequiredLOD)
@@ -264,7 +264,7 @@ TArray<UActorComponent*> ULandscapeHLODBuilder::Build(const FHLODBuildContext& I
 		FMeshDescription* MeshDescription = nullptr;
 
 		// Compute source landscape LOD
-		int32 LandscapeLOD = ComputeRequiredLandscapeLOD(LandscapeProxy, InHLODBuildContext.MinVisibleDistance);
+		int32 LandscapeLOD = ComputeRequiredLandscapeLOD(LandscapeProxy, static_cast<float>(InHLODBuildContext.MinVisibleDistance));
 
 		// Mesh
 		{
@@ -292,7 +292,8 @@ TArray<UActorComponent*> ULandscapeHLODBuilder::Build(const FHLODBuildContext& I
 
 		// Material
 		{
-			float TargetTexelDensityPerMeter = FMaterialUtilities::ComputeRequiredTexelDensityFromDrawDistance(InHLODBuildContext.MinVisibleDistance, MeshDescription->GetBounds().SphereRadius);
+			const float TargetTexelDensityPerMeter = FMaterialUtilities::ComputeRequiredTexelDensityFromDrawDistance(
+				static_cast<float>(InHLODBuildContext.MinVisibleDistance), static_cast<float>(MeshDescription->GetBounds().SphereRadius));
 			int32 RequiredTextureSize = GetMeshTextureSizeFromTargetTexelDensity(*MeshDescription, TargetTexelDensityPerMeter);
 
 			UMaterialInterface* LandscapeMaterial = BakeLandscapeMaterial(InHLODBuildContext, *MeshDescription, LandscapeProxy, LandscapeLOD, RequiredTextureSize);
@@ -304,7 +305,7 @@ TArray<UActorComponent*> ULandscapeHLODBuilder::Build(const FHLODBuildContext& I
 		StaticMeshes.Add(StaticMesh);
 
 		ULandscapeMeshProxyComponent* LandcapeMeshProxyComponent = NewObject<ULandscapeMeshProxyComponent>();
-		LandcapeMeshProxyComponent->InitializeForLandscape(LandscapeProxy, LandscapeLOD);
+		LandcapeMeshProxyComponent->InitializeForLandscape(LandscapeProxy, static_cast<int8>(LandscapeLOD));
 		LandcapeMeshProxyComponent->SetStaticMesh(StaticMesh);
 		HLODComponents.Add(LandcapeMeshProxyComponent);
 	}
diff --git a/Engine/Source/Runtime/Landscape/Private/LandscapeLight.cpp b/Engine/Source/Runtime/Landscape/Private/LandscapeLight.cpp
--- a/Engine/Source/Runtime/Landscape/Private/LandscapeLight.cpp
+++ b/Engine/Source/Runtime/Landscape/Private/LandscapeLight.cpp
@@ -174,6 +174,28 @@ FLandscapeStaticLightingMesh::~FLandscapeStaticLightingMesh()
 
 namespace
 {
+	template <typename T, typename U>
+	T MultiLerp(T A, T B, U F1, T C, T D, U F2, U F3)
+	{
+		return static_cast<T>(FMath::RoundToInt32(
+			FMath::Lerp<float>(
+				FMath::Lerp<float>(A, B, F1),
+				FMath::Lerp<float>(C, D, F2),
+				F3
+			)
+		));
+	}
+
+	template <typename U>
+	float MultiLerp(int32 A, int32 B, U F1, int32 C, int32 D, U F2, U F3)
+	{
+		return FMath::Lerp(
+			FMath::Lerp(static_cast<float>(A), static_cast<float>(B), static_cast<float>(F1)),
+			FMath::Lerp(static_cast<float>(C), static_cast<float>(D), static_cast<float>(F2)),
+			static_cast<float>(F3)
+		);
+	}
+
 	void GetLODData(ULandscapeComponent* LandscapeComponent, int32 X, int32 Y, int32 HeightmapOffsetX, int32 HeightmapOffsetY, int32 LODValue, int32 HeightmapStride, FColor& OutHeight, FColor& OutXYOffset)
 	{
 		int32 ComponentSize = ((LandscapeComponent->SubsectionSizeQuads + 1) * LandscapeComponent->NumSubsections) >> LODValue;
@@ -201,18 +223,9 @@ namespace
 		FColor H3 = HeightMipData[XI + FMath::Min(YI + 1, LODHeightmapSizeY - 1) * LODHeightmapSizeX];
 		FColor H4 = HeightMipData[FMath::Min(XI + 1, LODHeightmapSizeX - 1) + FMath::Min(YI + 1, LODHeightmapSizeY - 1) * LODHeightmapSizeX];
 
-		const uint16 Height = static_cast<uint16>(FMath::RoundToInt(FMath::Lerp(
-			FMath::Lerp<float>(static_cast<float>((H1.R << 8) + H1.G), (H2.R << 8) + H2.G, XF),
-			FMath::Lerp<float>(static_cast<float>((H3.R << 8) + H3.G), (H4.R << 8) + H4.G, XF),
-			YF)));
-		const uint8 B = static_cast<uint8>(FMath::RoundToInt(FMath::Lerp(
-			FMath::Lerp<float>(H1.B, H2.B, XF),
-			FMath::Lerp<float>(H3.B, H4.B, XF),
-			YF)));
-		const uint8 A = static_cast<uint8>(FMath::RoundToInt(FMath::Lerp<float>(
-			FMath::Lerp<float>(H1.A, H2.A, XF),
-			FMath::Lerp<float>(H3.A, H4.A, XF),
-			YF)));
+		uint16 Height = MultiLerp<uint16, float>((H1.R << 8) + H1.G, (H2.R << 8) + H2.G, XF, (H3.R << 8) + H3.G, (H4.R << 8) + H4.G, XF, YF);
+		uint8 B = MultiLerp<uint8, float>(H1.B, H2.B, XF, H3.B, H4.B, XF, YF);
+		uint8 A = MultiLerp<uint8, float>(H1.A, H2.A, XF, H3.A, H4.A, XF, YF);
 
 		OutHeight = FColor((Height >> 8), Height & 255, B, A);
 
@@ -223,14 +236,8 @@ namespace
 			FColor X3 = XYOffsetMipData[XI + FMath::Min(YI + 1, LODHeightmapSizeY - 1) * LODHeightmapSizeX];
 			FColor X4 = XYOffsetMipData[FMath::Min(XI + 1, LODHeightmapSizeX - 1) + FMath::Min(YI + 1, LODHeightmapSizeY - 1) * LODHeightmapSizeX];
 
-			const uint16 XComp = static_cast<uint16>(FMath::RoundToInt(FMath::Lerp(
-				FMath::Lerp<float>(static_cast<float>((X1.R << 8) + X1.G), (X2.R << 8) + X2.G, XF),
-				FMath::Lerp<float>(static_cast<float>((X3.R << 8) + X3.G), (X4.R << 8) + X4.G, XF),
-				YF)));
-			const uint16 YComp = static_cast<uint16>(FMath::RoundToInt(FMath::Lerp(
-				FMath::Lerp<float>(static_cast<float>((X1.B << 8) + X1.A), (X2.B << 8) + X2.A, XF),
-				FMath::Lerp<float>(static_cast<float>((X3.B << 8) + X3.A), (X4.B << 8) + X4.A, XF),
-				YF)));
+			uint16 XComp = MultiLerp<uint16, float>((X1.R << 8) + X1.G, (X2.R << 8) + X2.G, XF, (X3.R << 8) + X3.G, (X4.R << 8) + X4.G, XF, YF);
+			uint16 YComp = MultiLerp<uint16, float>((X1.B << 8) + X1.A, (X2.B << 8) + X2.A, XF, (X3.B << 8) + X3.A, (X4.B << 8) + X4.A, XF, YF);
 
 			OutXYOffset = FColor((XComp >> 8), XComp & 255, YComp >> 8, YComp & 255);
 		}
@@ -259,7 +266,7 @@ namespace
 				}
 
 				ULandscapeComponent* Neighbor = Info->XYtoComponentMap.FindRef(ComponentBase + FIntPoint(x, y));
-				int32 NeighborLOD = -1;
+				int32 NeighborLOD;
 
 				if (Neighbor)
 				{
@@ -337,34 +344,22 @@ namespace
 						{
 							if (XY.Y < 0.f)
 							{
-								RealLOD = FMath::Lerp(
-									FMath::Lerp<float>(static_cast<float>(NeighborLODs[0]), NeighborLODs[1], XY.X + 1.f),
-									FMath::Lerp<float>(static_cast<float>(NeighborLODs[3]), GeometryLOD, XY.X + 1.f),
-									XY.Y + 1.f); // 0
+								RealLOD = MultiLerp(NeighborLODs[0], NeighborLODs[1], XY.X + 1.f, NeighborLODs[3], GeometryLOD, XY.X + 1.f, XY.Y + 1.f); // 0
 							}
 							else
 							{
-								RealLOD = FMath::Lerp(
-									FMath::Lerp<float>(static_cast<float>(NeighborLODs[3]), GeometryLOD, XY.X + 1.f),
-									FMath::Lerp<float>(static_cast<float>(NeighborLODs[5]), NeighborLODs[6], XY.X + 1.f),
-									XY.Y); // 2
+								RealLOD = MultiLerp(NeighborLODs[3], GeometryLOD, XY.X + 1.f, NeighborLODs[5], NeighborLODs[6], XY.X + 1.f, XY.Y); // 2
 							}
 						}
 						else
 						{
 							if (XY.Y < 0.f)
 							{
-								RealLOD = FMath::Lerp(
-									FMath::Lerp<float>(static_cast<float>(NeighborLODs[1]), NeighborLODs[2], XY.X),
-									FMath::Lerp<float>(static_cast<float>(GeometryLOD), NeighborLODs[4], XY.X),
-									XY.Y + 1.f); // 1
+								RealLOD = MultiLerp(NeighborLODs[1], NeighborLODs[2], XY.X, GeometryLOD, NeighborLODs[4], XY.X, XY.Y + 1.f); // 1
 							}
 							else
 							{
-								RealLOD = FMath::Lerp(
-									FMath::Lerp<float>(static_cast<float>(GeometryLOD), NeighborLODs[4], XY.X),
-									FMath::Lerp<float>(static_cast<float>(NeighborLODs[6]), NeighborLODs[7], XY.X),
-									XY.Y); // 3
+								RealLOD = MultiLerp(GeometryLOD, NeighborLODs[4], XY.X, NeighborLODs[6], NeighborLODs[7], XY.X, XY.Y); // 3
 							}
 						}
 
@@ -387,21 +382,21 @@ namespace
 							// Need interpolation
 							const uint16 Height0 = (Height[0].R << 8) + Height[0].G;
 							const uint16 Height1 = (Height[1].R << 8) + Height[1].G;
-							const uint16 LerpHeight = static_cast<uint16>(FMath::RoundToInt(FMath::Lerp<float>(Height0, Height1, MorphAlpha)));
+							const uint16 LerpHeight = static_cast<uint16>(FMath::RoundToInt32(FMath::Lerp<float>(Height0, Height1, MorphAlpha)));
 
 							CompHeightData[X + Y * HeightmapStride] =
 								FColor((LerpHeight >> 8), LerpHeight & 255,
-								static_cast<uint8>(FMath::RoundToInt(FMath::Lerp<float>(Height[0].B, Height[1].B, MorphAlpha))),
-								static_cast<uint8>(FMath::RoundToInt(FMath::Lerp<float>(Height[0].A, Height[1].A, MorphAlpha))));
+								static_cast<uint8>(FMath::RoundToInt32(FMath::Lerp<float>(Height[0].B, Height[1].B, MorphAlpha))),
+								static_cast<uint8>(FMath::RoundToInt32(FMath::Lerp<float>(Height[0].A, Height[1].A, MorphAlpha))));
 							if (LandscapeComponent->XYOffsetmapTexture)
 							{
 								uint16 XComp0 = (XYOffset[0].R << 8) + XYOffset[0].G;
 								uint16 XComp1 = (XYOffset[1].R << 8) + XYOffset[1].G;
-								uint16 LerpXComp = static_cast<uint16>(FMath::RoundToInt(FMath::Lerp<float>(XComp0, XComp1, MorphAlpha)));
+								uint16 LerpXComp = static_cast<uint16>(FMath::RoundToInt32(FMath::Lerp<float>(XComp0, XComp1, MorphAlpha)));
 
 								uint16 YComp0 = (XYOffset[0].B << 8) + XYOffset[0].A;
 								uint16 YComp1 = (XYOffset[1].B << 8) + XYOffset[1].A;
-								uint16 LerpYComp = static_cast<uint16>(FMath::RoundToInt(FMath::Lerp<float>(YComp0, YComp1, MorphAlpha)));
+								uint16 LerpYComp = static_cast<uint16>(FMath::RoundToInt32(FMath::Lerp<float>(YComp0, YComp1, MorphAlpha)));
 
 								CompXYOffsetData[X + Y * HeightmapStride] =
 									FColor(LerpXComp >> 8, LerpXComp & 255, LerpYComp >> 8, LerpYComp & 255);
@@ -452,8 +447,6 @@ void FLandscapeStaticLightingMesh::GetHeightmapData(int32 InLOD, int32 GeometryL
 	check(ExpandQuadsX <= SubsectionSizeQuads);
 	check(ExpandQuadsY <= SubsectionSizeQuads);
 
-	int32 MaxLOD = FMath::CeilLogTwo(LandscapeComponent->SubsectionSizeQuads + 1) - 1;
-
 	// copy heightmap data for this component...
 	{
 		// Data array for upscaling case
@@ -634,18 +627,17 @@ void FLandscapeStaticLightingMesh::GetTriangleIndices(int32 TriangleIndex,int32&
 	int32 QuadX = QuadIndex % (NumVertices - 1);
 	int32 QuadY = QuadIndex / (NumVertices - 1);
 
-	switch(QuadTriIndex)
+	if (QuadTriIndex == 0)
 	{
-	case 0:
 		OutI0 = (QuadX + 0) + (QuadY + 0) * NumVertices;
 		OutI1 = (QuadX + 1) + (QuadY + 1) * NumVertices;
 		OutI2 = (QuadX + 1) + (QuadY + 0) * NumVertices;
-		break;
-	case 1:
+	}
+	else // QuadTriIndex == 1
+	{
 		OutI0 = (QuadX + 0) + (QuadY + 0) * NumVertices;
 		OutI1 = (QuadX + 0) + (QuadY + 1) * NumVertices;
 		OutI2 = (QuadX + 1) + (QuadY + 1) * NumVertices;
-		break;
 	}
 
 	if (bReverseWinding)
@@ -661,7 +653,6 @@ FLightRayIntersection FLandscapeStaticLightingMesh::IntersectLightRay(const FVec
 	// Intersect the light ray with the terrain component.
 	FHitResult Result(1.0f);
 
-	FHitResult NewHitInfo;
 	FCollisionQueryParams NewTraceParams(SCENE_QUERY_STAT(FLandscapeStaticLightingMesh_IntersectLightRay), true );
 	
 	const bool bIntersects = LandscapeComponent->LineTraceComponent( Result, Start, End, NewTraceParams );
@@ -722,7 +713,7 @@ bool ULandscapeComponent::GetLightMapResolution( int32& Width, int32& Height ) c
 	int32 DesiredSize = 1;
 	uint32 LightingLOD = GetLandscapeProxy()->StaticLightingLOD;
 
-	float LightMapRatio = ::GetTerrainExpandPatchCount(LightMapRes, PatchExpandCountX, PatchExpandCountY, ComponentSizeQuads, (NumSubsections * (SubsectionSizeQuads+1)), DesiredSize, LightingLOD);
+	GetTerrainExpandPatchCount(LightMapRes, PatchExpandCountX, PatchExpandCountY, ComponentSizeQuads, (NumSubsections * (SubsectionSizeQuads+1)), DesiredSize, LightingLOD);
 
 	Width = DesiredSize;
 	Height = DesiredSize;
diff --git a/Engine/Source/Runtime/Landscape/Private/LandscapeSplines.cpp b/Engine/Source/Runtime/Landscape/Private/LandscapeSplines.cpp
--- a/Engine/Source/Runtime/Landscape/Private/LandscapeSplines.cpp
+++ b/Engine/Source/Runtime/Landscape/Private/LandscapeSplines.cpp
@@ -222,7 +222,7 @@ public:
 		{
 			if (VisibilityMap & (1 << ViewIndex))
 			{
-				const FSceneView* View = Views[ViewIndex];
+				const FSceneView* const View = Views[ViewIndex];
 				FPrimitiveDrawInterface* PDI = Collector.GetPDI(ViewIndex);
 				const uint8 DepthPriority = static_cast<uint8>(GetDepthPriorityGroup(View));
 
@@ -281,8 +281,8 @@ public:
 						double ControlPointSpriteScale = MyLocalToWorld.GetScaleVector().X * ControlPoint.SpriteScale;
 						const FLinearColor ControlPointSpriteColor = ControlPoint.bSelected ? SelectedControlPointSpriteColor : FLinearColor::White;
 						PDI->SetHitProxy(ControlPoint.HitProxy);
-						const FMatrix& ProectionMatrix = View->ViewMatrices.GetProjectionMatrix();
-						const double ZoomFactor = FMath::Min<double>(ProectionMatrix.M[0][0], ProectionMatrix.M[1][1]);
+						const FMatrix& ProjectionMatrix = View->ViewMatrices.GetProjectionMatrix();
+						const double ZoomFactor = FMath::Min<double>(ProjectionMatrix.M[0][0], ProjectionMatrix.M[1][1]);
 						const double Scale = View->WorldToScreen(ControlPointLocation).W * (4.0f / View->UnscaledViewRect.Width() / ZoomFactor);
 						ControlPointSpriteScale *= Scale;
 
@@ -1758,7 +1758,7 @@ bool ULandscapeSplineControlPoint::SupportsForeignSplineMesh() const
 FName ULandscapeSplineControlPoint::GetBestConnectionTo(FVector Destination) const
 {
 	FName BestSocket = NAME_None;
-	float BestScore = -FLT_MAX;
+	double BestScore = -DBL_MAX;
 
 	if (Mesh != nullptr)
 	{
@@ -1768,8 +1768,8 @@ FName ULandscapeSplineControlPoint::GetBestConnectionTo(FVector Destination) con
 			FVector SocketLocation = SocketTransform.GetTranslation();
 			FRotator SocketRotation = SocketTransform.GetRotation().Rotator();
 
-			const float Score = static_cast<float>((Destination - Location).Size() - (Destination - SocketLocation).Size() // Score closer sockets higher
-				* FMath::Abs(FVector::DotProduct((Destination - SocketLocation), SocketRotation.Vector()))); // score closer rotation higher
+			const double Score = (Destination - Location).Size() - (Destination - SocketLocation).Size() // Score closer sockets higher
+				* FMath::Abs(FVector::DotProduct((Destination - SocketLocation), SocketRotation.Vector())); // score closer rotation higher
 
 			if (Score > BestScore)
 			{
@@ -2691,7 +2691,7 @@ void ULandscapeSplineSegment::AutoFlipTangents()
 
 static float ApproxLength(const FInterpCurveVector& SplineInfo, const float Start = 0.0f, const float End = 1.0f, const int32 ApproxSections = 4)
 {
-	float SplineLength = 0;
+	double SplineLength = 0;
 	FVector OldPos = SplineInfo.Eval(Start, FVector::ZeroVector);
 	for (int32 i = 1; i <= ApproxSections; i++)
 	{
@@ -2700,7 +2700,7 @@ static float ApproxLength(const FInterpCurveVector& SplineInfo, const float Star
 		OldPos = NewPos;
 	}
 
-	return SplineLength;
+	return static_cast<float>(SplineLength);
 }
 
 static ESplineMeshAxis::Type CrossAxis(ESplineMeshAxis::Type InForwardAxis, ESplineMeshAxis::Type InUpAxis)
diff --git a/Engine/Source/Runtime/Landscape/Private/LandscapeSubsystem.cpp b/Engine/Source/Runtime/Landscape/Private/LandscapeSubsystem.cpp
--- a/Engine/Source/Runtime/Landscape/Private/LandscapeSubsystem.cpp
+++ b/Engine/Source/Runtime/Landscape/Private/LandscapeSubsystem.cpp
@@ -407,7 +407,7 @@ void ULandscapeSubsystem::BuildNanite(TArrayView<ALandscapeProxy*> InProxiesToBu
 	// Don't keep those that are null or already up to date :
 	FinalProxiesToBuild.SetNum(Algo::RemoveIf(FinalProxiesToBuild, [bForceRebuild](ALandscapeProxy* InProxy) { return (InProxy == nullptr) || (!bForceRebuild && InProxy->IsNaniteMeshUpToDate()); }));
 
-	FScopedSlowTask SlowTask(FinalProxiesToBuild.Num(), (LOCTEXT("Landscape_BuildNanite", "Building Nanite Landscape Meshes")));
+	FScopedSlowTask SlowTask(static_cast<float>(FinalProxiesToBuild.Num()), (LOCTEXT("Landscape_BuildNanite", "Building Nanite Landscape Meshes")));
 	SlowTask.MakeDialog(/*bShowCancelButton = */true);
 
 	for (ALandscapeProxy* Proxy : FinalProxiesToBuild)
