Add support for skipping landscape ray tracing geometry updates if fractional LOD changes by less than a threshold

* Add CVar r.RayTracing.Geometry.Landscape.FractionalLODUpdateThreshold (default: 0, keeping previous behavior of immediate BLAS update on any fractional LOD change)

Code contributed by [REDACTED_PERSON], NVIDIA

diff --git a/Engine/Source/Runtime/Landscape/Private/LandscapeRender.cpp b/Engine/Source/Runtime/Landscape/Private/LandscapeRender.cpp
--- a/Engine/Source/Runtime/Landscape/Private/LandscapeRender.cpp
+++ b/Engine/Source/Runtime/Landscape/Private/LandscapeRender.cpp
@@ -163,6 +163,16 @@ static FAutoConsoleVariableRef CVarLandscapeRayTracingGeometryDetectTextureStrea
 	GLandscapeRayTracingGeometryDetectTextureStreaming,
 	TEXT("If on, update ray tracing geometry when texture streaming state changes. Useful when WorldPositionOffset is used in the landscape material")
 );
+
+float GLandscapeRayTracingGeometryFractionalLODUpdateThreshold = 0.0f;
+static FAutoConsoleVariableRef CVarLandscapeRayTracingGeometryFractionalLODUpdateThreshold(
+	TEXT("r.RayTracing.Geometry.Landscape.FractionalLODUpdateThreshold"),
+	GLandscapeRayTracingGeometryFractionalLODUpdateThreshold,
+	TEXT("Minimal difference in fractional LOD between latest built/cached ray tracing geometry and latest value used for rendering (default 0)\n")
+	TEXT("0.1 implies a 10% change in the fraction LOD, for instance a change from LOD level 1.1 to 1.2\n")
+	TEXT("Larger values will reduce the number of landscape tile updates, but introduce more error between the ray tracing and raster representations")
+);
+
 #endif
 
 //
@@ -2521,10 +2531,12 @@ void FLandscapeComponentSceneProxy::GetDynamicRayTracingInstances(FRayTracingMat
 					RayTracingState->Sections[SubSectionIdx].HeightmapLODBias = RenderSystem.GetSectionLODBias(ComponentBase);
 				}
 
-				if (RayTracingState->Sections[SubSectionIdx].FractionalLOD != RenderSystem.GetSectionLODValue(SceneView, ComponentBase))
+				const float PendingFractionalLOD = RenderSystem.GetSectionLODValue(SceneView, ComponentBase);
+				const float FractionLODAbsoluteDifference = FMath::Abs(RayTracingState->Sections[SubSectionIdx].FractionalLOD - PendingFractionalLOD);
+				if (FractionLODAbsoluteDifference > GLandscapeRayTracingGeometryFractionalLODUpdateThreshold)
 				{
 					bNeedsRayTracingGeometryUpdate = true;
-					RayTracingState->Sections[SubSectionIdx].FractionalLOD = RenderSystem.GetSectionLODValue(SceneView, ComponentBase);
+					RayTracingState->Sections[SubSectionIdx].FractionalLOD = PendingFractionalLOD;
 				}
 			}
 
