Lowered the priority of shader compile jobs submitted by water and landscape to High to allow them to distribute via XGE.

diff --git a/Engine/Plugins/Experimental/Water/Source/Runtime/Private/WaterZoneActor.cpp b/Engine/Plugins/Experimental/Water/Source/Runtime/Private/WaterZoneActor.cpp
--- a/Engine/Plugins/Experimental/Water/Source/Runtime/Private/WaterZoneActor.cpp
+++ b/Engine/Plugins/Experimental/Water/Source/Runtime/Private/WaterZoneActor.cpp
@@ -509,7 +509,7 @@ bool AWaterZone::UpdateWaterInfoTexture()
 					if (!MaterialResource->IsGameThreadShaderMapComplete())
 					{
 #if WITH_EDITOR
-						MaterialResource->SubmitCompileJobs_GameThread(EShaderCompileJobPriority::ForceLocal);
+						MaterialResource->SubmitCompileJobs_GameThread(EShaderCompileJobPriority::High);
 #endif
 						bHasIncompleteShaderMaps = true;
 					}
diff --git a/Engine/Source/Runtime/Landscape/Private/LandscapeEditLayers.cpp b/Engine/Source/Runtime/Landscape/Private/LandscapeEditLayers.cpp
--- a/Engine/Source/Runtime/Landscape/Private/LandscapeEditLayers.cpp
+++ b/Engine/Source/Runtime/Landscape/Private/LandscapeEditLayers.cpp
@@ -1802,7 +1802,7 @@ bool ALandscape::IsMaterialResourceCompiled(FMaterialResource* InMaterialResourc
 
 	{
 		TRACE_CPUPROFILER_EVENT_SCOPE(Landscape_WaitForMaterialCompilation);
-		InMaterialResource->SubmitCompileJobs_GameThread(EShaderCompileJobPriority::ForceLocal);
+		InMaterialResource->SubmitCompileJobs_GameThread(EShaderCompileJobPriority::High);
 		if (bInWaitForCompilation)
 		{
 			InMaterialResource->FinishCompilation();
