Added Landscape section culling on the GPU
Landscape sections are spit into small tiles (4x4 quads) which get frustum culled in compute for each view. Visible tiles then rendered using an indirect instanced draw
Active only for LOD0 and if section size is at least 31 quads
Culling for shadows views is not active yet

diff --git a/Engine/Shaders/Private/Landscape/LandscapeCulling.usf b/Engine/Shaders/Private/Landscape/LandscapeCulling.usf
new file mode 100644
--- /dev/null
+++ b/Engine/Shaders/Private/Landscape/LandscapeCulling.usf
@@ -0,0 +1,91 @@
+// Copyright Epic Games, Inc. All Rights Reserved.
+
+#include "../Common.ush"
+#include "../Nanite/NaniteHZBCull.ush"
+
+#ifndef LANDSCAPE_TILE_QUADS
+	#define LANDSCAPE_TILE_QUADS 4u
+#endif
+
+struct FLandscapeView
+{
+	float4x4 ViewToClip;
+	float4x4 TranslatedWorldToClip;
+	float3 RelativePreViewTranslation;
+	float _Pad0;
+	float3 ViewTilePosition;
+	float _Pad1;
+};
+
+struct FLandscapeSection
+{
+	float4x4 LocalToRelativeWorld;
+	float3 TilePosition;
+	float LocalZ;
+	float HalfHeight;
+	float NeighborLODExtent;
+	float _Pad[2];
+};
+
+StructuredBuffer<FLandscapeView> LandscapeViews;
+StructuredBuffer<FLandscapeSection> LandscapeSections;
+
+uint NumLandscapeSections;
+uint NumLandscapeViews;
+uint SubsectionSizeTiles;
+uint NumSubsections;
+uint NearClip;
+
+RWBuffer<uint> IndirectArgsBufferOut;
+RWByteAddressBuffer TileDataOut;
+
+// XY for each tile in a section
+// Z for each NumLandscapeSections * NumLandscapeViews 
+[numthreads(8, 8, 1)]
+void BuildLandscapeTileDataCS(uint3 ThreadID : SV_DispatchThreadID) 
+{
+	uint SectionSizeTiles = SubsectionSizeTiles * NumSubsections;
+	if (ThreadID.x >= SectionSizeTiles)
+	{
+		return;
+	}
+
+	uint ViewIdx = ThreadID.z / NumLandscapeSections;
+	uint SectionIdx = ThreadID.z % NumLandscapeSections;
+	
+	FLandscapeSection Section = LandscapeSections[SectionIdx];
+	FLandscapeView LandscapeView = LandscapeViews[ViewIdx];
+
+	const bool bIsOrtho = IsOrthoProjection(LandscapeView.ViewToClip);
+	const bool bNearClip = (NearClip != 0u);
+	const bool bSkipFrustumCull = false;
+
+	float2 TileExtentQuads = LANDSCAPE_TILE_QUADS * 0.5f;
+	float2 TileCenterQuads = (ThreadID.xy * LANDSCAPE_TILE_QUADS) + TileExtentQuads;
+	// Inflate tile size to account for neighbors section LOD
+	TileExtentQuads += float2(Section.NeighborLODExtent, Section.NeighborLODExtent);
+
+	float3 BoxExtentLocal = float3(TileExtentQuads, Section.HalfHeight);
+	float3 BoxCenterLocal = float3(TileCenterQuads, Section.LocalZ);
+	
+	FLWCMatrix LocalToWorld	= MakeLWCMatrix4x3(Section.TilePosition, Section.LocalToRelativeWorld);
+	float4x4 LocalToTranslatedWorld = LWCMultiplyTranslation(LocalToWorld, MakeLWCVector3(LandscapeView.ViewTilePosition, LandscapeView.RelativePreViewTranslation));
+	
+	FFrustumCullData CullData = BoxCullFrustum(BoxCenterLocal, BoxExtentLocal, LocalToTranslatedWorld, LandscapeView.TranslatedWorldToClip, bIsOrtho, bNearClip, bSkipFrustumCull);
+	if (CullData.bIsVisible)
+	{
+		uint TileOutputIndex = 0u;
+		InterlockedAdd(IndirectArgsBufferOut[ThreadID.z * INDIRECT_ARGS_NUM_WORDS + 1u], 1u, TileOutputIndex);
+		{
+			uint SubsectionX = (ThreadID.x < SubsectionSizeTiles ? 0u : 1u);
+			uint SubsectionY = (ThreadID.y < SubsectionSizeTiles ? 0u : 1u);
+			uint QuadX = (ThreadID.x - SubsectionSizeTiles * SubsectionX) * LANDSCAPE_TILE_QUADS;
+			uint QuadY = (ThreadID.y - SubsectionSizeTiles * SubsectionY) * LANDSCAPE_TILE_QUADS;
+			
+			uint TileData4Bytes = (SubsectionY << 24u) | (SubsectionX << 16u) | (QuadY << 8u) | (QuadX << 0u);
+			uint TileGlobalIndex = TileOutputIndex + (ThreadID.z * SectionSizeTiles * SectionSizeTiles);
+
+			TileDataOut.Store(TileGlobalIndex * 4u, TileData4Bytes);
+		}
+	}
+}
diff --git a/Engine/Shaders/Private/LandscapeVertexFactory.ush b/Engine/Shaders/Private/LandscapeVertexFactory.ush
--- a/Engine/Shaders/Private/LandscapeVertexFactory.ush
+++ b/Engine/Shaders/Private/LandscapeVertexFactory.ush
@@ -151,10 +151,17 @@ float CalcLOD(uint ComponentIndex, float2 xyLocalToSubsection, float2 Subsection
 
 struct FVertexFactoryInput
 {
-	float4 Position : ATTRIBUTE0;
+	// UByte4
+	uint4 Position : ATTRIBUTE0;
+
+#if LANDSCAPE_TILE
+	// UByte4
+	uint4 TileData : ATTRIBUTE1;
+#else
 	// Dynamic instancing related attributes with InstanceIdOffset : ATTRIBUTE1
 	VF_GPUSCENE_DECLARE_INPUT_BLOCK(1)
 	VF_INSTANCED_STEREO_DECLARE_INPUT_BLOCK()
+#endif	
 };
 
 // RHI_RAYTRACING
@@ -587,6 +594,11 @@ FVertexFactoryIntermediates GetVertexFactoryIntermediates(FVertexFactoryInput In
 	Intermediates.LodBias = GetLodBias(Intermediates.ComponentIndex);
 	Intermediates.InputPosition = Input.Position;
 
+#if LANDSCAPE_TILE   
+	Intermediates.InputPosition = Intermediates.InputPosition + Input.TileData;
+	Intermediates.InputPosition.xy = min(Intermediates.InputPosition.xy, float2(Intermediates.LodValues.z, Intermediates.LodValues.z));
+#endif
+
 	float2 xy = Intermediates.InputPosition.xy * Intermediates.LodValues.w;
 
 	float LODCalculated = CalcLOD(Intermediates.ComponentIndex, xy, Intermediates.InputPosition.zw);
diff --git a/Engine/Source/Runtime/Engine/Private/PrimitiveSceneProxy.cpp b/Engine/Source/Runtime/Engine/Private/PrimitiveSceneProxy.cpp
--- a/Engine/Source/Runtime/Engine/Private/PrimitiveSceneProxy.cpp
+++ b/Engine/Source/Runtime/Engine/Private/PrimitiveSceneProxy.cpp
@@ -112,6 +112,9 @@ bool SupportsCachingMeshDrawCommands(const FMeshBatch& MeshBatch)
 		// Cached mesh commands only allow for a single mesh element per batch.
 		(MeshBatch.Elements.Num() == 1) &&
 
+		// View dependent arguments can't be cached
+		(MeshBatch.bViewDependentArguments == 0) &&
+
 		// Vertex factory needs to support caching.
 		MeshBatch.VertexFactory->GetType()->SupportsCachingMeshDrawCommands();
 }
diff --git a/Engine/Source/Runtime/Engine/Private/SceneManagement.cpp b/Engine/Source/Runtime/Engine/Private/SceneManagement.cpp
--- a/Engine/Source/Runtime/Engine/Private/SceneManagement.cpp
+++ b/Engine/Source/Runtime/Engine/Private/SceneManagement.cpp
@@ -1462,15 +1462,17 @@ bool FMeshBatch::Validate(const FPrimitiveSceneProxy* PrimitiveSceneProxy, ERHIF
 	}
 
 	const bool bVFSupportsPrimitiveIdStream = VertexFactory->GetType()->SupportsPrimitiveIdStream();
+	const bool bVFRequiresPrimitiveUniformBuffer = PrimitiveSceneProxy->DoesVFRequirePrimitiveUniformBuffer();
 
-	if (!PrimitiveSceneProxy->DoesVFRequirePrimitiveUniformBuffer() && !bVFSupportsPrimitiveIdStream)
+	if (!bVFRequiresPrimitiveUniformBuffer && !bVFSupportsPrimitiveIdStream)
 	{
 		return LogMeshError(TEXT("PrimitiveSceneProxy has bVFRequiresPrimitiveUniformBuffer disabled yet tried to draw with a vertex factory that did not support PrimitiveIdStream"));
 	}
 
-	if (PrimitiveSceneProxy->SupportsGPUScene() && !VertexFactory->SupportsGPUScene(FeatureLevel))
+	// Some primitives may use several VFs with a mixed support for a GPUScene
+	if (PrimitiveSceneProxy->SupportsGPUScene() && !(VertexFactory->SupportsGPUScene(FeatureLevel) || bVFRequiresPrimitiveUniformBuffer))
 	{
-		return LogMeshError(TEXT("PrimitiveSceneProxy has SupportsGPUScene() does not match VertexFactory->SupportsGPUScene()"));
+		return LogMeshError(TEXT("PrimitiveSceneProxy has SupportsGPUScene() does not match VertexFactory->SupportsGPUScene() or bVFRequiresPrimitiveUniformBuffer"));
 	}
 	const bool bUseGPUScene = UseGPUScene(GMaxRHIShaderPlatform, FeatureLevel);
 	
@@ -1482,7 +1484,9 @@ bool FMeshBatch::Validate(const FPrimitiveSceneProxy* PrimitiveSceneProxy, ERHIF
 	{
 		const FMeshBatchElement& MeshElement = Elements[ElementIndex];
 
-		if (bPrimitiveShaderDataComesFromSceneBuffer && Elements[ElementIndex].PrimitiveUniformBuffer)
+		// Some primitives may use several VFs with a mixed support for a GPUScene 
+		// in this case all mesh batches get Primitive UB assigned regardless of VF type
+		if (bPrimitiveShaderDataComesFromSceneBuffer && Elements[ElementIndex].PrimitiveUniformBuffer && !bVFRequiresPrimitiveUniformBuffer)
 		{
 			// on mobile VS has access to PrimitiveUniformBuffer
 			if (FeatureLevel > ERHIFeatureLevel::ES3_1)
diff --git a/Engine/Source/Runtime/Engine/Public/MeshBatch.h b/Engine/Source/Runtime/Engine/Public/MeshBatch.h
--- a/Engine/Source/Runtime/Engine/Public/MeshBatch.h
+++ b/Engine/Source/Runtime/Engine/Public/MeshBatch.h
@@ -397,6 +397,12 @@ struct FMeshBatch
 	uint32 CastRayTracedShadow : 1;	// Whether it casts ray traced shadow.
 #endif
 
+	/** 
+	 * Whether mesh has a view dependent draw arguments. 
+	 * Gives an opportunity to override mesh arguments for a View just before creating MDC for a mesh (makes MDC "non-cached")
+	 */
+	uint32 bViewDependentArguments : 1;
+
 #if UE_ENABLE_DEBUG_DRAWING
 	/** Conceptual HLOD index used for the HLOD Coloration visualization. */
 	int8 VisualizeHLODIndex;
@@ -475,6 +481,7 @@ struct FMeshBatch
 #if RHI_RAYTRACING
 	,	CastRayTracedShadow(true)
 #endif
+	,	bViewDependentArguments(false)
 #if (!(UE_BUILD_SHIPPING || UE_BUILD_TEST) || WITH_EDITOR)
 	,	VisualizeHLODIndex(INDEX_NONE)
 #endif
diff --git a/Engine/Source/Runtime/Engine/Public/PrimitiveSceneProxy.h b/Engine/Source/Runtime/Engine/Public/PrimitiveSceneProxy.h
--- a/Engine/Source/Runtime/Engine/Public/PrimitiveSceneProxy.h
+++ b/Engine/Source/Runtime/Engine/Public/PrimitiveSceneProxy.h
@@ -324,6 +324,14 @@ public:
 
 	virtual const class FCardRepresentationData* GetMeshCardRepresentation() const { return nullptr; }
 
+	/** 
+	* Gives the primitive an opportunity to override MeshBatch arguments for a specific View
+	* Only called for a MeshBatch with a bViewDependentArguments property set
+	* @param View - the view to override for
+	* @param ViewDependentMeshBatch - view dependent mesh copy (does not affect a cached FMeshBatch)
+	*/
+	virtual void ApplyViewDependentMeshArguments(const FSceneView& View, FMeshBatch& ViewDependentMeshBatch) const {};
+
 	/** 
 	 * Gets the boxes for sub occlusion queries
 	 * @param View - the view the occlusion results are for
diff --git a/Engine/Source/Runtime/Landscape/Private/LandscapeCulling.cpp b/Engine/Source/Runtime/Landscape/Private/LandscapeCulling.cpp
new file mode 100644
--- /dev/null
+++ b/Engine/Source/Runtime/Landscape/Private/LandscapeCulling.cpp
@@ -0,0 +1,832 @@
+// Copyright Epic Games, Inc. All Rights Reserved.
+
+#include "LandscapeCulling.h"
+#include "LandscapeRender.h"
+#include "ShaderParameterStruct.h"
+#include "RenderGraphBuilder.h"
+#include "RenderGraphUtils.h"
+#include "ShaderPlatformCachedIniValue.h"
+#include "MeshMaterialShader.h"
+#include "DataDrivenShaderPlatformInfo.h"
+
+extern TMap<uint32, FLandscapeRenderSystem*> LandscapeRenderSystems;
+
+static TAutoConsoleVariable<int32> CVarLandscapeSupportCulling(
+	TEXT("landscape.SupportGPUCulling"),
+	0,
+	TEXT("Whether to support landscape GPU culling"),
+	ECVF_ReadOnly | ECVF_RenderThreadSafe);
+
+static TAutoConsoleVariable<int32> CVarLandscapeEnableGPUCulling(
+	TEXT("landscape.EnableGPUCulling"),
+	1,
+	TEXT("Whether to use landscape GPU culling when it's supported. Allows to toggle culling at runtime"),
+	ECVF_RenderThreadSafe);
+
+static TAutoConsoleVariable<int32> CVarLandscapeEnableGPUCullingShadows(
+	TEXT("landscape.EnableGPUCullingShadows"),
+	1,
+	TEXT("Whether to use landscape GPU culling for a shadow views when it's supported. Allows to toggle shadow views culling at runtime"),
+	ECVF_RenderThreadSafe);
+
+
+/** Vertex factory for a tiled landscape rendering  */
+class FLandscapeTileVertexFactory : public FLandscapeVertexFactory
+{
+	DECLARE_VERTEX_FACTORY_TYPE(FLandscapeTileVertexFactory);
+
+public:
+	FLandscapeTileVertexFactory(ERHIFeatureLevel::Type InFeatureLevel)
+		: FLandscapeVertexFactory(InFeatureLevel)
+	{
+	}
+
+	struct FDataType : public FLandscapeVertexFactory::FDataType
+	{
+		FVertexStreamComponent TileDataComponent;
+	};
+
+	void SetData(const FDataType& InData)
+	{
+		FLandscapeVertexFactory::Data = InData;
+		Data = InData;
+		UpdateRHI();
+	}
+
+	/**
+	* Should we cache the material's shadertype on this platform with this vertex factory?
+	*/
+	static bool ShouldCompilePermutation(const FVertexFactoryShaderPermutationParameters& Parameters);
+
+	/**
+	* Can be overridden by FVertexFactory subclasses to modify their compile environment just before compilation occurs.
+	*/
+	static void ModifyCompilationEnvironment(const FVertexFactoryShaderPermutationParameters& Parameters, FShaderCompilerEnvironment& OutEnvironment);
+
+	/**
+	* Get vertex elements used when during PSO precaching materials using this vertex factory type
+	*/
+	static void GetPSOPrecacheVertexFetchElements(EVertexInputStreamType VertexInputStreamType, FVertexDeclarationElementList& Elements);
+
+	void Copy(const FLandscapeTileVertexFactory& Other);
+	// FRenderResource interface.
+	virtual void InitRHI() override;
+
+	FDataType Data;
+};
+
+//
+// FLandscapeTileVertexFactory
+//
+void FLandscapeTileVertexFactory::InitRHI()
+{
+	// list of declaration items
+	FVertexDeclarationElementList Elements;
+
+	// position decls
+	Elements.Add(AccessStreamComponent(Data.PositionComponent, 0));
+	// Tile per-instance data
+	Elements.Add(AccessStreamComponent(Data.TileDataComponent, 1));
+
+	// create the actual device decls
+	InitDeclaration(Elements);
+}
+
+/**
+* Copy the data from another vertex factory
+* @param Other - factory to copy from
+*/
+void FLandscapeTileVertexFactory::Copy(const FLandscapeTileVertexFactory& Other)
+{
+	FLandscapeTileVertexFactory* VertexFactory = this;
+	const FDataType* DataCopy = &Other.Data;
+	ENQUEUE_RENDER_COMMAND(FLandscapeTileVertexFactoryCopyData)(
+		[VertexFactory, DataCopy](FRHICommandListImmediate& RHICmdList)
+		{
+			VertexFactory->Data = *DataCopy;
+		});
+	BeginUpdateResourceRHI(this);
+}
+
+void FLandscapeTileVertexFactory::GetPSOPrecacheVertexFetchElements(EVertexInputStreamType VertexInputStreamType, FVertexDeclarationElementList& Elements)
+{
+	Elements.Add(FVertexElement(0, 0, VET_UByte4, 0, 0, false)); // Position
+	Elements.Add(FVertexElement(1, 0, VET_UByte4, 1, 0, true)); // Tile per-instance data
+}
+
+bool FLandscapeTileVertexFactory::ShouldCompilePermutation(const FVertexFactoryShaderPermutationParameters& Parameters)
+{
+	static FShaderPlatformCachedIniValue<int32> LandscapeSupportCullingIniValue(TEXT("landscape.SupportGPUCulling"));
+	return FLandscapeVertexFactory::ShouldCompilePermutation(Parameters) && LandscapeSupportCullingIniValue.Get(Parameters.Platform) != 0;
+}
+
+void FLandscapeTileVertexFactory::ModifyCompilationEnvironment(const FVertexFactoryShaderPermutationParameters& Parameters, FShaderCompilerEnvironment& OutEnvironment)
+{
+	FLandscapeVertexFactory::ModifyCompilationEnvironment(Parameters, OutEnvironment);
+	OutEnvironment.SetDefine(TEXT("LANDSCAPE_TILE"), TEXT("1"));
+}
+
+IMPLEMENT_VERTEX_FACTORY_PARAMETER_TYPE(FLandscapeTileVertexFactory, SF_Vertex, FLandscapeVertexFactoryVertexShaderParameters);
+IMPLEMENT_VERTEX_FACTORY_PARAMETER_TYPE(FLandscapeTileVertexFactory, SF_Pixel, FLandscapeVertexFactoryPixelShaderParameters);
+
+IMPLEMENT_VERTEX_FACTORY_TYPE(FLandscapeTileVertexFactory, "/Engine/Private/LandscapeVertexFactory.ush",
+	EVertexFactoryFlags::UsedWithMaterials
+	| EVertexFactoryFlags::SupportsStaticLighting
+	| EVertexFactoryFlags::SupportsDynamicLighting
+	| EVertexFactoryFlags::SupportsLightmapBaking
+	| EVertexFactoryFlags::SupportsPSOPrecaching
+);
+
+
+
+namespace UE::Landscape::Culling
+{
+
+constexpr static uint32 LANDSCAPE_TILE_QUADS = 4u;
+constexpr static uint32 MIN_SECTION_SIZE_QUADS = 31u;
+constexpr static uint32 INDIRECT_ARGS_NUM_WORDS = 5u;
+constexpr static uint32 TILE_DATA_ENTRY_NUM_BYTES = 4u;
+
+bool UseCulling(EShaderPlatform Platform)
+{
+	static FShaderPlatformCachedIniValue<int32> LandscapeSupportCullingIniValue(TEXT("landscape.SupportGPUCulling"));
+	
+	return 
+		LandscapeSupportCullingIniValue.Get(Platform) != 0 &&
+		// These features require VF PrimitiveID support which is not possible with culling VF atm
+		// Note that UseVirtualShadowMaps includes runtime logic, so this function can't be used in cook time decisions
+		!UseVirtualShadowMaps(Platform, GMaxRHIFeatureLevel) &&
+		!FDataDrivenShaderPlatformInfo::GetSupportsLumenGI(Platform);
+}
+
+FVertexFactoryType* GetTileVertexFactoryType()
+{
+	return &FLandscapeTileVertexFactory::StaticType;
+}
+
+class FLandscapeTileMesh final : public FRenderResource
+{
+public:
+	FLandscapeTileMesh()
+	{
+		InitResource();
+	}
+
+	/** Destructor. */
+	virtual ~FLandscapeTileMesh()
+	{
+		ReleaseResource();
+	}
+
+	virtual void InitResource() override;
+	virtual void ReleaseResource() override;
+	virtual void InitRHI() override;
+
+public:
+	FVertexBuffer VertexBuffer;
+	FIndexBuffer IndexBuffer;
+};
+
+
+void FLandscapeTileMesh::InitResource()
+{
+	FRenderResource::InitResource();
+	VertexBuffer.InitResource();
+	IndexBuffer.InitResource();
+}
+
+void FLandscapeTileMesh::ReleaseResource()
+{
+	FRenderResource::ReleaseResource();
+	VertexBuffer.ReleaseResource();
+	IndexBuffer.ReleaseResource();
+}
+
+void FLandscapeTileMesh::InitRHI()
+{
+	const uint32 TileSizeVertx = LANDSCAPE_TILE_QUADS + 1u;
+
+	// create a static vertex buffer
+	{
+		const uint32 NumVertx = FMath::Square(TileSizeVertx);
+		const uint32 Stride = sizeof(FLandscapeVertex);
+
+		FRHIResourceCreateInfo CreateInfo(TEXT("FLandscapeTileMeshVertexBuffer"));
+		VertexBuffer.VertexBufferRHI = RHICreateBuffer(NumVertx * Stride, BUF_Static | BUF_VertexBuffer, Stride, ERHIAccess::VertexOrIndexBuffer, CreateInfo);
+		FLandscapeVertex* Vertex = (FLandscapeVertex*)RHILockBuffer(VertexBuffer.VertexBufferRHI, 0, NumVertx * Stride, RLM_WriteOnly);
+
+		for (uint32 y = 0; y < TileSizeVertx; y++)
+		{
+			for (uint32 x = 0; x < TileSizeVertx; x++)
+			{
+				Vertex->VertexX = x;
+				Vertex->VertexY = y;
+				Vertex->SubX = 0;
+				Vertex->SubY = 0;
+				Vertex++;
+			}
+		}
+
+		RHIUnlockBuffer(VertexBuffer.VertexBufferRHI);
+	}
+
+	// 
+	{
+		const int32 NumIndices = FMath::Square(LANDSCAPE_TILE_QUADS) * 6u;
+		const uint32 Stride = sizeof(uint16);
+		FRHIResourceCreateInfo CreateInfo(TEXT("FLandscapeTileMeshIndexBuffer"));
+		IndexBuffer.IndexBufferRHI = RHICreateBuffer(NumIndices * Stride, BUF_Static | BUF_IndexBuffer, Stride, ERHIAccess::VertexOrIndexBuffer, CreateInfo);
+		uint16* Indices = (uint16*)RHILockBuffer(IndexBuffer.IndexBufferRHI, 0, NumIndices * Stride, RLM_WriteOnly);
+
+		for (uint32 y = 0; y < LANDSCAPE_TILE_QUADS; y++)
+		{
+			for (uint32 x = 0; x < LANDSCAPE_TILE_QUADS; x++)
+			{
+				uint16 i00 = (x + 0) + (y + 0) * (TileSizeVertx);
+				uint16 i10 = (x + 1) + (y + 0) * (TileSizeVertx);
+				uint16 i11 = (x + 1) + (y + 1) * (TileSizeVertx);
+				uint16 i01 = (x + 0) + (y + 1) * (TileSizeVertx);
+
+				Indices[0] = i00;
+				Indices[1] = i11;
+				Indices[2] = i10;
+
+				Indices[3] = i00;
+				Indices[4] = i01;
+				Indices[5] = i11;
+
+				Indices += 6;
+			}
+		}
+		RHIUnlockBuffer(IndexBuffer.IndexBufferRHI);
+	}
+}
+
+
+class FLandscapeTileDataBuffer final : public FVertexBuffer
+{
+	uint32 NumSubsections;
+	uint32 SubsectionSizeQuads;
+public:
+
+	FLandscapeTileDataBuffer(uint32 InNumSubsections, uint32 InSubsectionSizeQuads)
+		: NumSubsections(InNumSubsections)
+		, SubsectionSizeQuads(InSubsectionSizeQuads)
+	{
+		InitResource();
+	}
+
+	/** Destructor. */
+	virtual ~FLandscapeTileDataBuffer()
+	{
+		ReleaseResource();
+	}
+
+	virtual void InitRHI() override;
+};
+
+void FLandscapeTileDataBuffer::InitRHI()
+{
+	const uint32 SubsectionTilesRow = FMath::DivideAndRoundUp(SubsectionSizeQuads, LANDSCAPE_TILE_QUADS);
+	const uint32 SubsectionTiles = SubsectionTilesRow * SubsectionTilesRow;
+	const uint32 ComponentTiles = SubsectionTiles * NumSubsections * NumSubsections;
+	const uint32 Stride = TILE_DATA_ENTRY_NUM_BYTES;
+
+	FRHIResourceCreateInfo CreateInfo(TEXT("FLandscapeTileVertexBuffer"));
+	VertexBufferRHI = RHICreateBuffer(ComponentTiles * Stride, BUF_Static | BUF_VertexBuffer, Stride, ERHIAccess::VertexOrIndexBuffer, CreateInfo);
+	uint8* TileData = (uint8*)RHILockBuffer(VertexBufferRHI, 0, ComponentTiles * Stride, RLM_WriteOnly);
+
+	for (uint32 SubY = 0; SubY < NumSubsections; SubY++)
+	{
+		for (uint32 SubX = 0; SubX < NumSubsections; SubX++)
+		{
+			for (uint32 y = 0; y < SubsectionTilesRow; y++)
+			{
+				for (uint32 x = 0; x < SubsectionTilesRow; x++)
+				{
+					TileData[0] = x * LANDSCAPE_TILE_QUADS;
+					TileData[1] = y * LANDSCAPE_TILE_QUADS;
+					TileData[2] = SubX;
+					TileData[3] = SubY;
+					TileData += Stride;
+				}
+			}
+		}
+	}
+
+	RHIUnlockBuffer(VertexBufferRHI);
+}
+
+class FBuildLandscapeTileDataCS : public FGlobalShader
+{
+	DECLARE_GLOBAL_SHADER(FBuildLandscapeTileDataCS);
+	SHADER_USE_PARAMETER_STRUCT(FBuildLandscapeTileDataCS, FGlobalShader)
+
+public:
+	static bool ShouldCompilePermutation(const FGlobalShaderPermutationParameters& Parameters)
+	{
+		static FShaderPlatformCachedIniValue<int32> LandscapeSupportCullingIniValue(TEXT("landscape.SupportGPUCulling"));
+		return LandscapeSupportCullingIniValue.Get(Parameters.Platform) != 0;
+	}
+
+	static void ModifyCompilationEnvironment(const FGlobalShaderPermutationParameters& Parameters, FShaderCompilerEnvironment& OutEnvironment)
+	{
+		FGlobalShader::ModifyCompilationEnvironment(Parameters, OutEnvironment);
+		OutEnvironment.SetDefine(TEXT("INDIRECT_ARGS_NUM_WORDS"), INDIRECT_ARGS_NUM_WORDS);
+		OutEnvironment.SetDefine(TEXT("LANDSCAPE_TILE_QUADS"), LANDSCAPE_TILE_QUADS);
+	}
+
+	struct FLandscapeView
+	{
+		FMatrix44f ViewToClip;
+		FMatrix44f TranslatedWorldToClip;
+		FVector3f RelativePreViewTranslation;
+		float _Pad0;
+		FVector3f ViewTilePosition;
+		float _Pad1;
+	};
+
+	struct FLandscapeSection
+	{
+		FMatrix44f LocalToRelativeWorld;
+		FVector3f TilePosition;
+		float LocalZ;
+		float HalfHeight;
+		float NeighborLODExtent;
+		float _Pad[2];
+	};
+
+	BEGIN_SHADER_PARAMETER_STRUCT(FParameters, )
+		SHADER_PARAMETER_RDG_BUFFER_UAV(RWBuffer<uint>, IndirectArgsBufferOut)
+		SHADER_PARAMETER_RDG_BUFFER_UAV(RWByteAddressBuffer, TileDataOut)
+		SHADER_PARAMETER_RDG_BUFFER_SRV(StructuredBuffer<FLandscapeView>, LandscapeViews)
+		SHADER_PARAMETER_RDG_BUFFER_SRV(StructuredBuffer<FLandscapeSection>, LandscapeSections)
+		SHADER_PARAMETER(uint32, NumLandscapeSections)
+		SHADER_PARAMETER(uint32, NumLandscapeViews)
+		SHADER_PARAMETER(uint32, SubsectionSizeTiles)
+		SHADER_PARAMETER(uint32, NumSubsections)
+		SHADER_PARAMETER(uint32, NearClip)
+		END_SHADER_PARAMETER_STRUCT()
+};
+IMPLEMENT_GLOBAL_SHADER(FBuildLandscapeTileDataCS, "/Engine/Private/Landscape/LandscapeCulling.usf", "BuildLandscapeTileDataCS", SF_Compute);
+
+struct FViewStateIntermediateData
+{
+	FRDGBufferRef SectionsBufferRDG;
+	TArray<FIntPoint, TInlineAllocator<8>> SectionBases;
+};
+
+struct FDispatchIntermediates
+{
+	TRefCountPtr<FRDGPooledBuffer> IndirectArgsBuffer;
+	TRefCountPtr<FRDGPooledBuffer> TileDataBuffer;
+};
+
+struct FArgumentsKey
+{
+	const void* ViewPtr;
+	FIntPoint SectionBase;
+	uint32 ViewStateKey;
+
+	bool operator==(const FArgumentsKey& Other) const
+	{
+		return ViewPtr == Other.ViewPtr && SectionBase == Other.SectionBase && ViewStateKey == Other.ViewStateKey;
+	}
+};
+
+uint32 GetTypeHash(const FArgumentsKey& Key)
+{
+	uint32 Hash = PointerHash(Key.ViewPtr);
+	Hash = HashCombine(Hash, GetTypeHash(Key.SectionBase));
+	Hash = HashCombine(Hash, ::GetTypeHash(Key.ViewStateKey));
+	return Hash;
+}
+
+struct FCullingEntry
+{
+	uint32 LandscapeKey;
+	int32 ReferenceCount;
+	uint32 NumSubsections;
+	uint32 SubsectionSizeQuads;
+
+	TArray<FViewStateIntermediateData, TInlineAllocator<2>> IntermediateData;
+	TArray<FDispatchIntermediates, TInlineAllocator<2>> DispatchIntermediateData;
+	TMap<FArgumentsKey, FArguments> CullingArguments;
+};
+
+struct FCullingSystem
+{
+	FLandscapeTileMesh* TileMesh = nullptr;
+	
+	TArray<FCullingEntry> Landscapes; 
+	TArray<uint32, TInlineAllocator<2>> ViewStateKeys;
+};
+
+FCullingSystem GCullingSystem;
+
+static bool EnableGPUCullingForSection(int32 SubsectionSizeVerts, int32 NumSubsections)
+{
+	uint32 SectionSizeQuads = (SubsectionSizeVerts - 1) * NumSubsections;
+	return (SectionSizeQuads >= MIN_SECTION_SIZE_QUADS);
+}
+
+static void ResetIntermediateData()
+{
+	for (int32 LandscapeIdx = 0; LandscapeIdx < GCullingSystem.Landscapes.Num(); ++LandscapeIdx)
+	{
+		FCullingEntry& CullingEntry = GCullingSystem.Landscapes[LandscapeIdx];
+
+		CullingEntry.IntermediateData.Reset();
+		CullingEntry.DispatchIntermediateData.Reset();
+		CullingEntry.CullingArguments.Reset();
+	}
+
+	GCullingSystem.ViewStateKeys.Reset();
+}
+
+void PreRenderViewFamily(FSceneViewFamily& InViewFamily)
+{
+	ResetIntermediateData();
+}
+
+void InitSharedBuffers(FLandscapeSharedBuffers& SharedBuffers, const ERHIFeatureLevel::Type InFeatureLevel)
+{
+	if (!EnableGPUCullingForSection(SharedBuffers.SubsectionSizeVerts, SharedBuffers.NumSubsections))
+	{
+		SharedBuffers.TileVertexFactory = nullptr;
+		SharedBuffers.TileDataBuffer = nullptr;
+		return;
+	}
+
+	if (GCullingSystem.TileMesh == nullptr)
+	{
+		GCullingSystem.TileMesh = new FLandscapeTileMesh();
+	}
+
+	uint32 SubsectionSizeQuads = (SharedBuffers.SubsectionSizeVerts - 1);
+
+	FLandscapeTileMesh* TileMesh = GCullingSystem.TileMesh;
+	FLandscapeTileDataBuffer* TileDataBuffer = new FLandscapeTileDataBuffer(SharedBuffers.NumSubsections, SubsectionSizeQuads);
+	FLandscapeTileVertexFactory* TileVF = new FLandscapeTileVertexFactory(InFeatureLevel);
+
+	uint32 TileDataStride = TILE_DATA_ENTRY_NUM_BYTES;
+
+	TileVF->Data.PositionComponent = FVertexStreamComponent(&TileMesh->VertexBuffer, 0, sizeof(FLandscapeVertex), VET_UByte4);
+	TileVF->Data.TileDataComponent = FVertexStreamComponent(TileDataBuffer, 0, TileDataStride, VET_UByte4, EVertexStreamUsage::Instancing);
+	TileVF->InitResource();
+
+	SharedBuffers.TileVertexFactory = TileVF;
+	SharedBuffers.TileDataBuffer = TileDataBuffer;
+}
+
+void DeleteSharedBuffers(FLandscapeSharedBuffers& SharedBuffers)
+{
+	delete SharedBuffers.TileDataBuffer;
+	SharedBuffers.TileDataBuffer = nullptr;
+	
+	delete SharedBuffers.TileVertexFactory;
+	SharedBuffers.TileVertexFactory = nullptr;
+	
+	if (GCullingSystem.Landscapes.Num() == 0)
+	{
+		delete GCullingSystem.TileMesh;
+		GCullingSystem.TileMesh = nullptr;
+	}
+}
+
+void SetupMeshBatch(const FLandscapeSharedBuffers& SharedBuffers, FMeshBatch& MeshBatch)
+{
+	if (!EnableGPUCullingForSection(SharedBuffers.SubsectionSizeVerts, SharedBuffers.NumSubsections))
+	{
+		return;
+	}
+		
+	if (MeshBatch.LODIndex == 0)
+	{
+		uint32 SubsectionSizeQuads = (SharedBuffers.SubsectionSizeVerts - 1);
+
+		check(SharedBuffers.TileVertexFactory);
+
+		MeshBatch.bViewDependentArguments = true;
+		MeshBatch.VertexFactory = SharedBuffers.TileVertexFactory;
+
+		FMeshBatchElement& BatchElement = MeshBatch.Elements[0];
+		BatchElement.IndexBuffer = &GCullingSystem.TileMesh->IndexBuffer;
+		BatchElement.NumPrimitives = FMath::Square(LANDSCAPE_TILE_QUADS) * 2u;
+		BatchElement.NumInstances = FMath::Square(FMath::DivideAndRoundUp(SubsectionSizeQuads, LANDSCAPE_TILE_QUADS)) * FMath::Square(SharedBuffers.NumSubsections);
+		BatchElement.FirstIndex = 0;
+		BatchElement.MinVertexIndex = 0;
+		BatchElement.MaxVertexIndex = FMath::Square(LANDSCAPE_TILE_QUADS + 1u) - 1u;
+	}
+}
+
+void RegisterLandscape(uint32 LandscapeKey, int32 SubsectionSizeVerts, int32 NumSubsections)
+{
+	if (!EnableGPUCullingForSection(SubsectionSizeVerts, NumSubsections))
+	{
+		return;
+	}
+
+	FCullingEntry* CullingEntryPtr = GCullingSystem.Landscapes.FindByPredicate([LandscapeKey](const FCullingEntry& Other) {
+		return Other.LandscapeKey == LandscapeKey;
+	});
+
+	if (!CullingEntryPtr)
+	{
+		FCullingEntry& CullingEntry = GCullingSystem.Landscapes.AddDefaulted_GetRef();
+
+		CullingEntry.LandscapeKey = LandscapeKey;
+		CullingEntry.NumSubsections = NumSubsections;
+		CullingEntry.SubsectionSizeQuads = (SubsectionSizeVerts - 1);
+		CullingEntry.ReferenceCount = 1;
+	}
+	else
+	{
+		CullingEntryPtr->ReferenceCount++;
+	}
+}
+
+void UnregisterLandscape(uint32 LandscapeKey)
+{
+	for (int32 i = GCullingSystem.Landscapes.Num() - 1; i >= 0; --i)
+	{
+		FCullingEntry& CullingEntry = GCullingSystem.Landscapes[i];
+		if (CullingEntry.LandscapeKey == LandscapeKey)
+		{
+			CullingEntry.ReferenceCount--;
+			if (CullingEntry.ReferenceCount <= 0)
+			{
+				GCullingSystem.Landscapes.RemoveAtSwap(i);
+			}
+			break;
+		}
+	}
+}
+
+static float ComputeNeighborsMaxLOD(const FLandscapeRenderSystem& RenderSystem, const TResourceArray<float>& SectionLODValues, FIntPoint SectionBase)
+{
+	static const FIntPoint NeightborBaseOffsets[4]
+	{
+		{1, 0}, {-1, 0}, {0, 1}, {0, -1}
+	};
+
+	float NeighborsMaxLOD = 0.f;
+	for (int32 i = 0; i < UE_ARRAY_COUNT(NeightborBaseOffsets); ++i)
+	{
+		int32 SectionIdx = RenderSystem.GetSectionLinearIndex(SectionBase + NeightborBaseOffsets[i]);
+		if (SectionLODValues.IsValidIndex(SectionIdx))
+		{
+			NeighborsMaxLOD = FMath::Max(SectionLODValues[SectionIdx], NeighborsMaxLOD);
+		}
+	}
+
+	return NeighborsMaxLOD;
+}
+
+static void ComputeSectionIntermediateData(FRDGBuilder& GraphBuilder, TArrayView<const FSceneView*> Views)
+{
+	TArray<FBuildLandscapeTileDataCS::FLandscapeSection, TInlineAllocator<16>> SectionsData;
+
+	for (int32 ViewIdx = 0; ViewIdx < Views.Num(); ++ViewIdx)
+	{
+		const FSceneView& View = *Views[ViewIdx];
+		uint32 ViewKey = View.GetViewKey();
+
+		GCullingSystem.ViewStateKeys.Add(ViewKey);
+		
+		// Collect all LOD0 sections for each ViewState+Landscape and upload to GPU
+		for (int32 LandscapeIdx = 0; LandscapeIdx < GCullingSystem.Landscapes.Num(); ++LandscapeIdx)
+		{
+			FCullingEntry& CullingEntry = GCullingSystem.Landscapes[LandscapeIdx];
+
+			FViewStateIntermediateData& ViewStateIntermediates = CullingEntry.IntermediateData.AddDefaulted_GetRef();
+			ViewStateIntermediates.SectionsBufferRDG = nullptr;
+
+			const FLandscapeRenderSystem& RenderSystem = *LandscapeRenderSystems.FindChecked(CullingEntry.LandscapeKey);
+			const TResourceArray<float>& SectionLODValues = RenderSystem.CachedSectionLODValues.FindRef(ViewKey);
+
+			for (int32 SectionIdx = 0; SectionIdx < SectionLODValues.Num(); ++SectionIdx)
+			{
+				int32 LODValue = SectionLODValues[SectionIdx];
+				FLandscapeSectionInfo* SectionInfo = RenderSystem.SectionInfos[SectionIdx];
+
+				if (LODValue == 0 && SectionInfo != nullptr)
+				{
+					FBuildLandscapeTileDataCS::FLandscapeSection& Section = SectionsData.AddDefaulted_GetRef();
+
+					FBoxSphereBounds SectionLocalBounds;
+					FMatrix SectionLocalToWorld;
+					SectionInfo->GetSectionBoundsAndLocalToWorld(SectionLocalBounds, SectionLocalToWorld);
+					const FLargeWorldRenderPosition SectionAbsoluteOrigin(SectionLocalToWorld.GetOrigin());
+					const int32 NeighborsMaxLOD = FMath::RoundFromZero(ComputeNeighborsMaxLOD(RenderSystem, SectionLODValues, SectionInfo->ComponentBase));
+
+					Section.LocalToRelativeWorld = FLargeWorldRenderScalar::MakeToRelativeWorldMatrix(SectionAbsoluteOrigin.GetTileOffset(), SectionLocalToWorld);
+					Section.TilePosition = SectionAbsoluteOrigin.GetTile();
+					Section.LocalZ = SectionLocalBounds.Origin.Z;
+					Section.HalfHeight = SectionLocalBounds.BoxExtent.Z;
+					// How many quads to add to each tile extent to compensate for a neighbors LOD
+					Section.NeighborLODExtent = (1 << NeighborsMaxLOD) - 1; 
+
+					ViewStateIntermediates.SectionBases.Add(SectionInfo->ComponentBase);
+				}
+			}
+
+			if (SectionsData.Num() != 0)
+			{
+				ViewStateIntermediates.SectionsBufferRDG = CreateStructuredBuffer<FBuildLandscapeTileDataCS::FLandscapeSection>(GraphBuilder, TEXT("LandscapeCulling.SectionsData"), SectionsData);
+				SectionsData.Reset();
+			}
+		}
+	}
+}
+
+static void DispatchCulling(FRDGBuilder& GraphBuilder, TArrayView<const FSceneView*> CullingViews, TArrayView<FViewMatrices> CullingViewMatrices, bool bNearClip)
+{
+	// A separate dispatch for each Landscape and ViewState
+	for (int32 ViewStateIdx = 0; ViewStateIdx < GCullingSystem.ViewStateKeys.Num(); ++ViewStateIdx)
+	{
+		uint32 ViewStateKey = GCullingSystem.ViewStateKeys[ViewStateIdx];
+
+		TArray<FBuildLandscapeTileDataCS::FLandscapeView, TInlineAllocator<8>> CullingViewsData;
+		CullingViewsData.Reserve(CullingViews.Num());
+		TArray<const FSceneView*, TInlineAllocator<8>> FilteredViews;
+		FilteredViews.Reserve(CullingViews.Num());
+
+		// Filter out views for a current ViewState
+		for (int32 CullingViewIdx = 0; CullingViewIdx < CullingViews.Num(); ++CullingViewIdx)
+		{
+			const FSceneView& CullingView = *CullingViews[CullingViewIdx];
+			if (CullingView.GetViewKey() != ViewStateKey)
+			{
+				continue;
+			}
+			
+			FilteredViews.Add(&CullingView);
+			FBuildLandscapeTileDataCS::FLandscapeView& LandscapeView = CullingViewsData.AddDefaulted_GetRef();
+
+			const FViewMatrices& ViewMatrices = CullingViewMatrices[CullingViewIdx];
+			const FLargeWorldRenderPosition AbsoluteViewOrigin(ViewMatrices.GetViewOrigin());
+			const FVector ViewTileOffset = AbsoluteViewOrigin.GetTileOffset();
+			const FVector3f ViewTilePosition = AbsoluteViewOrigin.GetTile();
+						
+			LandscapeView.ViewToClip = FMatrix44f(ViewMatrices.GetProjectionMatrix());
+			LandscapeView.TranslatedWorldToClip = FMatrix44f(ViewMatrices.GetTranslatedViewProjectionMatrix());
+			LandscapeView.RelativePreViewTranslation = FVector3f(ViewMatrices.GetPreViewTranslation() + ViewTileOffset);
+			LandscapeView.ViewTilePosition = ViewTilePosition;
+		}
+
+		// Share view data between all landscapes
+		FRDGBufferRef CullingViewsBufferRDG = nullptr;
+
+		for (int32 LandscapeIdx = 0; LandscapeIdx < GCullingSystem.Landscapes.Num(); ++LandscapeIdx)
+		{
+			FCullingEntry& CullingEntry = GCullingSystem.Landscapes[LandscapeIdx];
+			FViewStateIntermediateData& ViewStateIntermediates = CullingEntry.IntermediateData[ViewStateIdx];
+
+			uint32 NumCullingViews = CullingViewsData.Num();
+			uint32 NumCullingSections = ViewStateIntermediates.SectionBases.Num();
+			uint32 NumCullingItems = NumCullingSections * NumCullingViews;
+			if (NumCullingItems == 0)
+			{
+				continue;
+			}
+
+			TArray<FRHIDrawIndexedIndirectParameters> IndirectArgsParameters;
+			IndirectArgsParameters.Init(FRHIDrawIndexedIndirectParameters{ FMath::Square(LANDSCAPE_TILE_QUADS) * 6u, 0u, 0u, 0u, 0u }, NumCullingItems);
+
+			const uint32 IndirectArgsBufferElements = INDIRECT_ARGS_NUM_WORDS * IndirectArgsParameters.Num();
+			FRDGBufferRef IndirectArgsRDG = GraphBuilder.CreateBuffer(FRDGBufferDesc::CreateIndirectDesc(IndirectArgsBufferElements), TEXT("LandscapeCulling.DrawIndirectArgsBuffer"));
+
+			const uint32 NumTilesSectionRow = FMath::DivideAndRoundUp(CullingEntry.SubsectionSizeQuads, LANDSCAPE_TILE_QUADS) * CullingEntry.NumSubsections;
+			const uint32 NumTilesSection = FMath::Square(NumTilesSectionRow);
+
+			const uint32 TileDataBufferElements = NumTilesSection * NumCullingItems;
+			FRDGBufferDesc TileDataBufferDesc = FRDGBufferDesc::CreateByteAddressDesc(TileDataBufferElements * TILE_DATA_ENTRY_NUM_BYTES);
+			TileDataBufferDesc.Usage |= BUF_VertexBuffer;
+			FRDGBufferRef TileDataRDG = GraphBuilder.CreateBuffer(TileDataBufferDesc, TEXT("LandscapeCulling.CulledTileDataBuffer"));
+
+			// Create buffer for indirect args and upload draw arg data, also clears the instance to zero
+			GraphBuilder.QueueBufferUpload(IndirectArgsRDG, IndirectArgsParameters.GetData(), IndirectArgsParameters.Num() * sizeof(FRHIDrawIndexedIndirectParameters));
+
+			if (CullingViewsBufferRDG == nullptr)
+			{
+				CullingViewsBufferRDG = CreateStructuredBuffer<FBuildLandscapeTileDataCS::FLandscapeView>(GraphBuilder, TEXT("LandscapeCulling.ViewsData"), CullingViewsData);
+			}
+
+			{
+				ERHIFeatureLevel::Type FeatureLevel = CullingViews[0]->GetFeatureLevel();
+				FGlobalShaderMap* ShaderMap = GetGlobalShaderMap(FeatureLevel);
+				FIntVector SizeInTiles{ (int32)NumTilesSectionRow, (int32)NumTilesSectionRow, (int32)NumCullingItems };
+				FIntVector GroupSize{ 8, 8, 1 };
+
+				auto PassParameters = GraphBuilder.AllocParameters<FBuildLandscapeTileDataCS::FParameters>();
+				PassParameters->IndirectArgsBufferOut = GraphBuilder.CreateUAV(IndirectArgsRDG);
+				PassParameters->TileDataOut = GraphBuilder.CreateUAV(TileDataRDG);
+				PassParameters->LandscapeViews = GraphBuilder.CreateSRV(CullingViewsBufferRDG);
+				PassParameters->LandscapeSections = GraphBuilder.CreateSRV(ViewStateIntermediates.SectionsBufferRDG);
+				PassParameters->SubsectionSizeTiles = FMath::DivideAndRoundUp(CullingEntry.SubsectionSizeQuads, LANDSCAPE_TILE_QUADS);
+				PassParameters->NumSubsections = CullingEntry.NumSubsections;
+				PassParameters->NumLandscapeViews = NumCullingViews;
+				PassParameters->NumLandscapeSections = NumCullingSections;
+				PassParameters->NearClip = bNearClip ? 1 : 0;
+
+				auto ComputeShader = ShaderMap->GetShader<FBuildLandscapeTileDataCS>();
+
+				FComputeShaderUtils::AddPass(
+					GraphBuilder,
+					RDG_EVENT_NAME("LandscapeCulling, SubsectionSizeQuads %d, Subsections %d, Views:%d", CullingEntry.SubsectionSizeQuads, CullingEntry.NumSubsections, NumCullingViews),
+					ComputeShader,
+					PassParameters,
+					FComputeShaderUtils::GetGroupCount(SizeInTiles, GroupSize)
+				);
+			}
+			FDispatchIntermediates& DispatchIntermediates = CullingEntry.DispatchIntermediateData.AddDefaulted_GetRef();
+			DispatchIntermediates.IndirectArgsBuffer = ConvertToExternalAccessBuffer(GraphBuilder, IndirectArgsRDG, ERHIAccess::IndirectArgs);
+			DispatchIntermediates.TileDataBuffer = ConvertToExternalAccessBuffer(GraphBuilder, TileDataRDG, ERHIAccess::VertexOrIndexBuffer);
+
+			FRHIBuffer* IndirectArgsBufferRHI = DispatchIntermediates.IndirectArgsBuffer->GetRHI();
+			FRHIBuffer* TileDataBufferRHI = DispatchIntermediates.TileDataBuffer->GetRHI();
+
+			// Create a map for each view+section combo to its draw args
+			for (int32 ViewIdx = 0; ViewIdx < FilteredViews.Num(); ++ViewIdx)
+			{
+				const FSceneView* ViewPtr = FilteredViews[ViewIdx];
+				for (int32 SectionIdx = 0; SectionIdx < ViewStateIntermediates.SectionBases.Num(); ++SectionIdx)
+				{
+					FIntPoint SectionBase = ViewStateIntermediates.SectionBases[SectionIdx];
+					FArguments& Args = CullingEntry.CullingArguments.Add(FArgumentsKey{ ViewPtr, SectionBase, ViewStateKey });
+					Args.IndirectArgsBuffer = IndirectArgsBufferRHI;
+					Args.TileDataVertexBuffer = TileDataBufferRHI;
+					// offsets are in bytes
+					Args.IndirectArgsOffset = (4u * INDIRECT_ARGS_NUM_WORDS) * ((ViewIdx * NumCullingSections) + SectionIdx);
+					Args.TileDataOffset = TILE_DATA_ENTRY_NUM_BYTES * NumTilesSection * ((ViewIdx * NumCullingSections) + SectionIdx);
+				}
+			}
+		}
+	}
+}
+
+void InitMainViews(FRDGBuilder& GraphBuilder, TArrayView<const FSceneView*> Views)
+{
+	if (GCullingSystem.Landscapes.Num() == 0 ||
+		CVarLandscapeEnableGPUCulling.GetValueOnRenderThread() == 0)
+	{
+		return;
+	}
+
+	QUICK_SCOPE_CYCLE_COUNTER(STAT_LandscapeCulling_InitMainViews);
+
+	check(GCullingSystem.ViewStateKeys.Num() == 0);
+
+	ComputeSectionIntermediateData(GraphBuilder, Views);
+
+	TArray<FViewMatrices, TInlineAllocator<4>> CullingViewMatrices;
+	for (const FSceneView* View : Views)
+	{
+		CullingViewMatrices.Add(View->ViewMatrices);
+	}
+	DispatchCulling(GraphBuilder, Views, CullingViewMatrices, true);
+}
+
+void InitShadowViews(FRDGBuilder& GraphBuilder, TArrayView<const FSceneView*> ShadowDepthViews, TArrayView<FViewMatrices> ShadowViewMatrices)
+{
+	if (GCullingSystem.Landscapes.Num() == 0 ||
+		CVarLandscapeEnableGPUCulling.GetValueOnRenderThread() == 0 || 
+		CVarLandscapeEnableGPUCullingShadows.GetValueOnRenderThread() == 0)
+	{
+		return;
+	}
+
+	QUICK_SCOPE_CYCLE_COUNTER(STAT_LandscapeCulling_InitShadowViews);
+
+	DispatchCulling(GraphBuilder, ShadowDepthViews, ShadowViewMatrices, false);
+}
+
+bool GetViewArguments(const FSceneView& View, uint32 LandscapeKey, FIntPoint SectionBase, int32 LODIndex, FArguments& Args)
+{
+	if (LODIndex != 0 || 
+		GCullingSystem.Landscapes.Num() == 0 ||
+		CVarLandscapeEnableGPUCulling.GetValueOnRenderThread() == 0)
+	{
+		return false;
+	}
+
+	const FCullingEntry* CullingEntryPtr = GCullingSystem.Landscapes.FindByPredicate([LandscapeKey](const FCullingEntry& Other) {
+		return Other.LandscapeKey == LandscapeKey;
+	});
+
+	if (CullingEntryPtr == nullptr)
+	{
+		return false;
+	}
+	FArgumentsKey ArgumentsKey{ &View, SectionBase, View.GetViewKey() };
+	const FArguments* ArgsPtr = CullingEntryPtr->CullingArguments.Find(ArgumentsKey);
+	if (ArgsPtr)
+	{
+		Args = *ArgsPtr;
+		return true;
+	}
+	return false;
+}
+
+}//namespace UE::Landscape::Culling
\ No newline at end of file
diff --git a/Engine/Source/Runtime/Landscape/Private/LandscapeCulling.h b/Engine/Source/Runtime/Landscape/Private/LandscapeCulling.h
new file mode 100644
--- /dev/null
+++ b/Engine/Source/Runtime/Landscape/Private/LandscapeCulling.h
@@ -0,0 +1,44 @@
+// Copyright Epic Games, Inc. All Rights Reserved.
+
+#pragma once
+
+#include "CoreMinimal.h"
+
+class FLandscapeSharedBuffers;
+
+/**
+* Landscape GPU culling implementation
+* Each landscape section is split to a smaller tiles (4x4 quads)
+* Tiles are frustum culled in compute
+* Culling is done for each view, including shadow views
+* Only active for LOD0 atm
+*/
+namespace UE {
+	namespace Landscape {
+		namespace Culling {
+
+			struct FArguments
+			{
+				FRHIBuffer* IndirectArgsBuffer;
+				FRHIBuffer* TileDataVertexBuffer;
+				uint32 IndirectArgsOffset;
+				uint32 TileDataOffset;
+			};
+
+			bool UseCulling(EShaderPlatform ShaderPlatform);
+			
+			FVertexFactoryType* GetTileVertexFactoryType();
+
+			void InitSharedBuffers(FLandscapeSharedBuffers& SharedBuffers, ERHIFeatureLevel::Type FeatureLevel);
+			void DeleteSharedBuffers(FLandscapeSharedBuffers& SharedBuffers);
+			void SetupMeshBatch(const FLandscapeSharedBuffers& SharedBuffers, FMeshBatch& MeshBatch);
+			void RegisterLandscape(uint32 LandscapeKey, int32 SubsectionSizeVerts, int32 NumSubsections);
+			void UnregisterLandscape(uint32 LandscapeKey);
+
+			void PreRenderViewFamily(FSceneViewFamily& InViewFamily);
+			void InitMainViews(FRDGBuilder& GraphBuilder, TArrayView<const FSceneView*> Views);
+			void InitShadowViews(FRDGBuilder& GraphBuilder, TArrayView<const FSceneView*> ShadowDepthViews, TArrayView<FViewMatrices> ShadowViewMatrices);
+			bool GetViewArguments(const FSceneView& View, uint32 LandscapeKey, FIntPoint SectionBase, int32 LODIndex, FArguments& Args);
+		}
+	}
+};
diff --git a/Engine/Source/Runtime/Landscape/Private/LandscapeRender.cpp b/Engine/Source/Runtime/Landscape/Private/LandscapeRender.cpp
--- a/Engine/Source/Runtime/Landscape/Private/LandscapeRender.cpp
+++ b/Engine/Source/Runtime/Landscape/Private/LandscapeRender.cpp
@@ -52,9 +52,13 @@ LandscapeRender.cpp: New terrain rendering
 #include "RenderCore.h"
 #include "DataDrivenShaderPlatformInfo.h"
 #include "Algo/Transform.h"
+#include "LandscapeCulling.h"
+
+using namespace UE::Landscape;
 
 IMPLEMENT_GLOBAL_SHADER_PARAMETER_STRUCT(FLandscapeUniformShaderParameters, "LandscapeParameters");
 IMPLEMENT_GLOBAL_SHADER_PARAMETER_STRUCT(FLandscapeFixedGridUniformShaderParameters, "LandscapeFixedGrid");
+IMPLEMENT_TYPE_LAYOUT(FLandscapeVertexFactoryVertexShaderParameters);
 IMPLEMENT_TYPE_LAYOUT(FLandscapeVertexFactoryPixelShaderParameters);
 
 #if !UE_BUILD_SHIPPING
@@ -734,6 +738,11 @@ void FLandscapeSceneViewExtension::BeginRenderViewFamily(FSceneViewFamily& InVie
 	}
 }
 
+void FLandscapeSceneViewExtension::PreRenderViewFamily_RenderThread(FRDGBuilder& GraphBuilder, FSceneViewFamily& InViewFamily)
+{
+	Culling::PreRenderViewFamily(InViewFamily);
+}
+
 void FLandscapeSceneViewExtension::PreRenderView_RenderThread(FRDGBuilder& GraphBuilder, FSceneView& InView)
 {
 	LandscapeViews.Emplace(InView);
@@ -803,6 +812,8 @@ void FLandscapeSceneViewExtension::PreInitViews_RenderThread(FRDGBuilder& GraphB
 		RenderSystem.UpdateBuffers();
 	}
 
+	TArray<const FSceneView*, TInlineAllocator<2>> SceneViews;
+
 	for (FLandscapeViewData& LandscapeView : LandscapeViews)
 	{
 		if (!LandscapeRenderSystems.IsEmpty())
@@ -820,8 +831,15 @@ void FLandscapeSceneViewExtension::PreInitViews_RenderThread(FRDGBuilder& GraphB
 			LandscapeView.View->LandscapePerComponentDataBuffer = GWhiteVertexBufferWithSRV->ShaderResourceViewRHI;
 			LandscapeView.View->LandscapeIndirectionBuffer = GWhiteVertexBufferWithSRV->ShaderResourceViewRHI;
 		}
+
+		SceneViews.Add(LandscapeView.View);
 	}
 
+	if (Culling::UseCulling(SceneViews[0]->GetShaderPlatform()))
+	{
+		Culling::InitMainViews(GraphBuilder, SceneViews);
+	}
+		
 	LandscapeViews.Reset();
 }
 
@@ -1155,6 +1173,9 @@ FLandscapeComponentSceneProxy::FLandscapeComponentSceneProxy(ULandscapeComponent
 		}
 	}
 
+	// Landscape GPU culling uses VF that requires primitive UB
+	bVFRequiresPrimitiveUniformBuffer = (!bNaniteActive && Culling::UseCulling(GetScene().GetShaderPlatform()));
+
 	bSupportsInstanceDataBuffer = true;
 	UpdateDefaultInstanceSceneData();
 
@@ -1198,14 +1219,14 @@ void FLandscapeComponentSceneProxy::CreateRenderThreadResources()
 		if (!XYOffsetmapTexture)
 		{
 			FLandscapeVertexFactory* LandscapeVertexFactory = new FLandscapeVertexFactory(FeatureLevel);
-			LandscapeVertexFactory->Data.PositionComponent = FVertexStreamComponent(SharedBuffers->VertexBuffer, 0, sizeof(FLandscapeVertex), VET_Float4);
+			LandscapeVertexFactory->Data.PositionComponent = FVertexStreamComponent(SharedBuffers->VertexBuffer, 0, sizeof(FLandscapeVertex), VET_UByte4);
 			LandscapeVertexFactory->InitResource();
 			SharedBuffers->VertexFactory = LandscapeVertexFactory;
 		}
 		else
 		{
 			FLandscapeXYOffsetVertexFactory* LandscapeXYOffsetVertexFactory = new FLandscapeXYOffsetVertexFactory(FeatureLevel);
-			LandscapeXYOffsetVertexFactory->Data.PositionComponent = FVertexStreamComponent(SharedBuffers->VertexBuffer, 0, sizeof(FLandscapeVertex), VET_Float4);
+			LandscapeXYOffsetVertexFactory->Data.PositionComponent = FVertexStreamComponent(SharedBuffers->VertexBuffer, 0, sizeof(FLandscapeVertex), VET_UByte4);
 			LandscapeXYOffsetVertexFactory->InitResource();
 			SharedBuffers->VertexFactory = LandscapeXYOffsetVertexFactory;
 		}
@@ -1221,10 +1242,15 @@ void FLandscapeComponentSceneProxy::CreateRenderThreadResources()
 		{
 			//todo[vt]: We will need a version of this to support XYOffsetmapTexture
 			FLandscapeFixedGridVertexFactory* LandscapeVertexFactory = new FLandscapeFixedGridVertexFactory(FeatureLevel);
-			LandscapeVertexFactory->Data.PositionComponent = FVertexStreamComponent(SharedBuffers->VertexBuffer, 0, sizeof(FLandscapeVertex), VET_Float4);
+			LandscapeVertexFactory->Data.PositionComponent = FVertexStreamComponent(SharedBuffers->VertexBuffer, 0, sizeof(FLandscapeVertex), VET_UByte4);
 			LandscapeVertexFactory->InitResource();
 			SharedBuffers->FixedGridVertexFactory = LandscapeVertexFactory;
 		}
+
+		if (Culling::UseCulling(GetScene().GetShaderPlatform()))
+		{
+			Culling::InitSharedBuffers(*SharedBuffers, FeatureLevel);
+		}
 	}
 
 	SharedBuffers->AddRef();
@@ -1233,6 +1259,11 @@ void FLandscapeComponentSceneProxy::CreateRenderThreadResources()
 	VertexFactory = SharedBuffers->VertexFactory;
 	FixedGridVertexFactory = SharedBuffers->FixedGridVertexFactory;
 
+	if (!bNaniteActive && Culling::UseCulling(GetScene().GetShaderPlatform()))
+	{
+		Culling::RegisterLandscape(LandscapeKey, SubsectionSizeVerts, NumSubsections);
+	}
+	
 	// Assign LandscapeUniformShaderParameters
 	LandscapeUniformShaderParameters.InitResource();
 
@@ -1446,6 +1477,11 @@ FLandscapeComponentSceneProxy::~FLandscapeComponentSceneProxy()
 	// Free the subsection uniform buffer
 	LandscapeUniformShaderParameters.ReleaseResource();
 
+	if (Culling::UseCulling(GetScene().GetShaderPlatform()))
+	{
+		Culling::UnregisterLandscape(LandscapeKey);
+	}
+	
 	// Free the lod uniform buffers
 	for (int32 i = 0; i < LandscapeFixedGridUniformShaderParameters.Num(); ++i)
 	{
@@ -1932,6 +1968,11 @@ bool FLandscapeComponentSceneProxy::GetStaticMeshElement(int32 LODIndex, bool bF
 		BatchElement.FirstIndex = 0;
 		BatchElement.MinVertexIndex = SharedBuffers->IndexRanges[LODIndex].MinIndexFull;
 		BatchElement.MaxVertexIndex = SharedBuffers->IndexRanges[LODIndex].MaxIndexFull;
+				
+		if (!bForToolMesh && !bNaniteActive && Culling::UseCulling(GetScene().GetShaderPlatform()))
+		{
+			Culling::SetupMeshBatch(*SharedBuffers, MeshBatch);
+		}
 
 		// The default is overridden here only by mobile landscape to punch holes in the geometry
 		ApplyMeshElementModifier(BatchElement, LODIndex);
@@ -2423,6 +2464,18 @@ void FLandscapeComponentSceneProxy::GetDynamicMeshElements(const TArray<const FS
 #endif // !UE_BUILD_SHIPPING
 }
 
+void FLandscapeComponentSceneProxy::ApplyViewDependentMeshArguments(const FSceneView& View, FMeshBatch& ViewDependentMeshBatch) const
+{
+	Culling::FArguments CullingArgs{};
+	if (Culling::GetViewArguments(View, LandscapeKey, ComponentBase, ViewDependentMeshBatch.LODIndex, CullingArgs))
+	{
+		ViewDependentMeshBatch.Elements[0].NumPrimitives = 0;
+		ViewDependentMeshBatch.Elements[0].NumInstances = 0;
+		ViewDependentMeshBatch.Elements[0].IndirectArgsBuffer = CullingArgs.IndirectArgsBuffer;
+		ViewDependentMeshBatch.Elements[0].IndirectArgsOffset = CullingArgs.IndirectArgsOffset;
+	}
+}
+
 #if RHI_RAYTRACING
 void FLandscapeComponentSceneProxy::GetDynamicRayTracingInstances(FRayTracingMaterialGatheringContext& Context, TArray<FRayTracingInstance>& OutRayTracingInstances)
 {
@@ -2814,6 +2867,8 @@ FLandscapeSharedBuffers::FLandscapeSharedBuffers(const int32 InSharedBuffersKey,
 	, VertexFactory(nullptr)
 	, FixedGridVertexFactory(nullptr)
 	, VertexBuffer(nullptr)
+	, TileVertexFactory(nullptr)
+	, TileDataBuffer(nullptr)
 	, bUse32BitIndices(false)
 
 #if WITH_EDITOR
@@ -2861,6 +2916,8 @@ FLandscapeSharedBuffers::FLandscapeSharedBuffers(const int32 InSharedBuffersKey,
 
 FLandscapeSharedBuffers::~FLandscapeSharedBuffers()
 {
+	Culling::DeleteSharedBuffers(*this);
+	
 	delete VertexBuffer;
 
 	for (int32 i = 0; i < NumIndexBuffers; i++)
@@ -2897,52 +2954,48 @@ FLandscapeSharedBuffers::~FLandscapeSharedBuffers()
 //
 // FLandscapeVertexFactoryVertexShaderParameters
 //
-
-/** Shader parameters for use with FLandscapeVertexFactory */
-class FLandscapeVertexFactoryVertexShaderParameters : public FVertexFactoryShaderParameters
+void FLandscapeVertexFactoryVertexShaderParameters::GetElementShaderBindings(
+	const class FSceneInterface* Scene,
+	const FSceneView* InView,
+	const class FMeshMaterialShader* Shader,
+	const EVertexInputStreamType InputStreamType,
+	ERHIFeatureLevel::Type FeatureLevel,
+	const FVertexFactory* VertexFactory,
+	const FMeshBatchElement& BatchElement,
+	class FMeshDrawSingleShaderBindings& ShaderBindings,
+	FVertexInputStreamArray& VertexStreams
+) const
 {
-	DECLARE_TYPE_LAYOUT(FLandscapeVertexFactoryVertexShaderParameters, NonVirtual);
-public:
-	/**
-	* Bind shader constants by name
-	* @param	ParameterMap - mapping of named shader constants to indices
-	*/
-	void Bind(const FShaderParameterMap& ParameterMap)
-	{
-	}
-
-	void GetElementShaderBindings(
-		const class FSceneInterface* Scene,
-		const FSceneView* InView,
-		const class FMeshMaterialShader* Shader,
-		const EVertexInputStreamType InputStreamType,
-		ERHIFeatureLevel::Type FeatureLevel,
-		const FVertexFactory* VertexFactory,
-		const FMeshBatchElement& BatchElement,
-		class FMeshDrawSingleShaderBindings& ShaderBindings,
-		FVertexInputStreamArray& VertexStreams
-	) const
-	{
-		SCOPE_CYCLE_COUNTER(STAT_LandscapeVFDrawTimeVS);
+	SCOPE_CYCLE_COUNTER(STAT_LandscapeVFDrawTimeVS);
 
-		const FLandscapeBatchElementParams* BatchElementParams = (const FLandscapeBatchElementParams*)BatchElement.UserData;
-		check(BatchElementParams);
+	const FLandscapeBatchElementParams* BatchElementParams = (const FLandscapeBatchElementParams*)BatchElement.UserData;
+	check(BatchElementParams);
 
-		const FLandscapeComponentSceneProxy* SceneProxy = BatchElementParams->SceneProxy;
+	const FLandscapeComponentSceneProxy* SceneProxy = BatchElementParams->SceneProxy;
 
-		ShaderBindings.Add(Shader->GetUniformBufferParameter<FLandscapeUniformShaderParameters>(), *BatchElementParams->LandscapeUniformShaderParametersResource);
-		ShaderBindings.Add(Shader->GetUniformBufferParameter<FLandscapeSectionLODUniformParameters>(), BatchElementParams->LandscapeSectionLODUniformParameters);
+	ShaderBindings.Add(Shader->GetUniformBufferParameter<FLandscapeUniformShaderParameters>(), *BatchElementParams->LandscapeUniformShaderParametersResource);
+	ShaderBindings.Add(Shader->GetUniformBufferParameter<FLandscapeSectionLODUniformParameters>(), BatchElementParams->LandscapeSectionLODUniformParameters);
 
-#if RHI_RAYTRACING
-		if (IsRayTracingEnabled())
+	if (InView && VertexFactory->GetType() == Culling::GetTileVertexFactoryType())
+	{
+		// TileVF is used only for LOD0
+		int32 LODIndex = 0; 
+		Culling::FArguments CullingArgs{};
+		if (Culling::GetViewArguments(*InView, SceneProxy->LandscapeKey, SceneProxy->ComponentBase, LODIndex, CullingArgs))
 		{
-			ShaderBindings.Add(Shader->GetUniformBufferParameter<FLandscapeVertexFactoryMVFParameters>(), BatchElementParams->LandscapeVertexFactoryMVFUniformBuffer);
+			check(VertexStreams.IsValidIndex(1));
+			VertexStreams[1].Offset = CullingArgs.TileDataOffset;
+			VertexStreams[1].VertexBuffer = CullingArgs.TileDataVertexBuffer;
 		}
-#endif
 	}
-};
 
-IMPLEMENT_TYPE_LAYOUT(FLandscapeVertexFactoryVertexShaderParameters);
+#if RHI_RAYTRACING
+	if (IsRayTracingEnabled())
+	{
+		ShaderBindings.Add(Shader->GetUniformBufferParameter<FLandscapeVertexFactoryMVFParameters>(), BatchElementParams->LandscapeVertexFactoryMVFUniformBuffer);
+	}
+#endif
+}
 
 /**
   * Shader parameters for use with FLandscapeFixedGridVertexFactory
@@ -3050,7 +3103,7 @@ void FLandscapeVertexFactory::ModifyCompilationEnvironment(const FVertexFactoryS
 
 void FLandscapeVertexFactory::GetPSOPrecacheVertexFetchElements(EVertexInputStreamType VertexInputStreamType, FVertexDeclarationElementList& Elements)
 {
-	Elements.Add(FVertexElement(0, 0, VET_Float4, 0, 0, false));
+	Elements.Add(FVertexElement(0, 0, VET_UByte4, 0, 0, false));
 	Elements.Add(FVertexElement(1, 0, VET_UInt, 1, 0, true));
 }
 
@@ -3347,8 +3400,10 @@ public:
 
 				static const FName LandscapeVertexFactory = FName(TEXT("FLandscapeVertexFactory"));
 				static const FName LandscapeXYOffsetVertexFactory = FName(TEXT("FLandscapeXYOffsetVertexFactory"));
+				static const FName LandscapeTileVertexFactory = FName(TEXT("FLandscapeTileVertexFactory"));
 				static const FName NaniteVertexFactory = FName(TEXT("Nanite::FVertexFactory"));
 				if (VertexFactoryType->GetFName() == LandscapeVertexFactory ||
+					VertexFactoryType->GetFName() == LandscapeTileVertexFactory ||
 					VertexFactoryType->GetFName() == LandscapeXYOffsetVertexFactory ||
 					VertexFactoryType->GetFName() == NaniteVertexFactory)
 				{
@@ -3919,6 +3974,12 @@ double FLandscapeComponentSceneProxy::ComputeSectionResolution() const
 	return ComponentFullExtent / ComponentQuads;
 }
 
+void FLandscapeComponentSceneProxy::GetSectionBoundsAndLocalToWorld(FBoxSphereBounds& OutLocalBounds, FMatrix& OutLocalToWorld) const
+{
+	OutLocalBounds = GetLocalBounds();
+	OutLocalToWorld = GetLocalToWorld();
+}
+
 //
 // FLandscapeSectionInfo
 //
@@ -3965,6 +4026,12 @@ public:
 	{
 		return ProxyLOD;
 	}
+	
+	virtual void GetSectionBoundsAndLocalToWorld(FBoxSphereBounds& LocalBounds, FMatrix& LocalToWorld) const
+	{
+		LocalBounds = FBoxSphereBounds(ForceInit);
+		LocalToWorld = FMatrix::Identity;
+	}
 
 private:
 	int8 ProxyLOD;
diff --git a/Engine/Source/Runtime/Landscape/Public/LandscapeRender.h b/Engine/Source/Runtime/Landscape/Public/LandscapeRender.h
--- a/Engine/Source/Runtime/Landscape/Public/LandscapeRender.h
+++ b/Engine/Source/Runtime/Landscape/Public/LandscapeRender.h
@@ -154,6 +154,31 @@ public:
 	TArray<FLandscapeBatchElementParams, SceneRenderingAllocator> ElementParams;
 };
 
+class FLandscapeVertexFactoryVertexShaderParameters : public FVertexFactoryShaderParameters
+{
+	DECLARE_TYPE_LAYOUT(FLandscapeVertexFactoryVertexShaderParameters, NonVirtual);
+public:
+	/**
+	* Bind shader constants by name
+	* @param	ParameterMap - mapping of named shader constants to indices
+	*/
+	void Bind(const FShaderParameterMap& ParameterMap)
+	{
+	}
+
+	void GetElementShaderBindings(
+		const class FSceneInterface* Scene,
+		const FSceneView* InView,
+		const class FMeshMaterialShader* Shader,
+		const EVertexInputStreamType InputStreamType,
+		ERHIFeatureLevel::Type FeatureLevel,
+		const FVertexFactory* VertexFactory,
+		const FMeshBatchElement& BatchElement,
+		class FMeshDrawSingleShaderBindings& ShaderBindings,
+		FVertexInputStreamArray& VertexStreams
+	) const;
+};
+
 /** Pixel shader parameters for use with FLandscapeVertexFactory */
 class FLandscapeVertexFactoryPixelShaderParameters : public FVertexFactoryShaderParameters
 {
@@ -269,10 +294,10 @@ public:
 
 struct FLandscapeVertex
 {
-	float VertexX;
-	float VertexY;
-	float SubX;
-	float SubY;
+	uint8 VertexX;
+	uint8 VertexY;
+	uint8 SubX;
+	uint8 SubY;
 };
 
 //
@@ -332,6 +357,10 @@ public:
 	FLandscapeVertexFactory* VertexFactory;
 	FLandscapeVertexFactory* FixedGridVertexFactory;
 	FLandscapeVertexBuffer* VertexBuffer;
+	
+	FLandscapeVertexFactory* TileVertexFactory;
+	FVertexBuffer* TileDataBuffer;
+	
 	FIndexBuffer** IndexBuffers;
 	FLandscapeIndexRanges* IndexRanges;
 	bool bUse32BitIndices;
@@ -376,7 +405,8 @@ public:
 
 	/** Computes the worldspace units per vertex of the landscape section. */
 	virtual double ComputeSectionResolution() const { return -1.0; }
-
+	
+	virtual void GetSectionBoundsAndLocalToWorld(FBoxSphereBounds& LocalBounds, FMatrix& LocalToWorld) const = 0;
 public:
 	uint32 LandscapeKey;
 	FIntPoint ComponentBase;
@@ -502,7 +532,8 @@ public:
 	virtual void SetupViewFamily(FSceneViewFamily& InViewFamily) override {}
 	virtual void SetupView(FSceneViewFamily& InViewFamily, FSceneView& InView) override {}
 	virtual void BeginRenderViewFamily(FSceneViewFamily& InViewFamily) override;
-
+	
+	virtual void PreRenderViewFamily_RenderThread(FRDGBuilder& GraphBuilder, FSceneViewFamily& InViewFamily) override;
 	virtual void PreRenderView_RenderThread(FRDGBuilder& GraphBuilder, FSceneView& InView) override;
 	virtual void PreInitViews_RenderThread(FRDGBuilder& GraphBuilder) override;
 
@@ -793,6 +824,7 @@ public:
 	virtual void ApplyWorldOffset(FVector InOffset) override;
 	virtual void DrawStaticElements(FStaticPrimitiveDrawInterface* PDI) override;
 	virtual void GetDynamicMeshElements(const TArray<const FSceneView*>& Views, const FSceneViewFamily& ViewFamily, uint32 VisibilityMap, FMeshElementCollector& Collector) const override;
+	virtual void ApplyViewDependentMeshArguments(const FSceneView& View, FMeshBatch& ViewDependentMeshBatch) const override;
 	virtual uint32 GetMemoryFootprint() const override { return(sizeof(*this) + GetAllocatedSize()); }
 	virtual FPrimitiveViewRelevance GetViewRelevance(const FSceneView* View) const override;
 	virtual bool CanBeOccluded() const override;
@@ -836,6 +868,7 @@ public:
 	virtual float ComputeLODForView(const FSceneView& InView) const override;
 	virtual float ComputeLODBias() const override;
 	virtual double ComputeSectionResolution() const override;
+	virtual void GetSectionBoundsAndLocalToWorld(FBoxSphereBounds& LocalBounds, FMatrix& LocalToWorld) const override;
 };
 
 class FLandscapeDebugMaterialRenderProxy : public FMaterialRenderProxy
diff --git a/Engine/Source/Runtime/Renderer/Private/MeshDrawCommands.cpp b/Engine/Source/Runtime/Renderer/Private/MeshDrawCommands.cpp
--- a/Engine/Source/Runtime/Renderer/Private/MeshDrawCommands.cpp
+++ b/Engine/Source/Runtime/Renderer/Private/MeshDrawCommands.cpp
@@ -631,7 +631,17 @@ void GenerateDynamicMeshDrawCommands(
 		{
 			const FStaticMeshBatch* StaticMeshBatch = DynamicMeshCommandBuildRequests[MeshIndex];
 			const uint64 DefaultBatchElementMask = ~0ul;
-			PassMeshProcessor->AddMeshBatch(*StaticMeshBatch, DefaultBatchElementMask, StaticMeshBatch->PrimitiveSceneInfo->Proxy, StaticMeshBatch->Id);
+
+			if (StaticMeshBatch->bViewDependentArguments)
+			{
+				FMeshBatch ViewDepenedentMeshBatch(*StaticMeshBatch);
+				StaticMeshBatch->PrimitiveSceneInfo->Proxy->ApplyViewDependentMeshArguments(View, ViewDepenedentMeshBatch);
+				PassMeshProcessor->AddMeshBatch(ViewDepenedentMeshBatch, DefaultBatchElementMask, StaticMeshBatch->PrimitiveSceneInfo->Proxy, StaticMeshBatch->Id);
+			}
+			else
+			{
+				PassMeshProcessor->AddMeshBatch(*StaticMeshBatch, DefaultBatchElementMask, StaticMeshBatch->PrimitiveSceneInfo->Proxy, StaticMeshBatch->Id);
+			}
 		}
 
 		const int32 NumCommandsGenerated = VisibleCommands.Num() - NumCommandsBefore;
@@ -710,23 +720,34 @@ void GenerateMobileBasePassDynamicMeshDrawCommands(
 	{
 		const int32 NumCommandsBefore = VisibleCommands.Num();
 		const int32 NumStaticMeshBatches = DynamicMeshCommandBuildRequests.Num();
+		FMeshBatch ViewDependentMeshBatch{};
 
 		for (int32 MeshIndex = 0; MeshIndex < NumStaticMeshBatches; MeshIndex++)
 		{
 			const FStaticMeshBatch* StaticMeshBatch = DynamicMeshCommandBuildRequests[MeshIndex];
+			const int32 StaticMeshBatchId = StaticMeshBatch->Id;
+			const FPrimitiveSceneProxy* Proxy = StaticMeshBatch->PrimitiveSceneInfo->Proxy;
 
-			const int32 PrimitiveIndex = StaticMeshBatch->PrimitiveSceneInfo->Proxy->GetPrimitiveSceneInfo()->GetIndex();
+			const FMeshBatch* MeshBatch = StaticMeshBatch;
+			if (MeshBatch->bViewDependentArguments)
+			{
+				ViewDependentMeshBatch = *MeshBatch;
+				Proxy->ApplyViewDependentMeshArguments(View, ViewDependentMeshBatch);
+				MeshBatch = &ViewDependentMeshBatch;
+			}
+
+			const int32 PrimitiveIndex = Proxy->GetPrimitiveSceneInfo()->GetIndex();
 			if (!bSkipCSMShaderCulling
 				&& MobileCSMVisibilityInfo.bMobileDynamicCSMInUse
 				&& (MobileCSMVisibilityInfo.bAlwaysUseCSM || MobileCSMVisibilityInfo.MobilePrimitiveCSMReceiverVisibilityMap[PrimitiveIndex]))
 			{
 				const uint64 DefaultBatchElementMask = ~0ul;
-				MobilePassCSMPassMeshProcessor->AddMeshBatch(*StaticMeshBatch, DefaultBatchElementMask, StaticMeshBatch->PrimitiveSceneInfo->Proxy, StaticMeshBatch->Id);
+				MobilePassCSMPassMeshProcessor->AddMeshBatch(*MeshBatch, DefaultBatchElementMask, Proxy, StaticMeshBatchId);
 			}
 			else
 			{
 				const uint64 DefaultBatchElementMask = ~0ul;
-				PassMeshProcessor->AddMeshBatch(*StaticMeshBatch, DefaultBatchElementMask, StaticMeshBatch->PrimitiveSceneInfo->Proxy, StaticMeshBatch->Id);
+				PassMeshProcessor->AddMeshBatch(*MeshBatch, DefaultBatchElementMask, Proxy, StaticMeshBatchId);
 			}
 		}
 
diff --git a/Engine/Source/Runtime/Renderer/Private/ShaderBaseClasses.cpp b/Engine/Source/Runtime/Renderer/Private/ShaderBaseClasses.cpp
--- a/Engine/Source/Runtime/Renderer/Private/ShaderBaseClasses.cpp
+++ b/Engine/Source/Runtime/Renderer/Private/ShaderBaseClasses.cpp
@@ -517,7 +517,8 @@ void FMeshMaterialShader::GetElementShaderBindings(
 	{
 		const FShaderType* ShaderType = GetType(PointerTable);
 		ensureMsgf(!GetUniformBufferParameter<FPrimitiveUniformShaderParameters>().IsBound(), TEXT("Shader %s attempted to bind the Primitive uniform buffer even though Vertex Factory computes a PrimitiveId per-instance.  This will break auto-instancing.  Shaders should use GetPrimitiveData(PrimitiveId).Member instead of Primitive.Member."), ShaderType->GetName());
-		ensureMsgf(!BatchElement.PrimitiveUniformBuffer || (FeatureLevel == ERHIFeatureLevel::ES3_1), TEXT("FMeshBatchElement was assigned a PrimitiveUniformBuffer even though Vertex Factory %s fetches primitive shader data through a Scene buffer.  The assigned PrimitiveUniformBuffer cannot be respected.  Use PrimitiveUniformBufferResource instead for dynamic primitive data."), ShaderType->GetName());
+		// Some primitives may use several VFs with a mixed support for a GPUScene. In this case all mesh batches get Primitive UB assigned regardless of VF type 
+		ensureMsgf(!BatchElement.PrimitiveUniformBuffer || (FeatureLevel == ERHIFeatureLevel::ES3_1 || PrimitiveSceneProxy->DoesVFRequirePrimitiveUniformBuffer()), TEXT("FMeshBatchElement was assigned a PrimitiveUniformBuffer even though Vertex Factory %s fetches primitive shader data through a Scene buffer.  The assigned PrimitiveUniformBuffer cannot be respected.  Use PrimitiveUniformBufferResource instead for dynamic primitive data."), ShaderType->GetName());
 	}
 	else
 	{
