Fixed non-unity compilation errors

diff --git a/Engine/Source/Runtime/Landscape/Private/LandscapeCulling.h b/Engine/Source/Runtime/Landscape/Private/LandscapeCulling.h
--- a/Engine/Source/Runtime/Landscape/Private/LandscapeCulling.h
+++ b/Engine/Source/Runtime/Landscape/Private/LandscapeCulling.h
@@ -4,7 +4,16 @@
 
 #include "CoreMinimal.h"
 
+class FRHIBuffer;
+class FVertexFactoryType;
+enum EShaderPlatform : uint16;
+namespace ERHIFeatureLevel { enum Type : int; }
 class FLandscapeSharedBuffers;
+struct FMeshBatch;
+class FSceneViewFamily;
+class FSceneView;
+class FRDGBuilder;
+struct FViewMatrices;
 
 /**
 * Landscape GPU culling implementation
