Landscape: Fix truncation warnings, part 5

diff --git a/Engine/Source/Runtime/Landscape/Private/LandscapeCulling.cpp b/Engine/Source/Runtime/Landscape/Private/LandscapeCulling.cpp
--- a/Engine/Source/Runtime/Landscape/Private/LandscapeCulling.cpp
+++ b/Engine/Source/Runtime/Landscape/Private/LandscapeCulling.cpp
@@ -219,8 +219,8 @@ void FLandscapeTileMesh::InitRHI()
 		{
 			for (uint32 x = 0; x < TileSizeVertx; x++)
 			{
-				Vertex->VertexX = x;
-				Vertex->VertexY = y;
+				Vertex->VertexX = static_cast<uint8>(x);
+				Vertex->VertexY = static_cast<uint8>(y);
 				Vertex->SubX = 0;
 				Vertex->SubY = 0;
 				Vertex++;
@@ -242,10 +242,10 @@ void FLandscapeTileMesh::InitRHI()
 		{
 			for (uint32 x = 0; x < LANDSCAPE_TILE_QUADS; x++)
 			{
-				uint16 i00 = (x + 0) + (y + 0) * (TileSizeVertx);
-				uint16 i10 = (x + 1) + (y + 0) * (TileSizeVertx);
-				uint16 i11 = (x + 1) + (y + 1) * (TileSizeVertx);
-				uint16 i01 = (x + 0) + (y + 1) * (TileSizeVertx);
+				uint16 i00 = static_cast<uint16>((x + 0) + (y + 0) * (TileSizeVertx));
+				uint16 i10 = static_cast<uint16>((x + 1) + (y + 0) * (TileSizeVertx));
+				uint16 i11 = static_cast<uint16>((x + 1) + (y + 1) * (TileSizeVertx));
+				uint16 i01 = static_cast<uint16>((x + 0) + (y + 1) * (TileSizeVertx));
 
 				Indices[0] = i00;
 				Indices[1] = i11;
@@ -304,10 +304,10 @@ void FLandscapeTileDataBuffer::InitRHI()
 			{
 				for (uint32 x = 0; x < SubsectionTilesRow; x++)
 				{
-					TileData[0] = x * LANDSCAPE_TILE_QUADS;
-					TileData[1] = y * LANDSCAPE_TILE_QUADS;
-					TileData[2] = SubX;
-					TileData[3] = SubY;
+					TileData[0] = static_cast<uint8>(x * LANDSCAPE_TILE_QUADS);
+					TileData[1] = static_cast<uint8>(y * LANDSCAPE_TILE_QUADS);
+					TileData[2] = static_cast<uint8>(SubX);
+					TileData[3] = static_cast<uint8>(SubY);
 					TileData += Stride;
 				}
 			}
@@ -607,10 +607,9 @@ static void ComputeSectionIntermediateData(FRDGBuilder& GraphBuilder, TArrayView
 
 			for (int32 SectionIdx = 0; SectionIdx < SectionLODValues.Num(); ++SectionIdx)
 			{
-				int32 LODValue = SectionLODValues[SectionIdx];
 				FLandscapeSectionInfo* SectionInfo = RenderSystem.SectionInfos[SectionIdx];
 
-				if (LODValue == 0 && SectionInfo != nullptr)
+				if (SectionLODValues[SectionIdx] == 0 && SectionInfo != nullptr)
 				{
 					FBuildLandscapeTileDataCS::FLandscapeSection& Section = SectionsData.AddDefaulted_GetRef();
 
@@ -618,14 +617,14 @@ static void ComputeSectionIntermediateData(FRDGBuilder& GraphBuilder, TArrayView
 					FMatrix SectionLocalToWorld;
 					SectionInfo->GetSectionBoundsAndLocalToWorld(SectionLocalBounds, SectionLocalToWorld);
 					const FLargeWorldRenderPosition SectionAbsoluteOrigin(SectionLocalToWorld.GetOrigin());
-					const int32 NeighborsMaxLOD = FMath::RoundFromZero(ComputeNeighborsMaxLOD(RenderSystem, SectionLODValues, SectionInfo->ComponentBase));
+					const int32 NeighborsMaxLOD = static_cast<int32>(FMath::RoundFromZero(ComputeNeighborsMaxLOD(RenderSystem, SectionLODValues, SectionInfo->ComponentBase)));
 
 					Section.LocalToRelativeWorld = FLargeWorldRenderScalar::MakeToRelativeWorldMatrix(SectionAbsoluteOrigin.GetTileOffset(), SectionLocalToWorld);
 					Section.TilePosition = SectionAbsoluteOrigin.GetTile();
-					Section.LocalZ = SectionLocalBounds.Origin.Z;
-					Section.HalfHeight = SectionLocalBounds.BoxExtent.Z;
+					Section.LocalZ = static_cast<float>(SectionLocalBounds.Origin.Z);
+					Section.HalfHeight = static_cast<float>(SectionLocalBounds.BoxExtent.Z);
 					// How many quads to add to each tile extent to compensate for a neighbors LOD
-					Section.NeighborLODExtent = (1 << NeighborsMaxLOD) - 1; 
+					Section.NeighborLODExtent = static_cast<float>((1 << NeighborsMaxLOD) - 1); 
 
 					ViewStateIntermediates.SectionBases.Add(SectionInfo->ComponentBase);
 				}
diff --git a/Engine/Source/Runtime/Landscape/Private/LandscapeEditLayers.cpp b/Engine/Source/Runtime/Landscape/Private/LandscapeEditLayers.cpp
--- a/Engine/Source/Runtime/Landscape/Private/LandscapeEditLayers.cpp
+++ b/Engine/Source/Runtime/Landscape/Private/LandscapeEditLayers.cpp
@@ -3278,7 +3278,7 @@ void ALandscape::PrintLayersDebugHeightData(const FString& InContext, const TArr
 
 			if (DisplayHeightAsDelta)
 			{
-				HeightDelta = HeightDelta >= LandscapeDataAccess::MidValue ? HeightDelta - LandscapeDataAccess::MidValue : HeightDelta;
+				HeightDelta = static_cast<int32>(HeightDelta >= LandscapeDataAccess::MidValue ? HeightDelta - LandscapeDataAccess::MidValue : HeightDelta);
 			}
 
 			if (X > 0 && MipSize > 0 && X % MipSize == 0)
diff --git a/Engine/Source/Runtime/Landscape/Private/LandscapeGizmoActor.cpp b/Engine/Source/Runtime/Landscape/Private/LandscapeGizmoActor.cpp
--- a/Engine/Source/Runtime/Landscape/Private/LandscapeGizmoActor.cpp
+++ b/Engine/Source/Runtime/Landscape/Private/LandscapeGizmoActor.cpp
@@ -62,6 +62,14 @@ public:
 
 	virtual bool GetParameterValue(EMaterialParameterType Type, const FHashedMaterialParameterInfo& ParameterInfo, FMaterialParameterValue& OutValue, const FMaterialRenderContext& Context) const
 	{
+		auto GetColor = [](const FMatrix::FReal* MatrixColumn)
+		{
+			return FLinearColor(static_cast<float>(MatrixColumn[0]),
+			                    static_cast<float>(MatrixColumn[1]),
+			                    static_cast<float>(MatrixColumn[2]),
+			                    static_cast<float>(MatrixColumn[3]));
+		};
+		
 		switch (Type)
 		{
 		case EMaterialParameterType::Vector:
@@ -72,22 +80,22 @@ public:
 			}
 			else if (ParameterInfo.Name == FName(TEXT("MatrixRow1")))
 			{
-				OutValue = FLinearColor(WorldToLandscapeMatrix.M[0][0], WorldToLandscapeMatrix.M[0][1], WorldToLandscapeMatrix.M[0][2], WorldToLandscapeMatrix.M[0][3]);
+				OutValue = GetColor(WorldToLandscapeMatrix.M[0]);
 				return true;
 			}
 			else if (ParameterInfo.Name == FName(TEXT("MatrixRow2")))
 			{
-				OutValue = FLinearColor(WorldToLandscapeMatrix.M[1][0], WorldToLandscapeMatrix.M[1][1], WorldToLandscapeMatrix.M[1][2], WorldToLandscapeMatrix.M[1][3]);
+				OutValue = GetColor(WorldToLandscapeMatrix.M[1]);
 				return true;
 			}
 			else if (ParameterInfo.Name == FName(TEXT("MatrixRow3")))
 			{
-				OutValue = FLinearColor(WorldToLandscapeMatrix.M[2][0], WorldToLandscapeMatrix.M[2][1], WorldToLandscapeMatrix.M[2][2], WorldToLandscapeMatrix.M[2][3]);
+				OutValue = GetColor(WorldToLandscapeMatrix.M[2]);
 				return true;
 			}
 			else if (ParameterInfo.Name == FName(TEXT("MatrixRow4")))
 			{
-				OutValue = FLinearColor(WorldToLandscapeMatrix.M[3][0], WorldToLandscapeMatrix.M[3][1], WorldToLandscapeMatrix.M[3][2], WorldToLandscapeMatrix.M[3][3]);
+				OutValue = GetColor(WorldToLandscapeMatrix.M[3]);
 				return true;
 			}
 			break;
@@ -154,9 +162,9 @@ public:
 			ULandscapeInfo* LandscapeInfo = Gizmo->TargetLandscapeInfo;
  			if (LandscapeInfo && LandscapeInfo->GetLandscapeProxy() && !FMath::IsNearlyZero(LandscapeInfo->DrawScale.X))
 			{
-				const float ScaleXY = LandscapeInfo->DrawScale.X;
-				SampleSizeX = Gizmo->SampleSizeX;
-				SampleSizeY = Gizmo->SampleSizeY;
+				const float ScaleXY = static_cast<float>(LandscapeInfo->DrawScale.X);
+				SampleSizeX = static_cast<float>(Gizmo->SampleSizeX);
+				SampleSizeY = static_cast<float>(Gizmo->SampleSizeY);
 				bHeightmapRendering = (Gizmo->DataType & LGT_Height);
 				FTransform LToW = LandscapeInfo->GetLandscapeProxy()->LandscapeActorToWorld();
 				const float W = Gizmo->Width / 2;
@@ -193,7 +201,13 @@ public:
 				{
 					MeshRT = FTranslationMatrix(FVector(-W / ScaleXY + 0.5, -H / ScaleXY + 0.5, 0) * GizmoScale3D) * FRotationTranslationMatrix(FRotator(0, Gizmo->GetActorRotation().Yaw, 0), FVector(BaseLocation.X, BaseLocation.Y, 0)) * LToW.ToMatrixWithScale();
 
-					HeightmapRenderProxy = new FLandscapeGizmoMeshRenderProxy( Gizmo->GizmoMeshMaterial->GetRenderProxy(), BaseLocation.Z + L, BaseLocation.Z, Gizmo->GizmoTexture, FLinearColor(Gizmo->TextureScale.X, Gizmo->TextureScale.Y, 0, 0), WToL );
+					HeightmapRenderProxy = new FLandscapeGizmoMeshRenderProxy(Gizmo->GizmoMeshMaterial->GetRenderProxy(),
+					                                                          static_cast<float>(BaseLocation.Z + L),
+					                                                          static_cast<float>(BaseLocation.Z),
+					                                                          Gizmo->GizmoTexture,
+					                                                          FLinearColor(static_cast<float>(Gizmo->TextureScale.X),
+					                                                                       static_cast<float>(Gizmo->TextureScale.Y), 0.0f, 0.0f),
+					                                                          WToL);
 					GizmoRenderProxy = (Gizmo->DataType != LGT_None) ? Gizmo->GizmoDataMaterial->GetRenderProxy() : Gizmo->GizmoMaterial->GetRenderProxy();
 
 					FScaleMatrix Mat(FVector(ScaleX, ScaleY, L));
@@ -210,7 +224,7 @@ public:
 							FVector SampledPos = Gizmo->SampledHeight[X + Y * ALandscapeGizmoActiveActor::DataTexSize];
 							SampledPos.X *= ScaleX;
 							SampledPos.Y *= ScaleY;
-							SampledPos.Z = Gizmo->GetLandscapeHeight(SampledPos.Z);
+							SampledPos.Z = Gizmo->GetLandscapeHeight(static_cast<float>(SampledPos.Z));
 
 							FVector SampledNormal = NormalM.TransformVector(Gizmo->SampledNormal[X + Y * ALandscapeGizmoActiveActor::DataTexSize]);
 							SampledNormal = SampledNormal.GetSafeNormal();
@@ -305,17 +319,23 @@ public:
 
 					if (bHeightmapRendering)
 					{
+						auto GetSampleIndex = [this](int32 X, int32 Y) -> int32
+						{
+							return static_cast<int32>(X + Y * SampleSizeX);
+						};
+
 						FDynamicMeshBuilder MeshBuilder(View->GetFeatureLevel());
 
 						for (int32 Y = 0; Y < SampleSizeY; ++Y)
 						{
 							for (int32 X = 0; X < SampleSizeX; ++X)
 							{
-								FVector SampledNormal = SampledNormals[X + Y * SampleSizeX];
+								const int32 SampleIndex = GetSampleIndex(X, Y);
+								FVector SampledNormal = SampledNormals[SampleIndex];
 								FVector TangentX(SampledNormal.Z, 0, -SampledNormal.X);
 								TangentX = TangentX.GetSafeNormal();
 
-								MeshBuilder.AddVertex((FVector3f)SampledPositions[X + Y * SampleSizeX], FVector2f((float)X / (SampleSizeX), (float)Y / (SampleSizeY)), (FVector3f)TangentX, FVector3f(SampledNormal^TangentX), (FVector3f)SampledNormal, FColor::White);
+								MeshBuilder.AddVertex((FVector3f)SampledPositions[SampleIndex], FVector2f((float)X / (SampleSizeX), (float)Y / (SampleSizeY)), (FVector3f)TangentX, FVector3f(SampledNormal^TangentX), (FVector3f)SampledNormal, FColor::White);
 							}
 						}
 
@@ -325,8 +345,8 @@ public:
 							{
 								if (X < SampleSizeX - 1 && Y < SampleSizeY - 1)
 								{
-									MeshBuilder.AddTriangle( (X+0) + (Y+0) * SampleSizeX, (X+1) + (Y+1) * SampleSizeX, (X+1) + (Y+0) * SampleSizeX );
-									MeshBuilder.AddTriangle( (X+0) + (Y+0) * SampleSizeX, (X+0) + (Y+1) * SampleSizeX, (X+1) + (Y+1) * SampleSizeX );
+									MeshBuilder.AddTriangle(GetSampleIndex(X + 0, Y + 0), GetSampleIndex(X + 1, Y + 1), GetSampleIndex(X + 1, Y + 0));
+									MeshBuilder.AddTriangle(GetSampleIndex(X + 0, Y + 0), GetSampleIndex(X + 0, Y + 1), GetSampleIndex(X + 1, Y + 1));
 								}
 							}
 						}
@@ -353,7 +373,7 @@ public:
 	}
 
 	virtual uint32 GetMemoryFootprint( void ) const override { return( sizeof( *this ) + GetAllocatedSize() ); }
-	uint32 GetAllocatedSize( void ) const { return( FPrimitiveSceneProxy::GetAllocatedSize() ); }
+	uint32 GetAllocatedSize( void ) const { return static_cast<uint32>(FPrimitiveSceneProxy::GetAllocatedSize()); }
 };
 
 ULandscapeGizmoRenderComponent::ULandscapeGizmoRenderComponent(const FObjectInitializer& ObjectInitializer)
@@ -577,7 +597,7 @@ FVector ALandscapeGizmoActiveActor::SnapToLandscapeGrid(const FVector& GizmoLoca
 		switch (SnapType)
 		{
 		case ELandscapeGizmoSnapType::Component:
-			SnapDimension = LandscapeProxy->ComponentSizeQuads;
+			SnapDimension = static_cast<float>(LandscapeProxy->ComponentSizeQuads);
 			break;
 		case ELandscapeGizmoSnapType::Texel:
 			SnapDimension = 1.0f;
@@ -602,7 +622,7 @@ FVector ALandscapeGizmoActiveActor::SnapToLandscapeGrid(const FVector& GizmoLoca
 
 		if (bFollowTerrainHeight)
 		{
-			ResultLocation.Z = LandscapeProxy->GetHeightAtLocation(ResultLocation, EHeightfieldSource::Editor).Get(GizmoLocation.Z);
+			ResultLocation.Z = LandscapeProxy->GetHeightAtLocation(ResultLocation, EHeightfieldSource::Editor).Get(static_cast<float>(GizmoLocation.Z));
 		}
 	}
 	
@@ -712,8 +732,8 @@ void ALandscapeGizmoActiveActor::SetTargetLandscape(ULandscapeInfo* LandscapeInf
 		// if there's no copied data, try to move somewhere useful
 		if (TargetLandscapeInfo && DataType == LGT_None)
 		{
-			MarginZ = TargetLandscapeInfo->DrawScale.Z * 3;
-			Width = Height = TargetLandscapeInfo->DrawScale.X * (TargetLandscapeInfo->ComponentSizeQuads + 1);
+			MarginZ = static_cast<float>(TargetLandscapeInfo->DrawScale.Z * 3);
+			Width = Height = static_cast<float>(TargetLandscapeInfo->DrawScale.X * (TargetLandscapeInfo->ComponentSizeQuads + 1));
 
 			float NewLengthZ;
 			FVector NewLocation = TargetLandscapeInfo->GetLandscapeCenterPos(NewLengthZ);
@@ -756,9 +776,9 @@ void ALandscapeGizmoActiveActor::FitToSelection()
 		{
 			const FVector LocalScale3D = GetRootComponent()->GetRelativeScale3D();
 
-			float ScaleXY = TargetLandscapeInfo->DrawScale.X;
-			Width = ScaleXY * (MaxX - MinX + 1) / (LocalScale3D.X);
-			Height = ScaleXY * (MaxY - MinY + 1) / (LocalScale3D.Y);
+			const float ScaleXY = static_cast<float>(TargetLandscapeInfo->DrawScale.X);
+			Width = static_cast<float>(ScaleXY * (MaxX - MinX + 1) / LocalScale3D.X);
+			Height = static_cast<float>(ScaleXY * (MaxY - MinY + 1) / LocalScale3D.Y);
 			float NewLengthZ;
 			FVector NewLocation = TargetLandscapeInfo->GetLandscapeCenterPos(NewLengthZ, MinX, MinY, MaxX, MaxY);
 			SetLength(NewLengthZ);
@@ -787,8 +807,8 @@ void ALandscapeGizmoActiveActor::FitMinMaxHeight()
 
 		if (MinZ != UE_OLD_HALF_WORLD_MAX && MaxZ > MinZ + KINDA_SMALL_NUMBER)
 		{
-			MinRelativeZ = MinZ;
-			RelativeScaleZ = 1.f / (MaxZ - MinZ);
+			MinRelativeZ = static_cast<float>(MinZ);
+			RelativeScaleZ = static_cast<float>(1.f / (MaxZ - MinZ));
 			ReregisterAllComponents();
 		}
 	}
@@ -806,7 +826,9 @@ float ALandscapeGizmoActiveActor::GetNormalizedHeight(uint16 LandscapeHeight) co
 			if (ZScale > KINDA_SMALL_NUMBER)
 			{
 				FVector LocalGizmoPos = Proxy->LandscapeActorToWorld().InverseTransformPosition(GetActorLocation());
-				return FMath::Clamp<float>( (( LandscapeDataAccess::GetLocalHeight(LandscapeHeight) - LocalGizmoPos.Z) * TargetLandscapeInfo->DrawScale.Z) / ZScale, 0.f, 1.f );
+				return FMath::Clamp<float>(
+					static_cast<float>((LandscapeDataAccess::GetLocalHeight(LandscapeHeight) - LocalGizmoPos.Z) * TargetLandscapeInfo->DrawScale.Z / ZScale),
+					0.f, 1.f);
 			}
 		}
 	}
@@ -824,7 +846,7 @@ float ALandscapeGizmoActiveActor::GetWorldHeight(float NormalizedHeight) const
 			if (ZScale > KINDA_SMALL_NUMBER)
 			{
 				FVector LocalGizmoPos = Proxy->LandscapeActorToWorld().InverseTransformPosition(GetActorLocation());
-				return NormalizedHeight * ZScale + LocalGizmoPos.Z * TargetLandscapeInfo->DrawScale.Z;
+				return static_cast<float>(NormalizedHeight * ZScale + LocalGizmoPos.Z * TargetLandscapeInfo->DrawScale.Z);
 			}
 		}
 	}
@@ -836,7 +858,7 @@ float ALandscapeGizmoActiveActor::GetLandscapeHeight(float NormalizedHeight) con
 	if (TargetLandscapeInfo)
 	{
 		NormalizedHeight = (NormalizedHeight - MinRelativeZ) * RelativeScaleZ;
-		float ScaleZ = TargetLandscapeInfo->DrawScale.Z;
+		const float ScaleZ = static_cast<float>(TargetLandscapeInfo->DrawScale.Z);
 		return (GetWorldHeight(NormalizedHeight) / ScaleZ);
 	}
 	return 0.f;
@@ -916,11 +938,11 @@ void ALandscapeGizmoActiveActor::SampleData(int32 SizeX, int32 SizeY)
 				FGizmoSelectData* Data11 = SelectedData.Find(FIntPoint(LX+1, LY+1));
 
 				// Invert Tex Data to show selected region more visible
-				TexData[X + Y*GizmoTexSizeX] = 255 - FMath::Lerp(
+				TexData[X + Y*GizmoTexSizeX] = static_cast<uint8>(255 - FMath::Lerp(
 					FMath::Lerp(Data00 ? Data00->Ratio : 0, Data10 ? Data10->Ratio : 0, FracX),
 					FMath::Lerp(Data01 ? Data01->Ratio : 0, Data11 ? Data11->Ratio : 0, FracX),
 					FracY
-					) * 255;
+					) * 255);
 
 				if (DataType & LGT_Height)
 				{
@@ -961,7 +983,7 @@ LANDSCAPE_API void ALandscapeGizmoActiveActor::Import( int32 VertsX, int32 Verts
 
 	ClearGizmoData();
 
-	CachedScaleXY = TargetLandscapeInfo->DrawScale.X;
+	CachedScaleXY = static_cast<float>(TargetLandscapeInfo->DrawScale.X);
 	CachedWidth = CachedScaleXY * VertsX; // (DrawScale * DrawScale3D.X);
 	CachedHeight = CachedScaleXY * VertsY; // (DrawScale * DrawScale3D.Y);
 	
@@ -1063,7 +1085,7 @@ void ALandscapeGizmoActiveActor::Export(int32 Index, TArray<FString>& Filenames)
 						int32 Idx = (X-MinX) + Y *(1+MaxX-MinX);
 						if (!bExportOneTarget || Index == -1)
 						{
-							pHeightData[Idx] = FMath::Clamp<uint16>(Data->HeightData * LandscapeDataAccess::MaxValue, 0, LandscapeDataAccess::MaxValue);
+							pHeightData[Idx] = FMath::Clamp<uint16>(static_cast<uint16>(Data->HeightData * LandscapeDataAccess::MaxValue), 0, LandscapeDataAccess::MaxValue);
 						}
 
 						for( int32 i=1;i<Filenames.Num();i++ )
@@ -1071,7 +1093,7 @@ void ALandscapeGizmoActiveActor::Export(int32 Index, TArray<FString>& Filenames)
 							if (!bExportOneTarget || Index == i-1)
 							{
 								TArray<uint8>& WeightData = WeightDatas[i-1];
-								WeightData[Idx] = FMath::Clamp<uint8>(Data->WeightDataMap.FindRef(LayerInfos[i-1]), 0, 255);
+								WeightData[Idx] = FMath::Clamp<uint8>(static_cast<uint8>(Data->WeightDataMap.FindRef(LayerInfos[i-1])), 0, 255);
 							}
 						}
 					}
diff --git a/Engine/Source/Runtime/Landscape/Private/LandscapeGrass.cpp b/Engine/Source/Runtime/Landscape/Private/LandscapeGrass.cpp
--- a/Engine/Source/Runtime/Landscape/Private/LandscapeGrass.cpp
+++ b/Engine/Source/Runtime/Landscape/Private/LandscapeGrass.cpp
@@ -2003,8 +2003,8 @@ ULandscapeGrassType::ULandscapeGrassType(const FObjectInitializer& ObjectInitial
 : Super(ObjectInitializer)
 {
 	GrassDensity_DEPRECATED = 400;
-	StartCullDistance_DEPRECATED = 10000.0f;
-	EndCullDistance_DEPRECATED = 10000.0f;
+	StartCullDistance_DEPRECATED = 10000;
+	EndCullDistance_DEPRECATED = 10000;
 	PlacementJitter_DEPRECATED = 1.0f;
 	RandomRotation_DEPRECATED = true;
 	AlignToSurface_DEPRECATED = true;
@@ -2230,8 +2230,8 @@ FArchive& operator<<(FArchive& Ar, FLandscapeComponentGrassData& Data)
 				CollisionHeightData.BulkSerialize(Ar);
 				if (CollisionHeightData.Num())
 				{
-					const int32 ComponentSizeQuads = FMath::Sqrt(static_cast<float>(Data.GetHeightData().Num())) - 1;
-					const int32 CollisionSizeQuads = FMath::Sqrt(static_cast<float>(CollisionHeightData.Num())) - 1;
+					const int32 ComponentSizeQuads = static_cast<int32>(FMath::Sqrt(static_cast<float>(Data.GetHeightData().Num())) - 1);
+					const int32 CollisionSizeQuads = static_cast<int32>(FMath::Sqrt(static_cast<float>(CollisionHeightData.Num())) - 1);
 					const int32 CollisionMip = FMath::FloorLog2(ComponentSizeQuads / CollisionSizeQuads);
 					Data.HeightMipData.Add(CollisionMip, MoveTemp(CollisionHeightData));
 				}
@@ -2240,8 +2240,8 @@ FArchive& operator<<(FArchive& Ar, FLandscapeComponentGrassData& Data)
 				SimpleCollisionHeightData.BulkSerialize(Ar);
 				if (SimpleCollisionHeightData.Num())
 				{
-					const int32 ComponentSizeQuads = FMath::Sqrt(static_cast<float>(Data.GetHeightData().Num())) - 1;
-					const int32 SimpleCollisionSizeQuads = FMath::Sqrt(static_cast<float>(SimpleCollisionHeightData.Num())) - 1;
+					const int32 ComponentSizeQuads = static_cast<int32>(FMath::Sqrt(static_cast<float>(Data.GetHeightData().Num())) - 1);
+					const int32 SimpleCollisionSizeQuads = static_cast<int32>(FMath::Sqrt(static_cast<float>(SimpleCollisionHeightData.Num())) - 1);
 					const int32 SimpleCollisionMip = FMath::FloorLog2(ComponentSizeQuads / SimpleCollisionSizeQuads);
 					Data.HeightMipData.Add(SimpleCollisionMip, MoveTemp(SimpleCollisionHeightData));
 				}
@@ -2384,7 +2384,7 @@ struct FGrassBuilderBase
 		ComponentOrigin.Y = DrawScale.Y * (SectionBase.Y - LandscapeSectionOffset.Y);
 		ComponentOrigin.Z = 0.0f;
 
-		SqrtMaxInstances = FMath::CeilToInt(FMath::Sqrt(FMath::Abs(Extent.X * Extent.Y * GrassDensity / 1000.0f / 1000.0f)));
+		SqrtMaxInstances = FMath::CeilToInt32(FMath::Sqrt(FMath::Abs(Extent.X * Extent.Y * GrassDensity / 1000.0f / 1000.0f)));
 
 		bHaveValidData = SqrtMaxInstances != 0;
 
@@ -2882,8 +2882,8 @@ struct FAsyncGrassBuilder : public FGrassBuilderBase
 	FORCEINLINE_DEBUGGABLE void SampleLandscapeAtLocationLocal(const FVector& InLocation, FVector& OutLocation, float& OutLayerWeight, FVector* OutNormal = nullptr)
 	{
 		// Find location
-		const float TestX = InLocation.X / DrawScale.X - (float)SectionBase.X;
-		const float TestY = InLocation.Y / DrawScale.Y - (float)SectionBase.Y;
+		const float TestX = static_cast<float>(InLocation.X / DrawScale.X - SectionBase.X);
+		const float TestY = static_cast<float>(InLocation.Y / DrawScale.Y - SectionBase.Y);
 
 		// Find data
 		const int32 X1 = FMath::FloorToInt(TestX);
@@ -3046,7 +3046,7 @@ void ALandscapeProxy::GetGrassTypes(const UWorld* World, UMaterialInterface* Lan
 			{
 				for (auto& GrassVariety : GrassType->GrassVarieties)
 				{
-					int32 EndCullDistance = GrassVariety.GetEndCullDistance();
+					const float EndCullDistance = static_cast<float>(GrassVariety.GetEndCullDistance());
 					if (EndCullDistance > OutMaxDiscardDistance)
 					{
 						OutMaxDiscardDistance = EndCullDistance;
@@ -3095,7 +3095,7 @@ void FLandscapeGrassMapsBuilder::Build()
 	if (World)
 	{
 		int32 Count = GetOutdatedGrassMapCount();
-		FScopedSlowTask SlowTask(Count, (LOCTEXT("GrassMaps_BuildGrassMaps", "Building Grass maps")));
+		FScopedSlowTask SlowTask(static_cast<float>(Count), (LOCTEXT("GrassMaps_BuildGrassMaps", "Building Grass maps")));
 		SlowTask.MakeDialog();
 
 		for (TActorIterator<ALandscapeProxy> ProxyIt(World); ProxyIt; ++ProxyIt)
@@ -3326,7 +3326,7 @@ void ALandscapeProxy::UpdateGrassData(bool bInShouldMarkDirty, FScopedSlowTask*
 	{
 		if (InSlowTask && Increment && ((InSlowTask->CompletedWork + Increment) <= InSlowTask->TotalAmountOfWork))
 		{
-			InSlowTask->EnterProgressFrame(Increment, FText::Format(LOCTEXT("GrassMaps_BuildGrassMapsProgress", "Building Grass Map {0} of {1})"), FText::AsNumber(InSlowTask->CompletedWork), FText::AsNumber(InSlowTask->TotalAmountOfWork)));
+			InSlowTask->EnterProgressFrame(static_cast<float>(Increment), FText::Format(LOCTEXT("GrassMaps_BuildGrassMapsProgress", "Building Grass Map {0} of  {1})"), FText::AsNumber(InSlowTask->CompletedWork), FText::AsNumber(InSlowTask->TotalAmountOfWork)));
 		}
 	};
 
@@ -3480,7 +3480,7 @@ void ALandscapeProxy::UpdateGrass(const TArray<FVector>& Cameras, int32& InOutNu
 				float MinSqrDistanceToComponent = Cameras.Num() ? MAX_flt : 0.0f;
 				for (const FVector& CameraPos : Cameras)
 				{
-					MinSqrDistanceToComponent = FMath::Min<float>(MinSqrDistanceToComponent, WorldBounds.ComputeSquaredDistanceFromBoxToPoint(CameraPos));
+					MinSqrDistanceToComponent = FMath::Min<float>(MinSqrDistanceToComponent, static_cast<float>(WorldBounds.ComputeSquaredDistanceFromBoxToPoint(CameraPos)));
 				}
 
 				if (MinSqrDistanceToComponent > GrassMaxSquareDiscardDistance)
@@ -3598,7 +3598,7 @@ void ALandscapeProxy::UpdateGrass(const TArray<FVector>& Cameras, int32& InOutNu
 												MinDistanceToSubComp = Cameras.Num() ? MAX_flt : 0.0f;
 												for (auto& Pos : Cameras)
 												{
-													MinDistanceToSubComp = FMath::Min<float>(MinDistanceToSubComp, ComputeSquaredDistanceFromBoxToPoint(WorldSubBox.Min, WorldSubBox.Max, Pos));
+													MinDistanceToSubComp = FMath::Min<float>(MinDistanceToSubComp, static_cast<float>(ComputeSquaredDistanceFromBoxToPoint(WorldSubBox.Min, WorldSubBox.Max, Pos)));
 												}
 												MinDistanceToSubComp = FMath::Sqrt(MinDistanceToSubComp);
 											}
@@ -3786,8 +3786,8 @@ void ALandscapeProxy::UpdateGrass(const TArray<FVector>& Cameras, int32& InOutNu
 										}
 										else
 										{
-											GrassInstancedStaticMeshComponent->InstanceStartCullDistance = GrassVariety.GetStartCullDistance() * CullDistanceScale;
-											GrassInstancedStaticMeshComponent->InstanceEndCullDistance = GrassVariety.GetEndCullDistance() * CullDistanceScale;
+											GrassInstancedStaticMeshComponent->InstanceStartCullDistance = static_cast<int32>(GrassVariety.GetStartCullDistance() * CullDistanceScale);
+											GrassInstancedStaticMeshComponent->InstanceEndCullDistance = static_cast<int32>(GrassVariety.GetEndCullDistance() * CullDistanceScale);
 										}
 
 										{
diff --git a/Engine/Source/Runtime/Landscape/Private/LandscapeRender.cpp b/Engine/Source/Runtime/Landscape/Private/LandscapeRender.cpp
--- a/Engine/Source/Runtime/Landscape/Private/LandscapeRender.cpp
+++ b/Engine/Source/Runtime/Landscape/Private/LandscapeRender.cpp
@@ -899,9 +899,9 @@ bool FLandscapeVisibilityHelper::OnRemoveFromWorld()
 FLandscapeComponentSceneProxy::FLandscapeComponentSceneProxy(ULandscapeComponent* InComponent)
 	: FPrimitiveSceneProxy(InComponent, NAME_LandscapeResourceNameForDebugging)
 	, FLandscapeSectionInfo(InComponent->GetWorld(), InComponent->GetLandscapeProxy()->GetLandscapeGuid(), InComponent->GetSectionBase() / InComponent->ComponentSizeQuads)
-	, MaxLOD(FMath::CeilLogTwo(InComponent->SubsectionSizeQuads + 1) - 1)
-	, NumWeightmapLayerAllocations(InComponent->GetWeightmapLayerAllocations().Num())
-	, StaticLightingLOD(InComponent->GetLandscapeProxy()->StaticLightingLOD)
+	, MaxLOD(static_cast<int8>(FMath::CeilLogTwo(InComponent->SubsectionSizeQuads + 1) - 1))
+	, NumWeightmapLayerAllocations(static_cast<int8>(InComponent->GetWeightmapLayerAllocations().Num()))
+	, StaticLightingLOD(static_cast<uint8>(InComponent->GetLandscapeProxy()->StaticLightingLOD))
 	, WeightmapSubsectionOffset(InComponent->WeightmapSubsectionOffset)
 	, FirstLOD(0)
 	, LastLOD(MaxLOD)
@@ -1032,17 +1032,17 @@ FLandscapeComponentSceneProxy::FLandscapeComponentSceneProxy(ULandscapeComponent
 	const int32 MaxLODLevel = InComponent->GetLandscapeProxy()->MaxLODLevel;
 	if (MaxLODLevel >= 0)
 	{
-		MaxLOD = FMath::Min<int8>(MaxLODLevel, MaxLOD);
+		MaxLOD = FMath::Min<int8>(static_cast<int8>(MaxLODLevel), MaxLOD);
 		LastLOD = FMath::Min<int32>(MaxLODLevel, LastLOD);
 	}
 
 	// Clamp ForcedLOD to the valid range and then apply
-	int8 ForcedLOD = InComponent->ForcedLOD;
-	ForcedLOD = ForcedLOD >= 0 ? FMath::Clamp<int32>(ForcedLOD, FirstLOD, LastLOD) : ForcedLOD;
+	int8 ForcedLOD = static_cast<int8>(InComponent->ForcedLOD);
+	ForcedLOD = static_cast<int8>(ForcedLOD >= 0 ? FMath::Clamp<int32>(ForcedLOD, FirstLOD, LastLOD) : ForcedLOD);
 	FirstLOD = ForcedLOD >= 0 ? ForcedLOD : FirstLOD;
 	LastLOD = ForcedLOD >= 0 ? ForcedLOD : LastLOD;
 
-	LODSettings.LastLODIndex = LastLOD;
+	LODSettings.LastLODIndex = static_cast<int8>(LastLOD);
 	LODSettings.LastLODScreenSizeSquared = LODScreenRatioSquared[LastLOD];
 	LODSettings.ForcedLOD = ForcedLOD;
 
@@ -1052,16 +1052,16 @@ FLandscapeComponentSceneProxy::FLandscapeComponentSceneProxy(ULandscapeComponent
 
 	LastVirtualTextureLOD = MaxLOD;
 	FirstVirtualTextureLOD = bVirtualTextureRenderWithQuad ? MaxLOD : FMath::Max(MaxLOD - InComponent->GetLandscapeProxy()->VirtualTextureNumLods, 0);
-	VirtualTextureLodBias = bVirtualTextureRenderWithQuad ? 0 : InComponent->GetLandscapeProxy()->VirtualTextureLodBias;
+	VirtualTextureLodBias = static_cast<int8>(bVirtualTextureRenderWithQuad ? 0 : InComponent->GetLandscapeProxy()->VirtualTextureLodBias);
 
 #if WITH_EDITOR || !(UE_BUILD_SHIPPING || UE_BUILD_TEST)
-	LODSettings.DrawCollisionPawnLOD = CollisionResponse.GetResponse(ECC_Pawn) == ECR_Ignore ? -1 : SimpleCollisionMipLevel;
-	LODSettings.DrawCollisionVisibilityLOD = CollisionResponse.GetResponse(ECC_Visibility) == ECR_Ignore ? -1 : CollisionMipLevel;
+	LODSettings.DrawCollisionPawnLOD = static_cast<int8>(CollisionResponse.GetResponse(ECC_Pawn) == ECR_Ignore ? -1 : SimpleCollisionMipLevel);
+	LODSettings.DrawCollisionVisibilityLOD = static_cast<int8>(CollisionResponse.GetResponse(ECC_Visibility) == ECR_Ignore ? -1 : CollisionMipLevel);
 #else
 	LODSettings.DrawCollisionPawnLOD = LODSettings.DrawCollisionVisibilityLOD = -1;
 #endif
 
-	ComponentMaxExtend = SubsectionSizeQuads * FMath::Max(InComponent->GetComponentTransform().GetScale3D().X, InComponent->GetComponentTransform().GetScale3D().Y);
+	ComponentMaxExtend = static_cast<float>(SubsectionSizeQuads * FMath::Max(InComponent->GetComponentTransform().GetScale3D().X, InComponent->GetComponentTransform().GetScale3D().Y));
 
 	if (NumSubsections > 1)
 	{
@@ -1152,7 +1152,7 @@ FLandscapeComponentSceneProxy::FLandscapeComponentSceneProxy(ULandscapeComponent
 	}
 #endif
 	
-	const int8 SubsectionSizeLog2 = FMath::CeilLogTwo(InComponent->SubsectionSizeQuads + 1);
+	const int8 SubsectionSizeLog2 = static_cast<int8>(FMath::CeilLogTwo(InComponent->SubsectionSizeQuads + 1));
 	SharedBuffersKey = (SubsectionSizeLog2 & 0xf) | ((NumSubsections & 0xf) << 4) |	(XYOffsetmapTexture == nullptr ? 0 : 1 << 31);
 
 	bSupportsHeightfieldRepresentation = true;
@@ -1276,7 +1276,7 @@ void FLandscapeComponentSceneProxy::CreateRenderThreadResources()
 		LandscapeFixedGridUniformShaderParameters[LodIndex].InitResource();
 		FLandscapeFixedGridUniformShaderParameters Parameters;
 		Parameters.LodValues = FVector4f(
-			LodIndex,
+			static_cast<float>(LodIndex),
 			0.f,
 			(float)((SubsectionSizeVerts >> LodIndex) - 1),
 			1.f / (float)((SubsectionSizeVerts >> LodIndex) - 1));
@@ -1392,7 +1392,7 @@ FLandscapeRayTracingState* FLandscapeRayTracingImpl::FindOrCreateRayTracingState
 	{
 		for (int32 SubX = 0; SubX < NumSubsections; SubX++)
 		{
-			const int8 SubSectionIdx = SubX + SubY * NumSubsections;
+			const int8 SubSectionIdx = static_cast<int8>(SubX + SubY * NumSubsections);
 
 			FRayTracingGeometryInitializer Initializer;
 			static const FName DebugName("FLandscapeComponentSceneProxy");
@@ -1437,7 +1437,7 @@ FLandscapeRayTracingState::~FLandscapeRayTracingState()
 	{
 		for (int32 SubX = 0; SubX < NumSubsections; SubX++)
 		{
-			const int8 SubSectionIdx = SubX + SubY * NumSubsections;
+			const int8 SubSectionIdx = static_cast<int8>(SubX + SubY * NumSubsections);
 			Sections[SubSectionIdx].Geometry.ReleaseResource();
 			Sections[SubSectionIdx].RayTracingDynamicVertexBuffer.Release();
 		}
@@ -1794,16 +1794,16 @@ void FLandscapeComponentSceneProxy::OnTransformChanged()
 		LightmapBiasY,
 		LightmapBiasX);
 	LandscapeParams.SubsectionSizeVertsLayerUVPan = FVector4f(
-		SubsectionSizeVerts,
+		static_cast<float>(SubsectionSizeVerts),
 		1.f / (float)SubsectionSizeQuads,
-		SectionBase.X,
-		SectionBase.Y
+		static_cast<float>(SectionBase.X),
+		static_cast<float>(SectionBase.Y)
 	);
 	LandscapeParams.SubsectionOffsetParams = FVector4f(
 		HeightmapSubsectionOffsetU,
 		HeightmapSubsectionOffsetV,
 		WeightmapSubsectionOffset,
-		SubsectionSizeQuads
+		static_cast<float>(SubsectionSizeQuads)
 	);
 	LandscapeParams.LightmapSubsectionOffsetParams = FVector4f(
 		LightmapExtendFactorX,
@@ -1815,8 +1815,8 @@ void FLandscapeComponentSceneProxy::OnTransformChanged()
 	FTextureResource* HeightmapResource = HeightmapTexture ? HeightmapTexture->GetResource() : nullptr;
 	if (HeightmapResource)
 	{
-		const float SizeX = FMath::Max(HeightmapResource->GetSizeX(), 1u);
-		const float SizeY = FMath::Max(HeightmapResource->GetSizeY(), 1u);
+		const float SizeX = static_cast<float>(FMath::Max(HeightmapResource->GetSizeX(), 1u));
+		const float SizeY = static_cast<float>(FMath::Max(HeightmapResource->GetSizeY(), 1u));
 		LandscapeParams.HeightmapTextureSize = FVector4f(SizeX, SizeY, 1.f / SizeX, 1.f / SizeY);
 		LandscapeParams.HeightmapTexture = HeightmapTexture->TextureReference.TextureReferenceRHI;
 		LandscapeParams.HeightmapTextureSampler = TStaticSamplerState<SF_Point>::GetRHI();
@@ -1869,7 +1869,7 @@ bool FLandscapeComponentSceneProxy::GetMeshElementForVirtualTexture(int32 InLodI
 	OutMeshBatch.bUseForMaterial = false;
 	OutMeshBatch.Type = PT_TriangleList;
 	OutMeshBatch.DepthPriorityGroup = SDPG_World;
-	OutMeshBatch.LODIndex = InLodIndex;
+	OutMeshBatch.LODIndex = static_cast<int8>(InLodIndex);
 	OutMeshBatch.bDitheredLODTransition = false;
 	OutMeshBatch.bRenderToVirtualTexture = true;
 	OutMeshBatch.RuntimeVirtualTextureMaterialType = (uint32)MaterialType;
@@ -1947,7 +1947,7 @@ bool FLandscapeComponentSceneProxy::GetStaticMeshElement(int32 LODIndex, bool bF
 		MeshBatch.bUseForMaterial = true;
 		MeshBatch.Type = PT_TriangleList;
 		MeshBatch.DepthPriorityGroup = SDPG_World;
-		MeshBatch.LODIndex = LODIndex;
+		MeshBatch.LODIndex = static_cast<int8>(LODIndex);
 		MeshBatch.bDitheredLODTransition = false;
 
 		const FLandscapeRenderSystem& RenderSystem = *LandscapeRenderSystems.FindChecked(LandscapeKey);
@@ -2081,7 +2081,7 @@ void FLandscapeComponentSceneProxy::GetDynamicMeshElements(const TArray<const FS
 
 			const FSceneView* View = Views[ViewIndex];
 
-			int32 LODToRender = RenderSystem.GetSectionLODValue(*View, ComponentBase);
+			int32 LODToRender = static_cast<int32>(RenderSystem.GetSectionLODValue(*View, ComponentBase));
 
 			FMeshBatch& Mesh = Collector.AllocateMesh();
 			GetStaticMeshElement(LODToRender, false, Mesh, ParameterArray.ElementParams);
@@ -2174,7 +2174,7 @@ void FLandscapeComponentSceneProxy::GetDynamicMeshElements(const TArray<const FS
 					LODMesh.Elements.Empty(1);
 					LODMesh.Elements.Add(TemplateMesh.Elements[i]);
 					int32 CurrentLOD = ((FLandscapeBatchElementParams*)TemplateMesh.Elements[i].UserData)->CurrentLOD;
-					LODMesh.VisualizeLODIndex = CurrentLOD;
+					LODMesh.VisualizeLODIndex = static_cast<int8>(CurrentLOD);
 					FLinearColor Color = GetColorForLod(CurrentLOD, LODSettings.ForcedLOD, true);
 					FMaterialRenderProxy* LODMaterialProxy = (FMaterialRenderProxy*)new FColoredMaterialRenderProxy(GEngine->LevelColorationUnlitMaterial->GetRenderProxy(), Color);
 					Collector.RegisterOneFrameMaterialProxy(LODMaterialProxy);
@@ -2493,7 +2493,7 @@ void FLandscapeComponentSceneProxy::GetDynamicRayTracingInstances(FRayTracingMat
 	}
 	FLandscapeRayTracingState* RayTracingState = RayTracingImpl.Get()->FindOrCreateRayTracingState(SceneView.State, NumSubsections, SubsectionSizeVerts);
 
-	int32 LODToRender = RenderSystem.GetSectionLODValue(SceneView, ComponentBase);
+	int32 LODToRender = static_cast<int32>(RenderSystem.GetSectionLODValue(SceneView, ComponentBase));
 
 	FLandscapeElementParamArray& ParameterArray = Context.RayTracingMeshResourceCollector.AllocateOneFrameResource<FLandscapeElementParamArray>();
 	ParameterArray.ElementParams.AddDefaulted(NumSubsections * NumSubsections);
@@ -2503,7 +2503,7 @@ void FLandscapeComponentSceneProxy::GetDynamicRayTracingInstances(FRayTracingMat
 		return;
 	}
 
-	const int8 CurrentLODIndex = LODToRender;
+	const int8 CurrentLODIndex = static_cast<int8>(LODToRender);
 	int8 MaterialIndex = LODIndexToMaterialIndex.IsValidIndex(CurrentLODIndex) ? LODIndexToMaterialIndex[CurrentLODIndex] : INDEX_NONE;
 	FMaterialRenderProxy* SelectedMaterial = MaterialIndex != INDEX_NONE ? AvailableMaterials[MaterialIndex] : nullptr;
 
@@ -2528,8 +2528,8 @@ void FLandscapeComponentSceneProxy::GetDynamicRayTracingInstances(FRayTracingMat
 	{
 		for (int32 SubX = 0; SubX < NumSubsections; SubX++)
 		{
-			const int8 SubSectionIdx = SubX + SubY * NumSubsections;
-			const int8 CurrentLOD = LODToRender;
+			const int8 SubSectionIdx = static_cast<int8>(SubX + SubY * NumSubsections);
+			const int8 CurrentLOD = static_cast<int8>(LODToRender);
 
 			FMeshBatch MeshBatch = BaseMeshBatch;
 
@@ -2672,10 +2672,10 @@ void FLandscapeVertexBuffer::InitRHI()
 			{
 				for (int32 x = 0; x < SubsectionSizeVerts; x++)
 				{
-					Vertex->VertexX = x;
-					Vertex->VertexY = y;
-					Vertex->SubX = SubX;
-					Vertex->SubY = SubY;
+					Vertex->VertexX = static_cast<uint8>(x);
+					Vertex->VertexY = static_cast<uint8>(y);
+					Vertex->SubX = static_cast<uint8>(SubX);
+					Vertex->SubY = static_cast<uint8>(SubY);
 					Vertex++;
 					VertexIndex++;
 				}
@@ -2729,10 +2729,10 @@ void FLandscapeSharedBuffers::CreateIndexBuffers(const FName& InOwnerName)
 					{
 						for (int32 x = 0; x < LodSubsectionSizeQuads; x++)
 						{
-							INDEX_TYPE i00 = (x + 0) + (y + 0) * SubsectionSizeVerts + SubOffset;
-							INDEX_TYPE i10 = (x + 1) + (y + 0) * SubsectionSizeVerts + SubOffset;
-							INDEX_TYPE i11 = (x + 1) + (y + 1) * SubsectionSizeVerts + SubOffset;
-							INDEX_TYPE i01 = (x + 0) + (y + 1) * SubsectionSizeVerts + SubOffset;
+							INDEX_TYPE i00 = static_cast<uint16>((x + 0) + (y + 0) * SubsectionSizeVerts + SubOffset);
+							INDEX_TYPE i10 = static_cast<uint16>((x + 1) + (y + 0) * SubsectionSizeVerts + SubOffset);
+							INDEX_TYPE i11 = static_cast<uint16>((x + 1) + (y + 1) * SubsectionSizeVerts + SubOffset);
+							INDEX_TYPE i01 = static_cast<uint16>((x + 0) + (y + 1) * SubsectionSizeVerts + SubOffset);
 
 							NewIndices.Add(i00);
 							NewIndices.Add(i11);
@@ -2787,10 +2787,10 @@ void FLandscapeSharedBuffers::CreateIndexBuffers(const FName& InOwnerName)
 			{
 				for (int32 x = 0; x < LodSubsectionSizeQuads; x++)
 				{
-					INDEX_TYPE i00 = (x + 0) + (y + 0) * (SubsectionSizeVerts >> Mip);
-					INDEX_TYPE i10 = (x + 1) + (y + 0) * (SubsectionSizeVerts >> Mip);
-					INDEX_TYPE i11 = (x + 1) + (y + 1) * (SubsectionSizeVerts >> Mip);
-					INDEX_TYPE i01 = (x + 0) + (y + 1) * (SubsectionSizeVerts >> Mip);
+					INDEX_TYPE i00 = static_cast<uint16>((x + 0) + (y + 0) * (SubsectionSizeVerts >> Mip));
+					INDEX_TYPE i10 = static_cast<uint16>((x + 1) + (y + 0) * (SubsectionSizeVerts >> Mip));
+					INDEX_TYPE i11 = static_cast<uint16>((x + 1) + (y + 1) * (SubsectionSizeVerts >> Mip));
+					INDEX_TYPE i01 = static_cast<uint16>((x + 0) + (y + 1) * (SubsectionSizeVerts >> Mip));
 
 					ZeroOffsetIndices.Add(i00);
 					ZeroOffsetIndices.Add(i11);
@@ -2838,7 +2838,7 @@ void FLandscapeSharedBuffers::CreateGrassIndexBuffer(const FName& InOwnerName)
 					for (int32 x = 0; x < MipSubsectionSizeVerts; x++)
 					{
 						// intentionally using SubsectionSizeVerts not MipSubsectionSizeVerts, this is a vert buffer index not a mip vert index
-						NewIndices.Add(x + y * SubsectionSizeVerts + SubOffset);
+						NewIndices.Add(static_cast<uint16>(x + y * SubsectionSizeVerts + SubOffset));
 					}
 				}
 
@@ -3686,7 +3686,7 @@ void ULandscapeComponent::GetStreamingRenderAssetInfo(FStreamingTextureLevelCont
 			ScaleFactor = FMath::Abs(ProxyRootComponent->GetRelativeScale3D().X);
 		}
 		LocalStreamingDistanceMultiplier = FMath::Max(0.0f, Proxy->StreamingDistanceMultiplier);
-		TexelFactor = 0.75f * LocalStreamingDistanceMultiplier * ComponentSizeQuads * ScaleFactor;
+		TexelFactor = static_cast<float>(0.75f * LocalStreamingDistanceMultiplier * ComponentSizeQuads * ScaleFactor);
 	}
 
 	ERHIFeatureLevel::Type FeatureLevel = LevelContext.GetFeatureLevel();
@@ -3734,7 +3734,7 @@ void ULandscapeComponent::GetStreamingRenderAssetInfo(FStreamingTextureLevelCont
 				const FVector2D& Scale = Lightmap->GetCoordinateScale();
 				if (Scale.X > SMALL_NUMBER && Scale.Y > SMALL_NUMBER)
 				{
-					const float LightmapTexelFactor = TexelFactor / FMath::Min(Scale.X, Scale.Y);
+					const float LightmapTexelFactor = static_cast<float>(TexelFactor / FMath::Min(Scale.X, Scale.Y));
 					new (OutStreamingRenderAssets) FStreamingRenderAssetPrimitiveInfo(Lightmap->GetTexture(LightmapIndex), Bounds, LightmapTexelFactor);
 					new (OutStreamingRenderAssets) FStreamingRenderAssetPrimitiveInfo(Lightmap->GetAOMaterialMaskTexture(), Bounds, LightmapTexelFactor);
 					new (OutStreamingRenderAssets) FStreamingRenderAssetPrimitiveInfo(Lightmap->GetSkyOcclusionTexture(), Bounds, LightmapTexelFactor);
@@ -3748,7 +3748,7 @@ void ULandscapeComponent::GetStreamingRenderAssetInfo(FStreamingTextureLevelCont
 				const FVector2D& Scale = Shadowmap->GetCoordinateScale();
 				if (Scale.X > SMALL_NUMBER && Scale.Y > SMALL_NUMBER)
 				{
-					const float ShadowmapTexelFactor = TexelFactor / FMath::Min(Scale.X, Scale.Y);
+					const float ShadowmapTexelFactor = static_cast<float>(TexelFactor / FMath::Min(Scale.X, Scale.Y));
 					new (OutStreamingRenderAssets) FStreamingRenderAssetPrimitiveInfo(Shadowmap->GetTexture(), Bounds, ShadowmapTexelFactor);
 				}
 			}
@@ -3935,10 +3935,11 @@ float FLandscapeComponentSceneProxy::ComputeLODForView(const FSceneView& InView)
 	ForcedLODLevel = DrawCollisionLODOverride >= 0 ? DrawCollisionLODOverride : ForcedLODLevel;
 	ForcedLODLevel = FMath::Min<int32>(ForcedLODLevel, LODSettings.LastLODIndex);
 
-	float LODLevel = ForcedLODLevel;
+	float LODLevel = static_cast<float>(ForcedLODLevel);
 	if (ForcedLODLevel < 0)
 	{
-		float MeshScreenSizeSquared = ComputeBoundsScreenRadiusSquared(GetBounds().Origin, GetBounds().SphereRadius, ViewOrigin, ViewProjectionMatrix);
+		float MeshScreenSizeSquared = ComputeBoundsScreenRadiusSquared(GetBounds().Origin, static_cast<float>(GetBounds().SphereRadius), ViewOrigin, 
+		ViewProjectionMatrix);
 
 		float FractionalLOD;
 		FLandscapeRenderSystem::GetLODFromScreenSize(LODSettings, MeshScreenSizeSquared, LODScale * LODScale, FractionalLOD);
@@ -3956,7 +3957,7 @@ float FLandscapeComponentSceneProxy::ComputeLODBias() const
 	{
 		if (const FTexture2DResource* TextureResource = (const FTexture2DResource*)HeightmapTexture->GetResource())
 		{
-			ComputedLODBias = TextureResource->GetCurrentFirstMip();
+			ComputedLODBias = static_cast<float>(TextureResource->GetCurrentFirstMip());
 		}
 	}
 
diff --git a/Engine/Source/Runtime/Landscape/Public/LandscapeDataAccess.h b/Engine/Source/Runtime/Landscape/Public/LandscapeDataAccess.h
--- a/Engine/Source/Runtime/Landscape/Public/LandscapeDataAccess.h
+++ b/Engine/Source/Runtime/Landscape/Public/LandscapeDataAccess.h
@@ -21,8 +21,8 @@ class ULandscapeLayerInfoObject;
 
 namespace LandscapeDataAccess
 {
-	inline const int32 MaxValue = 65535;
-	inline const float MidValue = 32768.f;
+	inline constexpr int32 MaxValue = 65535;
+	inline constexpr float MidValue = 32768.f;
 	// Reserved 2 bits for other purpose
 	// Most significant bit - Visibility, 0 is visible(default), 1 is invisible
 	// 2nd significant bit - Triangle flip, not implemented yet
