Fixed issues with a runtime toggling of landscape culling feature

diff --git a/Engine/Source/Runtime/Landscape/Private/LandscapeCulling.cpp b/Engine/Source/Runtime/Landscape/Private/LandscapeCulling.cpp
--- a/Engine/Source/Runtime/Landscape/Private/LandscapeCulling.cpp
+++ b/Engine/Source/Runtime/Landscape/Private/LandscapeCulling.cpp
@@ -416,8 +416,6 @@ struct FCullingEntry
 
 struct FCullingSystem
 {
-	FLandscapeTileMesh* TileMesh = nullptr;
-	
 	TArray<FCullingEntry> Landscapes; 
 	TArray<uint32, TInlineAllocator<2>> ViewStateKeys;
 };
@@ -451,46 +449,23 @@ void PreRenderViewFamily(FSceneViewFamily& InViewFamily)
 
 void InitSharedBuffers(FLandscapeSharedBuffers& SharedBuffers, const ERHIFeatureLevel::Type InFeatureLevel)
 {
-	if (!EnableGPUCullingForSection(SharedBuffers.SubsectionSizeVerts, SharedBuffers.NumSubsections))
-	{
-		SharedBuffers.TileVertexFactory = nullptr;
-		SharedBuffers.TileDataBuffer = nullptr;
-		return;
-	}
-
-	if (GCullingSystem.TileMesh == nullptr)
+	if (SharedBuffers.TileVertexFactory == nullptr)
 	{
-		GCullingSystem.TileMesh = new FLandscapeTileMesh();
-	}
-
-	uint32 SubsectionSizeQuads = (SharedBuffers.SubsectionSizeVerts - 1);
-
-	FLandscapeTileMesh* TileMesh = GCullingSystem.TileMesh;
-	FLandscapeTileDataBuffer* TileDataBuffer = new FLandscapeTileDataBuffer(SharedBuffers.NumSubsections, SubsectionSizeQuads);
-	FLandscapeTileVertexFactory* TileVF = new FLandscapeTileVertexFactory(InFeatureLevel);
+		uint32 SubsectionSizeQuads = (SharedBuffers.SubsectionSizeVerts - 1);
 
-	uint32 TileDataStride = TILE_DATA_ENTRY_NUM_BYTES;
+		FLandscapeTileMesh* TileMesh = new FLandscapeTileMesh();
+		FLandscapeTileDataBuffer* TileDataBuffer = new FLandscapeTileDataBuffer(SharedBuffers.NumSubsections, SubsectionSizeQuads);
+		FLandscapeTileVertexFactory* TileVF = new FLandscapeTileVertexFactory(InFeatureLevel);
 
-	TileVF->Data.PositionComponent = FVertexStreamComponent(&TileMesh->VertexBuffer, 0, sizeof(FLandscapeVertex), VET_UByte4);
-	TileVF->Data.TileDataComponent = FVertexStreamComponent(TileDataBuffer, 0, TileDataStride, VET_UByte4, EVertexStreamUsage::Instancing);
-	TileVF->InitResource();
+		uint32 TileDataStride = TILE_DATA_ENTRY_NUM_BYTES;
 
-	SharedBuffers.TileVertexFactory = TileVF;
-	SharedBuffers.TileDataBuffer = TileDataBuffer;
-}
+		TileVF->Data.PositionComponent = FVertexStreamComponent(&TileMesh->VertexBuffer, 0, sizeof(FLandscapeVertex), VET_UByte4);
+		TileVF->Data.TileDataComponent = FVertexStreamComponent(TileDataBuffer, 0, TileDataStride, VET_UByte4, EVertexStreamUsage::Instancing);
+		TileVF->InitResource();
 
-void DeleteSharedBuffers(FLandscapeSharedBuffers& SharedBuffers)
-{
-	delete SharedBuffers.TileDataBuffer;
-	SharedBuffers.TileDataBuffer = nullptr;
-	
-	delete SharedBuffers.TileVertexFactory;
-	SharedBuffers.TileVertexFactory = nullptr;
-	
-	if (GCullingSystem.Landscapes.Num() == 0)
-	{
-		delete GCullingSystem.TileMesh;
-		GCullingSystem.TileMesh = nullptr;
+		SharedBuffers.TileMesh = TileMesh;
+		SharedBuffers.TileVertexFactory = TileVF;
+		SharedBuffers.TileDataBuffer = TileDataBuffer;
 	}
 }
 
@@ -511,7 +486,7 @@ void SetupMeshBatch(const FLandscapeSharedBuffers& SharedBuffers, FMeshBatch& Me
 		MeshBatch.VertexFactory = SharedBuffers.TileVertexFactory;
 
 		FMeshBatchElement& BatchElement = MeshBatch.Elements[0];
-		BatchElement.IndexBuffer = &GCullingSystem.TileMesh->IndexBuffer;
+		BatchElement.IndexBuffer = &static_cast<FLandscapeTileMesh*>(SharedBuffers.TileMesh)->IndexBuffer;
 		BatchElement.NumPrimitives = FMath::Square(LANDSCAPE_TILE_QUADS) * 2u;
 		BatchElement.NumInstances = FMath::Square(FMath::DivideAndRoundUp(SubsectionSizeQuads, LANDSCAPE_TILE_QUADS)) * FMath::Square(SharedBuffers.NumSubsections);
 		BatchElement.FirstIndex = 0;
@@ -520,13 +495,15 @@ void SetupMeshBatch(const FLandscapeSharedBuffers& SharedBuffers, FMeshBatch& Me
 	}
 }
 
-void RegisterLandscape(uint32 LandscapeKey, int32 SubsectionSizeVerts, int32 NumSubsections)
+void RegisterLandscape(FLandscapeSharedBuffers& SharedBuffers, ERHIFeatureLevel::Type FeatureLevel, uint32 LandscapeKey, int32 SubsectionSizeVerts, int32 NumSubsections)
 {
 	if (!EnableGPUCullingForSection(SubsectionSizeVerts, NumSubsections))
 	{
 		return;
 	}
 
+	InitSharedBuffers(SharedBuffers, FeatureLevel);
+
 	FCullingEntry* CullingEntryPtr = GCullingSystem.Landscapes.FindByPredicate([LandscapeKey](const FCullingEntry& Other) {
 		return Other.LandscapeKey == LandscapeKey;
 	});
diff --git a/Engine/Source/Runtime/Landscape/Private/LandscapeCulling.h b/Engine/Source/Runtime/Landscape/Private/LandscapeCulling.h
--- a/Engine/Source/Runtime/Landscape/Private/LandscapeCulling.h
+++ b/Engine/Source/Runtime/Landscape/Private/LandscapeCulling.h
@@ -38,10 +38,8 @@ namespace UE {
 			
 			FVertexFactoryType* GetTileVertexFactoryType();
 
-			void InitSharedBuffers(FLandscapeSharedBuffers& SharedBuffers, ERHIFeatureLevel::Type FeatureLevel);
-			void DeleteSharedBuffers(FLandscapeSharedBuffers& SharedBuffers);
 			void SetupMeshBatch(const FLandscapeSharedBuffers& SharedBuffers, FMeshBatch& MeshBatch);
-			void RegisterLandscape(uint32 LandscapeKey, int32 SubsectionSizeVerts, int32 NumSubsections);
+			void RegisterLandscape(FLandscapeSharedBuffers& SharedBuffers, ERHIFeatureLevel::Type FeatureLevel, uint32 LandscapeKey, int32 SubsectionSizeVerts, int32 NumSubsections);
 			void UnregisterLandscape(uint32 LandscapeKey);
 
 			void PreRenderViewFamily(FSceneViewFamily& InViewFamily);
diff --git a/Engine/Source/Runtime/Landscape/Private/LandscapeRender.cpp b/Engine/Source/Runtime/Landscape/Private/LandscapeRender.cpp
--- a/Engine/Source/Runtime/Landscape/Private/LandscapeRender.cpp
+++ b/Engine/Source/Runtime/Landscape/Private/LandscapeRender.cpp
@@ -1246,11 +1246,6 @@ void FLandscapeComponentSceneProxy::CreateRenderThreadResources()
 			LandscapeVertexFactory->InitResource();
 			SharedBuffers->FixedGridVertexFactory = LandscapeVertexFactory;
 		}
-
-		if (Culling::UseCulling(GetScene().GetShaderPlatform()))
-		{
-			Culling::InitSharedBuffers(*SharedBuffers, FeatureLevel);
-		}
 	}
 
 	SharedBuffers->AddRef();
@@ -1261,7 +1256,7 @@ void FLandscapeComponentSceneProxy::CreateRenderThreadResources()
 
 	if (!bNaniteActive && Culling::UseCulling(GetScene().GetShaderPlatform()))
 	{
-		Culling::RegisterLandscape(LandscapeKey, SubsectionSizeVerts, NumSubsections);
+		Culling::RegisterLandscape(*SharedBuffers, FeatureLevel, LandscapeKey, SubsectionSizeVerts, NumSubsections);
 	}
 	
 	// Assign LandscapeUniformShaderParameters
@@ -1450,6 +1445,7 @@ void FLandscapeComponentSceneProxy::DestroyRenderThreadResources()
 	FPrimitiveSceneProxy::DestroyRenderThreadResources();
 	FLandscapeRenderSystem::UnregisterSection(this);
 	FLandscapeRenderSystem::DestroyResources(this);
+	Culling::UnregisterLandscape(LandscapeKey);
 }
 
 bool FLandscapeComponentSceneProxy::OnLevelAddedToWorld_RenderThread()
@@ -1476,12 +1472,7 @@ FLandscapeComponentSceneProxy::~FLandscapeComponentSceneProxy()
 {
 	// Free the subsection uniform buffer
 	LandscapeUniformShaderParameters.ReleaseResource();
-
-	if (Culling::UseCulling(GetScene().GetShaderPlatform()))
-	{
-		Culling::UnregisterLandscape(LandscapeKey);
-	}
-	
+		
 	// Free the lod uniform buffers
 	for (int32 i = 0; i < LandscapeFixedGridUniformShaderParameters.Num(); ++i)
 	{
@@ -2864,6 +2855,7 @@ FLandscapeSharedBuffers::FLandscapeSharedBuffers(const int32 InSharedBuffersKey,
 	, NumIndexBuffers(FMath::CeilLogTwo(InSubsectionSizeQuads + 1))
 	, SubsectionSizeVerts(InSubsectionSizeQuads + 1)
 	, NumSubsections(InNumSubsections)
+	, TileMesh(nullptr)
 	, VertexFactory(nullptr)
 	, FixedGridVertexFactory(nullptr)
 	, VertexBuffer(nullptr)
@@ -2916,8 +2908,6 @@ FLandscapeSharedBuffers::FLandscapeSharedBuffers(const int32 InSharedBuffersKey,
 
 FLandscapeSharedBuffers::~FLandscapeSharedBuffers()
 {
-	Culling::DeleteSharedBuffers(*this);
-	
 	delete VertexBuffer;
 
 	for (int32 i = 0; i < NumIndexBuffers; i++)
@@ -2949,6 +2939,10 @@ FLandscapeSharedBuffers::~FLandscapeSharedBuffers()
 #endif
 
 	delete VertexFactory;
+
+	delete TileMesh;
+	delete TileDataBuffer;
+	delete TileVertexFactory;
 }
 
 //
diff --git a/Engine/Source/Runtime/Landscape/Public/LandscapeRender.h b/Engine/Source/Runtime/Landscape/Public/LandscapeRender.h
--- a/Engine/Source/Runtime/Landscape/Public/LandscapeRender.h
+++ b/Engine/Source/Runtime/Landscape/Public/LandscapeRender.h
@@ -358,6 +358,7 @@ public:
 	FLandscapeVertexFactory* FixedGridVertexFactory;
 	FLandscapeVertexBuffer* VertexBuffer;
 	
+	FRenderResource* TileMesh;
 	FLandscapeVertexFactory* TileVertexFactory;
 	FVertexBuffer* TileDataBuffer;
 	
