Revert GPU Landscape section culling

diff --git a/Engine/Config/Android/AndroidEngine.ini b/Engine/Config/Android/AndroidEngine.ini
--- a/Engine/Config/Android/AndroidEngine.ini
+++ b/Engine/Config/Android/AndroidEngine.ini
@@ -95,8 +95,6 @@ r.Vulkan.MemoryBacktrace=0
 [/Script/Engine.RendererSettings]
 r.DistanceFields=0
 r.TemporalAA.Mobile.UseCompute=0
-; Enable GPU culling for landscape
-landscape.SupportGPUCulling=1
 
 [AlternateTextureCompression]
 ;disabling TextureCompressionFormat for Android/IOS no longer needed
diff --git a/Engine/Config/IOS/IOSEngine.ini b/Engine/Config/IOS/IOSEngine.ini
--- a/Engine/Config/IOS/IOSEngine.ini
+++ b/Engine/Config/IOS/IOSEngine.ini
@@ -85,8 +85,6 @@ fx.Niagara.QualityLevel.Max=1
 [/Script/Engine.RendererSettings]
 r.DistanceFields=0
 r.TemporalAA.Mobile.UseCompute=0
-; Enable GPU culling for landscape
-landscape.SupportGPUCulling=1
 
 [AlternateTextureCompression]
 ;disabling TextureCompressionFormat for Android/IOS no longer needed
diff --git a/Engine/Shaders/Private/Landscape/LandscapeCulling.usf b/Engine/Shaders/Private/Landscape/LandscapeCulling.usf
deleted file mode 100644
diff --git a/Engine/Shaders/Private/LandscapeVertexFactory.ush b/Engine/Shaders/Private/LandscapeVertexFactory.ush
--- a/Engine/Shaders/Private/LandscapeVertexFactory.ush
+++ b/Engine/Shaders/Private/LandscapeVertexFactory.ush
@@ -151,17 +151,10 @@ float CalcLOD(uint ComponentIndex, float2 xyLocalToSubsection, float2 Subsection
 
 struct FVertexFactoryInput
 {
-	// UByte4
-	uint4 Position : ATTRIBUTE0;
-
-#if LANDSCAPE_TILE
-	// UByte4
-	uint4 TileData : ATTRIBUTE1;
-#else
+	float4 Position : ATTRIBUTE0;
 	// Dynamic instancing related attributes with InstanceIdOffset : ATTRIBUTE1
 	VF_GPUSCENE_DECLARE_INPUT_BLOCK(1)
 	VF_INSTANCED_STEREO_DECLARE_INPUT_BLOCK()
-#endif	
 };
 
 // RHI_RAYTRACING
@@ -594,11 +587,6 @@ FVertexFactoryIntermediates GetVertexFactoryIntermediates(FVertexFactoryInput In
 	Intermediates.LodBias = GetLodBias(Intermediates.ComponentIndex);
 	Intermediates.InputPosition = Input.Position;
 
-#if LANDSCAPE_TILE   
-	Intermediates.InputPosition = Intermediates.InputPosition + Input.TileData;
-	Intermediates.InputPosition.xy = min(Intermediates.InputPosition.xy, float2(Intermediates.LodValues.z, Intermediates.LodValues.z));
-#endif
-
 	float2 xy = Intermediates.InputPosition.xy * Intermediates.LodValues.w;
 
 	float LODCalculated = CalcLOD(Intermediates.ComponentIndex, xy, Intermediates.InputPosition.zw);
diff --git a/Engine/Source/Runtime/Engine/Private/PrimitiveSceneProxy.cpp b/Engine/Source/Runtime/Engine/Private/PrimitiveSceneProxy.cpp
--- a/Engine/Source/Runtime/Engine/Private/PrimitiveSceneProxy.cpp
+++ b/Engine/Source/Runtime/Engine/Private/PrimitiveSceneProxy.cpp
@@ -112,9 +112,6 @@ bool SupportsCachingMeshDrawCommands(const FMeshBatch& MeshBatch)
 		// Cached mesh commands only allow for a single mesh element per batch.
 		(MeshBatch.Elements.Num() == 1) &&
 
-		// View dependent arguments can't be cached
-		(MeshBatch.bViewDependentArguments == 0) &&
-
 		// Vertex factory needs to support caching.
 		MeshBatch.VertexFactory->GetType()->SupportsCachingMeshDrawCommands();
 }
diff --git a/Engine/Source/Runtime/Engine/Private/SceneManagement.cpp b/Engine/Source/Runtime/Engine/Private/SceneManagement.cpp
--- a/Engine/Source/Runtime/Engine/Private/SceneManagement.cpp
+++ b/Engine/Source/Runtime/Engine/Private/SceneManagement.cpp
@@ -1432,17 +1432,15 @@ bool FMeshBatch::Validate(const FPrimitiveSceneProxy* PrimitiveSceneProxy, ERHIF
 	}
 
 	const bool bVFSupportsPrimitiveIdStream = VertexFactory->GetType()->SupportsPrimitiveIdStream();
-	const bool bVFRequiresPrimitiveUniformBuffer = PrimitiveSceneProxy->DoesVFRequirePrimitiveUniformBuffer();
 
-	if (!bVFRequiresPrimitiveUniformBuffer && !bVFSupportsPrimitiveIdStream)
+	if (!PrimitiveSceneProxy->DoesVFRequirePrimitiveUniformBuffer() && !bVFSupportsPrimitiveIdStream)
 	{
 		return LogMeshError(TEXT("PrimitiveSceneProxy has bVFRequiresPrimitiveUniformBuffer disabled yet tried to draw with a vertex factory that did not support PrimitiveIdStream"));
 	}
 
-	// Some primitives may use several VFs with a mixed support for a GPUScene
-	if (PrimitiveSceneProxy->SupportsGPUScene() && !(VertexFactory->SupportsGPUScene(FeatureLevel) || bVFRequiresPrimitiveUniformBuffer))
+	if (PrimitiveSceneProxy->SupportsGPUScene() && !VertexFactory->SupportsGPUScene(FeatureLevel))
 	{
-		return LogMeshError(TEXT("PrimitiveSceneProxy has SupportsGPUScene() does not match VertexFactory->SupportsGPUScene() or bVFRequiresPrimitiveUniformBuffer"));
+		return LogMeshError(TEXT("PrimitiveSceneProxy has SupportsGPUScene() does not match VertexFactory->SupportsGPUScene()"));
 	}
 	const bool bUseGPUScene = UseGPUScene(GMaxRHIShaderPlatform, FeatureLevel);
 	
@@ -1454,9 +1452,7 @@ bool FMeshBatch::Validate(const FPrimitiveSceneProxy* PrimitiveSceneProxy, ERHIF
 	{
 		const FMeshBatchElement& MeshElement = Elements[ElementIndex];
 
-		// Some primitives may use several VFs with a mixed support for a GPUScene 
-		// in this case all mesh batches get Primitive UB assigned regardless of VF type
-		if (bPrimitiveShaderDataComesFromSceneBuffer && Elements[ElementIndex].PrimitiveUniformBuffer && !bVFRequiresPrimitiveUniformBuffer)
+		if (bPrimitiveShaderDataComesFromSceneBuffer && Elements[ElementIndex].PrimitiveUniformBuffer)
 		{
 			// on mobile VS has access to PrimitiveUniformBuffer
 			if (FeatureLevel > ERHIFeatureLevel::ES3_1)
diff --git a/Engine/Source/Runtime/Engine/Public/MeshBatch.h b/Engine/Source/Runtime/Engine/Public/MeshBatch.h
--- a/Engine/Source/Runtime/Engine/Public/MeshBatch.h
+++ b/Engine/Source/Runtime/Engine/Public/MeshBatch.h
@@ -397,12 +397,6 @@ struct FMeshBatch
 	uint32 CastRayTracedShadow : 1;	// Whether it casts ray traced shadow.
 #endif
 
-	/** 
-	 * Whether mesh has a view dependent draw arguments. 
-	 * Gives an opportunity to override mesh arguments for a View just before creating MDC for a mesh (makes MDC "non-cached")
-	 */
-	uint32 bViewDependentArguments : 1;
-
 #if UE_ENABLE_DEBUG_DRAWING
 	/** Conceptual HLOD index used for the HLOD Coloration visualization. */
 	int8 VisualizeHLODIndex;
@@ -481,7 +475,6 @@ struct FMeshBatch
 #if RHI_RAYTRACING
 	,	CastRayTracedShadow(true)
 #endif
-	,	bViewDependentArguments(false)
 #if (!(UE_BUILD_SHIPPING || UE_BUILD_TEST) || WITH_EDITOR)
 	,	VisualizeHLODIndex(INDEX_NONE)
 #endif
diff --git a/Engine/Source/Runtime/Engine/Public/PrimitiveSceneProxy.h b/Engine/Source/Runtime/Engine/Public/PrimitiveSceneProxy.h
--- a/Engine/Source/Runtime/Engine/Public/PrimitiveSceneProxy.h
+++ b/Engine/Source/Runtime/Engine/Public/PrimitiveSceneProxy.h
@@ -324,14 +324,6 @@ public:
 
 	virtual const class FCardRepresentationData* GetMeshCardRepresentation() const { return nullptr; }
 
-	/** 
-	* Gives the primitive an opportunity to override MeshBatch arguments for a specific View
-	* Only called for a MeshBatch with a bViewDependentArguments property set
-	* @param View - the view to override for
-	* @param ViewDependentMeshBatch - view dependent mesh copy (does not affect a cached FMeshBatch)
-	*/
-	virtual void ApplyViewDependentMeshArguments(const FSceneView& View, FMeshBatch& ViewDependentMeshBatch) const {};
-
 	/** 
 	 * Gets the boxes for sub occlusion queries
 	 * @param View - the view the occlusion results are for
diff --git a/Engine/Source/Runtime/Landscape/Private/LandscapeCulling.cpp b/Engine/Source/Runtime/Landscape/Private/LandscapeCulling.cpp
deleted file mode 100644
diff --git a/Engine/Source/Runtime/Landscape/Private/LandscapeCulling.h b/Engine/Source/Runtime/Landscape/Private/LandscapeCulling.h
deleted file mode 100644
diff --git a/Engine/Source/Runtime/Landscape/Private/LandscapeRender.cpp b/Engine/Source/Runtime/Landscape/Private/LandscapeRender.cpp
--- a/Engine/Source/Runtime/Landscape/Private/LandscapeRender.cpp
+++ b/Engine/Source/Runtime/Landscape/Private/LandscapeRender.cpp
@@ -52,13 +52,9 @@ LandscapeRender.cpp: New terrain rendering
 #include "RenderCore.h"
 #include "DataDrivenShaderPlatformInfo.h"
 #include "Algo/Transform.h"
-#include "LandscapeCulling.h"
-
-using namespace UE::Landscape;
 
 IMPLEMENT_GLOBAL_SHADER_PARAMETER_STRUCT(FLandscapeUniformShaderParameters, "LandscapeParameters");
 IMPLEMENT_GLOBAL_SHADER_PARAMETER_STRUCT(FLandscapeFixedGridUniformShaderParameters, "LandscapeFixedGrid");
-IMPLEMENT_TYPE_LAYOUT(FLandscapeVertexFactoryVertexShaderParameters);
 IMPLEMENT_TYPE_LAYOUT(FLandscapeVertexFactoryPixelShaderParameters);
 
 #if !UE_BUILD_SHIPPING
@@ -738,11 +734,6 @@ void FLandscapeSceneViewExtension::BeginRenderViewFamily(FSceneViewFamily& InVie
 	}
 }
 
-void FLandscapeSceneViewExtension::PreRenderViewFamily_RenderThread(FRDGBuilder& GraphBuilder, FSceneViewFamily& InViewFamily)
-{
-	Culling::PreRenderViewFamily(InViewFamily);
-}
-
 void FLandscapeSceneViewExtension::PreRenderView_RenderThread(FRDGBuilder& GraphBuilder, FSceneView& InView)
 {
 	LandscapeViews.Emplace(InView);
@@ -812,8 +803,6 @@ void FLandscapeSceneViewExtension::PreInitViews_RenderThread(FRDGBuilder& GraphB
 		RenderSystem.UpdateBuffers();
 	}
 
-	TArray<const FSceneView*, TInlineAllocator<2>> SceneViews;
-
 	for (FLandscapeViewData& LandscapeView : LandscapeViews)
 	{
 		if (!LandscapeRenderSystems.IsEmpty())
@@ -831,15 +820,8 @@ void FLandscapeSceneViewExtension::PreInitViews_RenderThread(FRDGBuilder& GraphB
 			LandscapeView.View->LandscapePerComponentDataBuffer = GWhiteVertexBufferWithSRV->ShaderResourceViewRHI;
 			LandscapeView.View->LandscapeIndirectionBuffer = GWhiteVertexBufferWithSRV->ShaderResourceViewRHI;
 		}
-
-		SceneViews.Add(LandscapeView.View);
 	}
 
-	if (Culling::UseCulling(SceneViews[0]->GetShaderPlatform()))
-	{
-		Culling::InitMainViews(GraphBuilder, SceneViews);
-	}
-		
 	LandscapeViews.Reset();
 }
 
@@ -1173,9 +1155,6 @@ FLandscapeComponentSceneProxy::FLandscapeComponentSceneProxy(ULandscapeComponent
 		}
 	}
 
-	// Landscape GPU culling uses VF that requires primitive UB
-	bVFRequiresPrimitiveUniformBuffer = (!bNaniteActive && Culling::UseCulling(GetScene().GetShaderPlatform()));
-
 	bSupportsInstanceDataBuffer = true;
 	UpdateDefaultInstanceSceneData();
 
@@ -1219,14 +1198,14 @@ void FLandscapeComponentSceneProxy::CreateRenderThreadResources()
 		if (!XYOffsetmapTexture)
 		{
 			FLandscapeVertexFactory* LandscapeVertexFactory = new FLandscapeVertexFactory(FeatureLevel);
-			LandscapeVertexFactory->Data.PositionComponent = FVertexStreamComponent(SharedBuffers->VertexBuffer, 0, sizeof(FLandscapeVertex), VET_UByte4);
+			LandscapeVertexFactory->Data.PositionComponent = FVertexStreamComponent(SharedBuffers->VertexBuffer, 0, sizeof(FLandscapeVertex), VET_Float4);
 			LandscapeVertexFactory->InitResource();
 			SharedBuffers->VertexFactory = LandscapeVertexFactory;
 		}
 		else
 		{
 			FLandscapeXYOffsetVertexFactory* LandscapeXYOffsetVertexFactory = new FLandscapeXYOffsetVertexFactory(FeatureLevel);
-			LandscapeXYOffsetVertexFactory->Data.PositionComponent = FVertexStreamComponent(SharedBuffers->VertexBuffer, 0, sizeof(FLandscapeVertex), VET_UByte4);
+			LandscapeXYOffsetVertexFactory->Data.PositionComponent = FVertexStreamComponent(SharedBuffers->VertexBuffer, 0, sizeof(FLandscapeVertex), VET_Float4);
 			LandscapeXYOffsetVertexFactory->InitResource();
 			SharedBuffers->VertexFactory = LandscapeXYOffsetVertexFactory;
 		}
@@ -1242,7 +1221,7 @@ void FLandscapeComponentSceneProxy::CreateRenderThreadResources()
 		{
 			//todo[vt]: We will need a version of this to support XYOffsetmapTexture
 			FLandscapeFixedGridVertexFactory* LandscapeVertexFactory = new FLandscapeFixedGridVertexFactory(FeatureLevel);
-			LandscapeVertexFactory->Data.PositionComponent = FVertexStreamComponent(SharedBuffers->VertexBuffer, 0, sizeof(FLandscapeVertex), VET_UByte4);
+			LandscapeVertexFactory->Data.PositionComponent = FVertexStreamComponent(SharedBuffers->VertexBuffer, 0, sizeof(FLandscapeVertex), VET_Float4);
 			LandscapeVertexFactory->InitResource();
 			SharedBuffers->FixedGridVertexFactory = LandscapeVertexFactory;
 		}
@@ -1254,11 +1233,6 @@ void FLandscapeComponentSceneProxy::CreateRenderThreadResources()
 	VertexFactory = SharedBuffers->VertexFactory;
 	FixedGridVertexFactory = SharedBuffers->FixedGridVertexFactory;
 
-	if (!bNaniteActive && Culling::UseCulling(GetScene().GetShaderPlatform()))
-	{
-		Culling::RegisterLandscape(*SharedBuffers, FeatureLevel, LandscapeKey, SubsectionSizeVerts, NumSubsections);
-	}
-	
 	// Assign LandscapeUniformShaderParameters
 	LandscapeUniformShaderParameters.InitResource();
 
@@ -1445,7 +1419,6 @@ void FLandscapeComponentSceneProxy::DestroyRenderThreadResources()
 	FPrimitiveSceneProxy::DestroyRenderThreadResources();
 	FLandscapeRenderSystem::UnregisterSection(this);
 	FLandscapeRenderSystem::DestroyResources(this);
-	Culling::UnregisterLandscape(LandscapeKey);
 }
 
 bool FLandscapeComponentSceneProxy::OnLevelAddedToWorld_RenderThread()
@@ -1472,7 +1445,7 @@ FLandscapeComponentSceneProxy::~FLandscapeComponentSceneProxy()
 {
 	// Free the subsection uniform buffer
 	LandscapeUniformShaderParameters.ReleaseResource();
-		
+
 	// Free the lod uniform buffers
 	for (int32 i = 0; i < LandscapeFixedGridUniformShaderParameters.Num(); ++i)
 	{
@@ -1959,11 +1932,6 @@ bool FLandscapeComponentSceneProxy::GetStaticMeshElement(int32 LODIndex, bool bF
 		BatchElement.FirstIndex = 0;
 		BatchElement.MinVertexIndex = SharedBuffers->IndexRanges[LODIndex].MinIndexFull;
 		BatchElement.MaxVertexIndex = SharedBuffers->IndexRanges[LODIndex].MaxIndexFull;
-				
-		if (!bForToolMesh && !bNaniteActive && Culling::UseCulling(GetScene().GetShaderPlatform()))
-		{
-			Culling::SetupMeshBatch(*SharedBuffers, MeshBatch);
-		}
 
 		// The default is overridden here only by mobile landscape to punch holes in the geometry
 		ApplyMeshElementModifier(BatchElement, LODIndex);
@@ -2455,18 +2423,6 @@ void FLandscapeComponentSceneProxy::GetDynamicMeshElements(const TArray<const FS
 #endif // !UE_BUILD_SHIPPING
 }
 
-void FLandscapeComponentSceneProxy::ApplyViewDependentMeshArguments(const FSceneView& View, FMeshBatch& ViewDependentMeshBatch) const
-{
-	Culling::FArguments CullingArgs{};
-	if (Culling::GetViewArguments(View, LandscapeKey, ComponentBase, ViewDependentMeshBatch.LODIndex, CullingArgs))
-	{
-		ViewDependentMeshBatch.Elements[0].NumPrimitives = 0;
-		ViewDependentMeshBatch.Elements[0].NumInstances = 0;
-		ViewDependentMeshBatch.Elements[0].IndirectArgsBuffer = CullingArgs.IndirectArgsBuffer;
-		ViewDependentMeshBatch.Elements[0].IndirectArgsOffset = CullingArgs.IndirectArgsOffset;
-	}
-}
-
 #if RHI_RAYTRACING
 void FLandscapeComponentSceneProxy::GetDynamicRayTracingInstances(FRayTracingMaterialGatheringContext& Context, TArray<FRayTracingInstance>& OutRayTracingInstances)
 {
@@ -2858,9 +2814,6 @@ FLandscapeSharedBuffers::FLandscapeSharedBuffers(const int32 InSharedBuffersKey,
 	, VertexFactory(nullptr)
 	, FixedGridVertexFactory(nullptr)
 	, VertexBuffer(nullptr)
-	, TileMesh(nullptr)
-	, TileVertexFactory(nullptr)
-	, TileDataBuffer(nullptr)
 	, bUse32BitIndices(false)
 
 #if WITH_EDITOR
@@ -2939,57 +2892,57 @@ FLandscapeSharedBuffers::~FLandscapeSharedBuffers()
 #endif
 
 	delete VertexFactory;
-
-	delete TileMesh;
-	delete TileDataBuffer;
-	delete TileVertexFactory;
 }
 
 //
 // FLandscapeVertexFactoryVertexShaderParameters
 //
-void FLandscapeVertexFactoryVertexShaderParameters::GetElementShaderBindings(
-	const class FSceneInterface* Scene,
-	const FSceneView* InView,
-	const class FMeshMaterialShader* Shader,
-	const EVertexInputStreamType InputStreamType,
-	ERHIFeatureLevel::Type FeatureLevel,
-	const FVertexFactory* VertexFactory,
-	const FMeshBatchElement& BatchElement,
-	class FMeshDrawSingleShaderBindings& ShaderBindings,
-	FVertexInputStreamArray& VertexStreams
-) const
+
+/** Shader parameters for use with FLandscapeVertexFactory */
+class FLandscapeVertexFactoryVertexShaderParameters : public FVertexFactoryShaderParameters
 {
-	SCOPE_CYCLE_COUNTER(STAT_LandscapeVFDrawTimeVS);
+	DECLARE_TYPE_LAYOUT(FLandscapeVertexFactoryVertexShaderParameters, NonVirtual);
+public:
+	/**
+	* Bind shader constants by name
+	* @param	ParameterMap - mapping of named shader constants to indices
+	*/
+	void Bind(const FShaderParameterMap& ParameterMap)
+	{
+	}
 
-	const FLandscapeBatchElementParams* BatchElementParams = (const FLandscapeBatchElementParams*)BatchElement.UserData;
-	check(BatchElementParams);
+	void GetElementShaderBindings(
+		const class FSceneInterface* Scene,
+		const FSceneView* InView,
+		const class FMeshMaterialShader* Shader,
+		const EVertexInputStreamType InputStreamType,
+		ERHIFeatureLevel::Type FeatureLevel,
+		const FVertexFactory* VertexFactory,
+		const FMeshBatchElement& BatchElement,
+		class FMeshDrawSingleShaderBindings& ShaderBindings,
+		FVertexInputStreamArray& VertexStreams
+	) const
+	{
+		SCOPE_CYCLE_COUNTER(STAT_LandscapeVFDrawTimeVS);
 
-	const FLandscapeComponentSceneProxy* SceneProxy = BatchElementParams->SceneProxy;
+		const FLandscapeBatchElementParams* BatchElementParams = (const FLandscapeBatchElementParams*)BatchElement.UserData;
+		check(BatchElementParams);
 
-	ShaderBindings.Add(Shader->GetUniformBufferParameter<FLandscapeUniformShaderParameters>(), *BatchElementParams->LandscapeUniformShaderParametersResource);
-	ShaderBindings.Add(Shader->GetUniformBufferParameter<FLandscapeSectionLODUniformParameters>(), BatchElementParams->LandscapeSectionLODUniformParameters);
+		const FLandscapeComponentSceneProxy* SceneProxy = BatchElementParams->SceneProxy;
 
-	if (InView && VertexFactory->GetType() == Culling::GetTileVertexFactoryType())
-	{
-		// TileVF is used only for LOD0
-		int32 LODIndex = 0; 
-		Culling::FArguments CullingArgs{};
-		if (Culling::GetViewArguments(*InView, SceneProxy->LandscapeKey, SceneProxy->ComponentBase, LODIndex, CullingArgs))
+		ShaderBindings.Add(Shader->GetUniformBufferParameter<FLandscapeUniformShaderParameters>(), *BatchElementParams->LandscapeUniformShaderParametersResource);
+		ShaderBindings.Add(Shader->GetUniformBufferParameter<FLandscapeSectionLODUniformParameters>(), BatchElementParams->LandscapeSectionLODUniformParameters);
+
+#if RHI_RAYTRACING
+		if (IsRayTracingEnabled())
 		{
-			check(VertexStreams.IsValidIndex(1));
-			VertexStreams[1].Offset = CullingArgs.TileDataOffset;
-			VertexStreams[1].VertexBuffer = CullingArgs.TileDataVertexBuffer;
+			ShaderBindings.Add(Shader->GetUniformBufferParameter<FLandscapeVertexFactoryMVFParameters>(), BatchElementParams->LandscapeVertexFactoryMVFUniformBuffer);
 		}
+#endif
 	}
+};
 
-#if RHI_RAYTRACING
-	if (IsRayTracingEnabled())
-	{
-		ShaderBindings.Add(Shader->GetUniformBufferParameter<FLandscapeVertexFactoryMVFParameters>(), BatchElementParams->LandscapeVertexFactoryMVFUniformBuffer);
-	}
-#endif
-}
+IMPLEMENT_TYPE_LAYOUT(FLandscapeVertexFactoryVertexShaderParameters);
 
 /**
   * Shader parameters for use with FLandscapeFixedGridVertexFactory
@@ -3097,7 +3050,7 @@ void FLandscapeVertexFactory::ModifyCompilationEnvironment(const FVertexFactoryS
 
 void FLandscapeVertexFactory::GetPSOPrecacheVertexFetchElements(EVertexInputStreamType VertexInputStreamType, FVertexDeclarationElementList& Elements)
 {
-	Elements.Add(FVertexElement(0, 0, VET_UByte4, 0, 0, false));
+	Elements.Add(FVertexElement(0, 0, VET_Float4, 0, 0, false));
 	Elements.Add(FVertexElement(1, 0, VET_UInt, 1, 0, true));
 }
 
@@ -3394,10 +3347,8 @@ public:
 
 				static const FName LandscapeVertexFactory = FName(TEXT("FLandscapeVertexFactory"));
 				static const FName LandscapeXYOffsetVertexFactory = FName(TEXT("FLandscapeXYOffsetVertexFactory"));
-				static const FName LandscapeTileVertexFactory = FName(TEXT("FLandscapeTileVertexFactory"));
 				static const FName NaniteVertexFactory = FName(TEXT("Nanite::FVertexFactory"));
 				if (VertexFactoryType->GetFName() == LandscapeVertexFactory ||
-					VertexFactoryType->GetFName() == LandscapeTileVertexFactory ||
 					VertexFactoryType->GetFName() == LandscapeXYOffsetVertexFactory ||
 					VertexFactoryType->GetFName() == NaniteVertexFactory)
 				{
@@ -3968,12 +3919,6 @@ double FLandscapeComponentSceneProxy::ComputeSectionResolution() const
 	return ComponentFullExtent / ComponentQuads;
 }
 
-void FLandscapeComponentSceneProxy::GetSectionBoundsAndLocalToWorld(FBoxSphereBounds& OutLocalBounds, FMatrix& OutLocalToWorld) const
-{
-	OutLocalBounds = GetLocalBounds();
-	OutLocalToWorld = GetLocalToWorld();
-}
-
 //
 // FLandscapeSectionInfo
 //
@@ -4020,12 +3965,6 @@ public:
 	{
 		return ProxyLOD;
 	}
-	
-	virtual void GetSectionBoundsAndLocalToWorld(FBoxSphereBounds& LocalBounds, FMatrix& LocalToWorld) const
-	{
-		LocalBounds = FBoxSphereBounds(ForceInit);
-		LocalToWorld = FMatrix::Identity;
-	}
 
 private:
 	int8 ProxyLOD;
diff --git a/Engine/Source/Runtime/Landscape/Public/LandscapeRender.h b/Engine/Source/Runtime/Landscape/Public/LandscapeRender.h
--- a/Engine/Source/Runtime/Landscape/Public/LandscapeRender.h
+++ b/Engine/Source/Runtime/Landscape/Public/LandscapeRender.h
@@ -154,31 +154,6 @@ public:
 	TArray<FLandscapeBatchElementParams, SceneRenderingAllocator> ElementParams;
 };
 
-class FLandscapeVertexFactoryVertexShaderParameters : public FVertexFactoryShaderParameters
-{
-	DECLARE_TYPE_LAYOUT(FLandscapeVertexFactoryVertexShaderParameters, NonVirtual);
-public:
-	/**
-	* Bind shader constants by name
-	* @param	ParameterMap - mapping of named shader constants to indices
-	*/
-	void Bind(const FShaderParameterMap& ParameterMap)
-	{
-	}
-
-	void GetElementShaderBindings(
-		const class FSceneInterface* Scene,
-		const FSceneView* InView,
-		const class FMeshMaterialShader* Shader,
-		const EVertexInputStreamType InputStreamType,
-		ERHIFeatureLevel::Type FeatureLevel,
-		const FVertexFactory* VertexFactory,
-		const FMeshBatchElement& BatchElement,
-		class FMeshDrawSingleShaderBindings& ShaderBindings,
-		FVertexInputStreamArray& VertexStreams
-	) const;
-};
-
 /** Pixel shader parameters for use with FLandscapeVertexFactory */
 class FLandscapeVertexFactoryPixelShaderParameters : public FVertexFactoryShaderParameters
 {
@@ -294,10 +269,10 @@ public:
 
 struct FLandscapeVertex
 {
-	uint8 VertexX;
-	uint8 VertexY;
-	uint8 SubX;
-	uint8 SubY;
+	float VertexX;
+	float VertexY;
+	float SubX;
+	float SubY;
 };
 
 //
@@ -357,11 +332,6 @@ public:
 	FLandscapeVertexFactory* VertexFactory;
 	FLandscapeVertexFactory* FixedGridVertexFactory;
 	FLandscapeVertexBuffer* VertexBuffer;
-	
-	FRenderResource* TileMesh;
-	FLandscapeVertexFactory* TileVertexFactory;
-	FVertexBuffer* TileDataBuffer;
-	
 	FIndexBuffer** IndexBuffers;
 	FLandscapeIndexRanges* IndexRanges;
 	bool bUse32BitIndices;
@@ -406,8 +376,7 @@ public:
 
 	/** Computes the worldspace units per vertex of the landscape section. */
 	virtual double ComputeSectionResolution() const { return -1.0; }
-	
-	virtual void GetSectionBoundsAndLocalToWorld(FBoxSphereBounds& LocalBounds, FMatrix& LocalToWorld) const = 0;
+
 public:
 	uint32 LandscapeKey;
 	FIntPoint ComponentBase;
@@ -533,8 +502,7 @@ public:
 	virtual void SetupViewFamily(FSceneViewFamily& InViewFamily) override {}
 	virtual void SetupView(FSceneViewFamily& InViewFamily, FSceneView& InView) override {}
 	virtual void BeginRenderViewFamily(FSceneViewFamily& InViewFamily) override;
-	
-	virtual void PreRenderViewFamily_RenderThread(FRDGBuilder& GraphBuilder, FSceneViewFamily& InViewFamily) override;
+
 	virtual void PreRenderView_RenderThread(FRDGBuilder& GraphBuilder, FSceneView& InView) override;
 	virtual void PreInitViews_RenderThread(FRDGBuilder& GraphBuilder) override;
 
@@ -825,7 +793,6 @@ public:
 	virtual void ApplyWorldOffset(FVector InOffset) override;
 	virtual void DrawStaticElements(FStaticPrimitiveDrawInterface* PDI) override;
 	virtual void GetDynamicMeshElements(const TArray<const FSceneView*>& Views, const FSceneViewFamily& ViewFamily, uint32 VisibilityMap, FMeshElementCollector& Collector) const override;
-	virtual void ApplyViewDependentMeshArguments(const FSceneView& View, FMeshBatch& ViewDependentMeshBatch) const override;
 	virtual uint32 GetMemoryFootprint() const override { return(sizeof(*this) + GetAllocatedSize()); }
 	virtual FPrimitiveViewRelevance GetViewRelevance(const FSceneView* View) const override;
 	virtual bool CanBeOccluded() const override;
@@ -869,7 +836,6 @@ public:
 	virtual float ComputeLODForView(const FSceneView& InView) const override;
 	virtual float ComputeLODBias() const override;
 	virtual double ComputeSectionResolution() const override;
-	virtual void GetSectionBoundsAndLocalToWorld(FBoxSphereBounds& LocalBounds, FMatrix& LocalToWorld) const override;
 };
 
 class FLandscapeDebugMaterialRenderProxy : public FMaterialRenderProxy
diff --git a/Engine/Source/Runtime/Renderer/Private/MeshDrawCommands.cpp b/Engine/Source/Runtime/Renderer/Private/MeshDrawCommands.cpp
--- a/Engine/Source/Runtime/Renderer/Private/MeshDrawCommands.cpp
+++ b/Engine/Source/Runtime/Renderer/Private/MeshDrawCommands.cpp
@@ -631,17 +631,7 @@ void GenerateDynamicMeshDrawCommands(
 		{
 			const FStaticMeshBatch* StaticMeshBatch = DynamicMeshCommandBuildRequests[MeshIndex];
 			const uint64 DefaultBatchElementMask = ~0ul;
-
-			if (StaticMeshBatch->bViewDependentArguments)
-			{
-				FMeshBatch ViewDepenedentMeshBatch(*StaticMeshBatch);
-				StaticMeshBatch->PrimitiveSceneInfo->Proxy->ApplyViewDependentMeshArguments(View, ViewDepenedentMeshBatch);
-				PassMeshProcessor->AddMeshBatch(ViewDepenedentMeshBatch, DefaultBatchElementMask, StaticMeshBatch->PrimitiveSceneInfo->Proxy, StaticMeshBatch->Id);
-			}
-			else
-			{
-				PassMeshProcessor->AddMeshBatch(*StaticMeshBatch, DefaultBatchElementMask, StaticMeshBatch->PrimitiveSceneInfo->Proxy, StaticMeshBatch->Id);
-			}
+			PassMeshProcessor->AddMeshBatch(*StaticMeshBatch, DefaultBatchElementMask, StaticMeshBatch->PrimitiveSceneInfo->Proxy, StaticMeshBatch->Id);
 		}
 
 		const int32 NumCommandsGenerated = VisibleCommands.Num() - NumCommandsBefore;
@@ -720,34 +710,23 @@ void GenerateMobileBasePassDynamicMeshDrawCommands(
 	{
 		const int32 NumCommandsBefore = VisibleCommands.Num();
 		const int32 NumStaticMeshBatches = DynamicMeshCommandBuildRequests.Num();
-		FMeshBatch ViewDependentMeshBatch{};
 
 		for (int32 MeshIndex = 0; MeshIndex < NumStaticMeshBatches; MeshIndex++)
 		{
 			const FStaticMeshBatch* StaticMeshBatch = DynamicMeshCommandBuildRequests[MeshIndex];
-			const int32 StaticMeshBatchId = StaticMeshBatch->Id;
-			const FPrimitiveSceneProxy* Proxy = StaticMeshBatch->PrimitiveSceneInfo->Proxy;
 
-			const FMeshBatch* MeshBatch = StaticMeshBatch;
-			if (MeshBatch->bViewDependentArguments)
-			{
-				ViewDependentMeshBatch = *MeshBatch;
-				Proxy->ApplyViewDependentMeshArguments(View, ViewDependentMeshBatch);
-				MeshBatch = &ViewDependentMeshBatch;
-			}
-
-			const int32 PrimitiveIndex = Proxy->GetPrimitiveSceneInfo()->GetIndex();
+			const int32 PrimitiveIndex = StaticMeshBatch->PrimitiveSceneInfo->Proxy->GetPrimitiveSceneInfo()->GetIndex();
 			if (!bSkipCSMShaderCulling
 				&& MobileCSMVisibilityInfo.bMobileDynamicCSMInUse
 				&& (MobileCSMVisibilityInfo.bAlwaysUseCSM || MobileCSMVisibilityInfo.MobilePrimitiveCSMReceiverVisibilityMap[PrimitiveIndex]))
 			{
 				const uint64 DefaultBatchElementMask = ~0ul;
-				MobilePassCSMPassMeshProcessor->AddMeshBatch(*MeshBatch, DefaultBatchElementMask, Proxy, StaticMeshBatchId);
+				MobilePassCSMPassMeshProcessor->AddMeshBatch(*StaticMeshBatch, DefaultBatchElementMask, StaticMeshBatch->PrimitiveSceneInfo->Proxy, StaticMeshBatch->Id);
 			}
 			else
 			{
 				const uint64 DefaultBatchElementMask = ~0ul;
-				PassMeshProcessor->AddMeshBatch(*MeshBatch, DefaultBatchElementMask, Proxy, StaticMeshBatchId);
+				PassMeshProcessor->AddMeshBatch(*StaticMeshBatch, DefaultBatchElementMask, StaticMeshBatch->PrimitiveSceneInfo->Proxy, StaticMeshBatch->Id);
 			}
 		}
 
diff --git a/Engine/Source/Runtime/Renderer/Private/ShaderBaseClasses.cpp b/Engine/Source/Runtime/Renderer/Private/ShaderBaseClasses.cpp
--- a/Engine/Source/Runtime/Renderer/Private/ShaderBaseClasses.cpp
+++ b/Engine/Source/Runtime/Renderer/Private/ShaderBaseClasses.cpp
@@ -517,8 +517,7 @@ void FMeshMaterialShader::GetElementShaderBindings(
 	{
 		const FShaderType* ShaderType = GetType(PointerTable);
 		ensureMsgf(!GetUniformBufferParameter<FPrimitiveUniformShaderParameters>().IsBound(), TEXT("Shader %s attempted to bind the Primitive uniform buffer even though Vertex Factory computes a PrimitiveId per-instance.  This will break auto-instancing.  Shaders should use GetPrimitiveData(PrimitiveId).Member instead of Primitive.Member."), ShaderType->GetName());
-		// Some primitives may use several VFs with a mixed support for a GPUScene. In this case all mesh batches get Primitive UB assigned regardless of VF type 
-		ensureMsgf(!BatchElement.PrimitiveUniformBuffer || (FeatureLevel == ERHIFeatureLevel::ES3_1 || PrimitiveSceneProxy->DoesVFRequirePrimitiveUniformBuffer()), TEXT("FMeshBatchElement was assigned a PrimitiveUniformBuffer even though Vertex Factory %s fetches primitive shader data through a Scene buffer.  The assigned PrimitiveUniformBuffer cannot be respected.  Use PrimitiveUniformBufferResource instead for dynamic primitive data."), ShaderType->GetName());
+		ensureMsgf(!BatchElement.PrimitiveUniformBuffer || (FeatureLevel == ERHIFeatureLevel::ES3_1), TEXT("FMeshBatchElement was assigned a PrimitiveUniformBuffer even though Vertex Factory %s fetches primitive shader data through a Scene buffer.  The assigned PrimitiveUniformBuffer cannot be respected.  Use PrimitiveUniformBufferResource instead for dynamic primitive data."), ShaderType->GetName());
 	}
 	else
 	{
