Landscape: Fix truncation warnings, part 6

diff --git a/Engine/Source/Runtime/Landscape/Private/LandscapeCulling.cpp b/Engine/Source/Runtime/Landscape/Private/LandscapeCulling.cpp
--- a/Engine/Source/Runtime/Landscape/Private/LandscapeCulling.cpp
+++ b/Engine/Source/Runtime/Landscape/Private/LandscapeCulling.cpp
@@ -219,8 +219,8 @@ void FLandscapeTileMesh::InitRHI()
 		{
 			for (uint32 x = 0; x < TileSizeVertx; x++)
 			{
-				Vertex->VertexX = x;
-				Vertex->VertexY = y;
+				Vertex->VertexX = static_cast<uint8>(x);
+				Vertex->VertexY = static_cast<uint8>(y);
 				Vertex->SubX = 0;
 				Vertex->SubY = 0;
 				Vertex++;
@@ -242,10 +242,10 @@ void FLandscapeTileMesh::InitRHI()
 		{
 			for (uint32 x = 0; x < LANDSCAPE_TILE_QUADS; x++)
 			{
-				uint16 i00 = (x + 0) + (y + 0) * (TileSizeVertx);
-				uint16 i10 = (x + 1) + (y + 0) * (TileSizeVertx);
-				uint16 i11 = (x + 1) + (y + 1) * (TileSizeVertx);
-				uint16 i01 = (x + 0) + (y + 1) * (TileSizeVertx);
+				uint16 i00 = static_cast<uint16>((x + 0) + (y + 0) * (TileSizeVertx));
+				uint16 i10 = static_cast<uint16>((x + 1) + (y + 0) * (TileSizeVertx));
+				uint16 i11 = static_cast<uint16>((x + 1) + (y + 1) * (TileSizeVertx));
+				uint16 i01 = static_cast<uint16>((x + 0) + (y + 1) * (TileSizeVertx));
 
 				Indices[0] = i00;
 				Indices[1] = i11;
@@ -304,10 +304,10 @@ void FLandscapeTileDataBuffer::InitRHI()
 			{
 				for (uint32 x = 0; x < SubsectionTilesRow; x++)
 				{
-					TileData[0] = x * LANDSCAPE_TILE_QUADS;
-					TileData[1] = y * LANDSCAPE_TILE_QUADS;
-					TileData[2] = SubX;
-					TileData[3] = SubY;
+					TileData[0] = static_cast<uint8>(x * LANDSCAPE_TILE_QUADS);
+					TileData[1] = static_cast<uint8>(y * LANDSCAPE_TILE_QUADS);
+					TileData[2] = static_cast<uint8>(SubX);
+					TileData[3] = static_cast<uint8>(SubY);
 					TileData += Stride;
 				}
 			}
@@ -584,7 +584,7 @@ static void ComputeSectionIntermediateData(FRDGBuilder& GraphBuilder, TArrayView
 
 			for (int32 SectionIdx = 0; SectionIdx < SectionLODValues.Num(); ++SectionIdx)
 			{
-				int32 LODValue = SectionLODValues[SectionIdx];
+				const int32 LODValue = static_cast<int32>(SectionLODValues[SectionIdx]);
 				FLandscapeSectionInfo* SectionInfo = RenderSystem.SectionInfos[SectionIdx];
 
 				if (LODValue == 0 && SectionInfo != nullptr)
@@ -595,14 +595,14 @@ static void ComputeSectionIntermediateData(FRDGBuilder& GraphBuilder, TArrayView
 					FMatrix SectionLocalToWorld;
 					SectionInfo->GetSectionBoundsAndLocalToWorld(SectionLocalBounds, SectionLocalToWorld);
 					const FLargeWorldRenderPosition SectionAbsoluteOrigin(SectionLocalToWorld.GetOrigin());
-					const int32 NeighborsMaxLOD = FMath::RoundFromZero(ComputeNeighborsMaxLOD(RenderSystem, SectionLODValues, SectionInfo->ComponentBase));
+					const int32 NeighborsMaxLOD = static_cast<int32>(FMath::RoundFromZero(ComputeNeighborsMaxLOD(RenderSystem, SectionLODValues, SectionInfo->ComponentBase)));
 
 					Section.LocalToRelativeWorld = FLargeWorldRenderScalar::MakeToRelativeWorldMatrix(SectionAbsoluteOrigin.GetTileOffset(), SectionLocalToWorld);
 					Section.TilePosition = SectionAbsoluteOrigin.GetTile();
-					Section.LocalZ = SectionLocalBounds.Origin.Z;
-					Section.HalfHeight = SectionLocalBounds.BoxExtent.Z;
+					Section.LocalZ = static_cast<float>(SectionLocalBounds.Origin.Z);
+					Section.HalfHeight = static_cast<float>(SectionLocalBounds.BoxExtent.Z);
 					// How many quads to add to each tile extent to compensate for a neighbors LOD
-					Section.NeighborLODExtent = (1 << NeighborsMaxLOD) - 1; 
+					Section.NeighborLODExtent = static_cast<float>((1 << NeighborsMaxLOD) - 1); 
 
 					ViewStateIntermediates.SectionBases.Add(SectionInfo->ComponentBase);
 				}
diff --git a/Engine/Source/Runtime/Landscape/Private/LandscapeRender.cpp b/Engine/Source/Runtime/Landscape/Private/LandscapeRender.cpp
--- a/Engine/Source/Runtime/Landscape/Private/LandscapeRender.cpp
+++ b/Engine/Source/Runtime/Landscape/Private/LandscapeRender.cpp
@@ -2663,10 +2663,10 @@ void FLandscapeVertexBuffer::InitRHI()
 			{
 				for (int32 x = 0; x < SubsectionSizeVerts; x++)
 				{
-					Vertex->VertexX = static_cast<float>(x);
-					Vertex->VertexY = static_cast<float>(y);
-					Vertex->SubX = static_cast<float>(SubX);
-					Vertex->SubY = static_cast<float>(SubY);
+					Vertex->VertexX = static_cast<uint8>(x);
+					Vertex->VertexY = static_cast<uint8>(y);
+					Vertex->SubX = static_cast<uint8>(SubX);
+					Vertex->SubY = static_cast<uint8>(SubY);
 					Vertex++;
 					VertexIndex++;
 				}
