Fixed mismatched copy constructor/assignment operator in FLandscapeLayer.

diff --git a/Engine/Source/Runtime/Landscape/Classes/Landscape.h b/Engine/Source/Runtime/Landscape/Classes/Landscape.h
--- a/Engine/Source/Runtime/Landscape/Classes/Landscape.h
+++ b/Engine/Source/Runtime/Landscape/Classes/Landscape.h
@@ -176,8 +176,6 @@ struct FLandscapeLayer
 		, BlendMode(LSBM_AdditiveBlend)
 	{}
 
-	FLandscapeLayer(const FLandscapeLayer& OtherLayer) = default;
-
 	UPROPERTY(meta = (IgnoreForMemberInitializationTest))
 	FGuid Guid;
 
