fix: LandscapeHLODBuilder missing PreEditChange/PostEditChange on Texture property change

#rnx

diff --git a/Engine/Source/Runtime/Landscape/Private/LandscapeHLODBuilder.cpp b/Engine/Source/Runtime/Landscape/Private/LandscapeHLODBuilder.cpp
--- a/Engine/Source/Runtime/Landscape/Private/LandscapeHLODBuilder.cpp
+++ b/Engine/Source/Runtime/Landscape/Private/LandscapeHLODBuilder.cpp
@@ -224,8 +224,10 @@ static UMaterialInterface* BakeLandscapeMaterial(const FHLODBuildContext& InHLOD
 		// Use clamp texture addressing to avoid artifacts between tiles
 		if (UTexture2D* Texture = Cast<UTexture2D>(Asset))
 		{
+			Texture->PreEditChange(nullptr);
 			Texture->AddressX = TA_Clamp;
 			Texture->AddressY = TA_Clamp;
+			Texture->PostEditChange();
 		}
 	});
 
