Landscape: Fixed a bug where ALandscape::GetCompleteBounds would not correctly consider streaming proxies which were unsaved/dirty.

diff --git a/Engine/Source/Runtime/Landscape/Private/Landscape.cpp b/Engine/Source/Runtime/Landscape/Private/Landscape.cpp
--- a/Engine/Source/Runtime/Landscape/Private/Landscape.cpp
+++ b/Engine/Source/Runtime/Landscape/Private/Landscape.cpp
@@ -4953,20 +4953,16 @@ FBox ULandscapeInfo::GetCompleteBounds() const
 	FWorldPartitionHelpers::ForEachActorDesc<ALandscapeProxy>(Landscape->GetWorld()->GetWorldPartition(), [this, &Bounds, Landscape](const FWorldPartitionActorDesc* ActorDesc)
 	{
 		FLandscapeActorDesc* LandscapeActorDesc = (FLandscapeActorDesc*)ActorDesc;
+		ALandscapeProxy* LandscapeProxy = Cast<ALandscapeProxy>(ActorDesc->GetActor());
 
-		if (LandscapeActorDesc->GridGuid == LandscapeGuid)
+		// Prioritize loaded bounds, as the bounds in the actor desc might not be up-to-date
+		if(LandscapeProxy && (LandscapeProxy->GetGridGuid() == LandscapeGuid))
 		{
-			ALandscapeProxy* LandscapeProxy = Cast<ALandscapeProxy>(ActorDesc->GetActor());
-
-			if (LandscapeProxy)
-			{
-				// Prioritize loaded bounds, as the bounds in the actor desc might not be up-to-date
-				LandscapeInfoBoundsHelper::AccumulateBounds(LandscapeProxy, Bounds);
-			}
-			else
-			{
-				Bounds += ActorDesc->GetEditorBounds();
-			}
+			LandscapeInfoBoundsHelper::AccumulateBounds(LandscapeProxy, Bounds);
+		}
+		else if (LandscapeActorDesc->GridGuid == LandscapeGuid)
+		{
+			Bounds += ActorDesc->GetEditorBounds();
 		}
 
 		return true;
