Landscape : Disable nanite rebuild on modification

diff --git a/Engine/Source/Runtime/Landscape/Private/LandscapeSubsystem.cpp b/Engine/Source/Runtime/Landscape/Private/LandscapeSubsystem.cpp
--- a/Engine/Source/Runtime/Landscape/Private/LandscapeSubsystem.cpp
+++ b/Engine/Source/Runtime/Landscape/Private/LandscapeSubsystem.cpp
@@ -668,7 +668,7 @@ void ULandscapeSubsystem::AddAsyncEvent(FGraphEventRef GraphEventRef)
 	NaniteMeshBuildEvents.Add(GraphEventRef);
 }
 
-int32 LiveRebuildNaniteOnModification = 1;
+int32 LiveRebuildNaniteOnModification = 0;
 static FAutoConsoleVariableRef CVarLiveRebuildNaniteOnModification(
 	TEXT("landscape.LiveRebuildNaniteOnModification"),
 	LiveRebuildNaniteOnModification,
