Fix for runtime crash rendering landscape

diff --git a/Engine/Source/Runtime/Landscape/Private/Landscape.cpp b/Engine/Source/Runtime/Landscape/Private/Landscape.cpp
--- a/Engine/Source/Runtime/Landscape/Private/Landscape.cpp
+++ b/Engine/Source/Runtime/Landscape/Private/Landscape.cpp
@@ -5444,7 +5444,7 @@ void ALandscapeProxy::UpdateRenderingMethod()
 	}
 
 	bool bNaniteActive = false;
-	if (GRenderNaniteLandscape != 0)
+	if ((GRenderNaniteLandscape != 0) && NaniteComponent)
 	{
 		bNaniteActive = UseNanite(GShaderPlatformForFeatureLevel[GMaxRHIFeatureLevel]);
 #if WITH_EDITOR
@@ -5464,7 +5464,7 @@ void ALandscapeProxy::UpdateRenderingMethod()
 #if WITH_EDITOR
 	if (bNaniteActive)
 	{
-		bNaniteActive = NaniteComponent && NaniteComponent->GetProxyContentId() == GetNaniteContentId();
+		bNaniteActive = NaniteComponent->GetProxyContentId() == GetNaniteContentId();
 	}
 #endif //WITH_EDITOR
 
