Fixed issues with landscape pre-computed lighting when culling is enabled

diff --git a/Engine/Source/Runtime/Landscape/Private/LandscapeRender.cpp b/Engine/Source/Runtime/Landscape/Private/LandscapeRender.cpp
--- a/Engine/Source/Runtime/Landscape/Private/LandscapeRender.cpp
+++ b/Engine/Source/Runtime/Landscape/Private/LandscapeRender.cpp
@@ -1090,7 +1090,10 @@ FLandscapeComponentSceneProxy::FLandscapeComponentSceneProxy(ULandscapeComponent
 		StaticLightingResolution = InComponent->GetLandscapeProxy()->StaticLightingResolution;
 	}
 
-	ComponentLightInfo = MakeUnique<FLandscapeLCI>(InComponent, FeatureLevel);
+	// Landscape GPU culling uses VF that requires primitive UB
+	bVFRequiresPrimitiveUniformBuffer |= (!bNaniteActive && Culling::UseCulling(GetScene().GetShaderPlatform()));
+
+	ComponentLightInfo = MakeUnique<FLandscapeLCI>(InComponent, FeatureLevel, bVFRequiresPrimitiveUniformBuffer != 0);
 	check(ComponentLightInfo);
 
 	const bool bHasStaticLighting = ComponentLightInfo->GetLightMap() || ComponentLightInfo->GetShadowMap();
@@ -1173,9 +1176,6 @@ FLandscapeComponentSceneProxy::FLandscapeComponentSceneProxy(ULandscapeComponent
 		}
 	}
 
-	// Landscape GPU culling uses VF that requires primitive UB
-	bVFRequiresPrimitiveUniformBuffer |= (!bNaniteActive && Culling::UseCulling(GetScene().GetShaderPlatform()));
-
 	bSupportsInstanceDataBuffer = true;
 	UpdateDefaultInstanceSceneData();
 
diff --git a/Engine/Source/Runtime/Landscape/Public/LandscapeRender.h b/Engine/Source/Runtime/Landscape/Public/LandscapeRender.h
--- a/Engine/Source/Runtime/Landscape/Public/LandscapeRender.h
+++ b/Engine/Source/Runtime/Landscape/Public/LandscapeRender.h
@@ -645,7 +645,7 @@ class LANDSCAPE_API FLandscapeComponentSceneProxy : public FPrimitiveSceneProxy,
 	{
 	public:
 		/** Initialization constructor. */
-		FLandscapeLCI(const ULandscapeComponent* InComponent, ERHIFeatureLevel::Type FeatureLevel)
+		FLandscapeLCI(const ULandscapeComponent* InComponent, ERHIFeatureLevel::Type FeatureLevel, bool bVFRequiresPrimitiveUniformBuffer)
 			: FLightCacheInterface()
 		{
 			const FMeshMapBuildData* MapBuildData = InComponent->GetMeshMapBuildData();
@@ -655,7 +655,8 @@ class LANDSCAPE_API FLandscapeComponentSceneProxy : public FPrimitiveSceneProxy,
 				SetLightMap(MapBuildData->LightMap);
 				SetShadowMap(MapBuildData->ShadowMap);
 				SetResourceCluster(MapBuildData->ResourceCluster);
-				if (FeatureLevel >= ERHIFeatureLevel::SM5)
+				// If landscape uses VF that requires primitive UB that means it does not use GPUScene therefore it may need precomputed lighting buffer as well
+				if (FeatureLevel >= ERHIFeatureLevel::SM5 && !bVFRequiresPrimitiveUniformBuffer)
 				{
 					// Landscape does not support GPUScene on mobile
 					// TODO: enable this when GPUScene support is implemented
