Landscape culling: Add a runtime check whether Lumen is actually used. This is to support landscape culling on configurations that do support Lumen and not using it

diff --git a/Engine/Source/Runtime/Landscape/Private/LandscapeCulling.cpp b/Engine/Source/Runtime/Landscape/Private/LandscapeCulling.cpp
--- a/Engine/Source/Runtime/Landscape/Private/LandscapeCulling.cpp
+++ b/Engine/Source/Runtime/Landscape/Private/LandscapeCulling.cpp
@@ -147,16 +147,58 @@ constexpr static uint32 MIN_SECTION_SIZE_QUADS = 31u;
 constexpr static uint32 INDIRECT_ARGS_NUM_WORDS = 5u;
 constexpr static uint32 TILE_DATA_ENTRY_NUM_BYTES = 4u;
 
+static void LandscapeComponentsRecreateRenderState(IConsoleVariable* Variable)
+{
+	for (auto* LandscapeComponent : TObjectRange<ULandscapeComponent>())
+	{
+		LandscapeComponent->MarkRenderStateDirty();
+	}
+}
+
+struct FLumenCVarsState
+{
+	FLumenCVarsState()
+	{
+		IConsoleVariable* SupportedCVar = IConsoleManager::Get().FindConsoleVariable(TEXT("r.Lumen.Supported"));
+		bLumenSupported = (SupportedCVar && SupportedCVar->GetInt() != 0);
+		
+		IConsoleVariable* DiffuseIndirectCVar = IConsoleManager::Get().FindConsoleVariable(TEXT("r.Lumen.DiffuseIndirect.Allow"));
+		if (bLumenSupported && DiffuseIndirectCVar)
+		{
+			DiffuseIndirectCVar->SetOnChangedCallback(FConsoleVariableDelegate::CreateStatic(&LandscapeComponentsRecreateRenderState));
+		}
+		
+		IConsoleVariable* ReflectionsCVar = IConsoleManager::Get().FindConsoleVariable(TEXT("r.Lumen.Reflections.Allow"));
+		if (bLumenSupported && ReflectionsCVar)
+		{
+			ReflectionsCVar->SetOnChangedCallback(FConsoleVariableDelegate::CreateStatic(&LandscapeComponentsRecreateRenderState));
+		}
+
+		DiffuseIndirectInt = DiffuseIndirectCVar->AsVariableInt();
+		ReflectionsInt = ReflectionsCVar->AsVariableInt();
+	}
+
+	bool IsActive() const
+	{
+		return bLumenSupported && ((DiffuseIndirectInt && DiffuseIndirectInt->GetValueOnAnyThread() != 0) || (ReflectionsInt && ReflectionsInt->GetValueOnAnyThread() != 0));
+	}
+
+	TConsoleVariableData<int32>* DiffuseIndirectInt;
+	TConsoleVariableData<int32>* ReflectionsInt;
+	bool bLumenSupported = false;
+};
+
 bool UseCulling(EShaderPlatform Platform)
 {
+	static FLumenCVarsState LumenCVarsState;
 	static FShaderPlatformCachedIniValue<int32> LandscapeSupportCullingIniValue(TEXT("landscape.SupportGPUCulling"));
 	
 	return 
 		LandscapeSupportCullingIniValue.Get(Platform) != 0 &&
 		// These features require VF PrimitiveID support which is not possible with culling VF atm
-		// Note that UseVirtualShadowMaps includes runtime logic, so this function can't be used in cook time decisions
+		// Note that VSM and Lumen test includes runtime logic, so this function can't be used in cook time decisions
 		!UseVirtualShadowMaps(Platform, GMaxRHIFeatureLevel) &&
-		!FDataDrivenShaderPlatformInfo::GetSupportsLumenGI(Platform);
+		!(FDataDrivenShaderPlatformInfo::GetSupportsLumenGI(Platform) && LumenCVarsState.IsActive());
 }
 
 FVertexFactoryType* GetTileVertexFactoryType()
