Remove unnecessary r.Lumen.Reflections.Allow test from landscape culling support

diff --git a/Engine/Source/Runtime/Landscape/Private/LandscapeCulling.cpp b/Engine/Source/Runtime/Landscape/Private/LandscapeCulling.cpp
--- a/Engine/Source/Runtime/Landscape/Private/LandscapeCulling.cpp
+++ b/Engine/Source/Runtime/Landscape/Private/LandscapeCulling.cpp
@@ -167,24 +167,15 @@ struct FLumenCVarsState
 		{
 			DiffuseIndirectCVar->SetOnChangedCallback(FConsoleVariableDelegate::CreateStatic(&LandscapeComponentsRecreateRenderState));
 		}
-		
-		IConsoleVariable* ReflectionsCVar = IConsoleManager::Get().FindConsoleVariable(TEXT("r.Lumen.Reflections.Allow"));
-		if (bLumenSupported && ReflectionsCVar)
-		{
-			ReflectionsCVar->SetOnChangedCallback(FConsoleVariableDelegate::CreateStatic(&LandscapeComponentsRecreateRenderState));
-		}
-
 		DiffuseIndirectInt = DiffuseIndirectCVar->AsVariableInt();
-		ReflectionsInt = ReflectionsCVar->AsVariableInt();
 	}
 
 	bool IsActive() const
 	{
-		return bLumenSupported && ((DiffuseIndirectInt && DiffuseIndirectInt->GetValueOnAnyThread() != 0) || (ReflectionsInt && ReflectionsInt->GetValueOnAnyThread() != 0));
+		return bLumenSupported && (DiffuseIndirectInt && DiffuseIndirectInt->GetValueOnAnyThread() != 0);
 	}
 
 	TConsoleVariableData<int32>* DiffuseIndirectInt;
-	TConsoleVariableData<int32>* ReflectionsInt;
 	bool bLumenSupported = false;
 };
 
