Enable landscape GPU culling by default. It will be deactivated on unsupported configurations (projects with VSM or Lumen enabled and used)

diff --git a/Engine/Source/Runtime/Landscape/Private/LandscapeCulling.cpp b/Engine/Source/Runtime/Landscape/Private/LandscapeCulling.cpp
--- a/Engine/Source/Runtime/Landscape/Private/LandscapeCulling.cpp
+++ b/Engine/Source/Runtime/Landscape/Private/LandscapeCulling.cpp
@@ -13,7 +13,7 @@ extern TMap<uint32, FLandscapeRenderSystem*> LandscapeRenderSystems;
 
 static TAutoConsoleVariable<int32> CVarLandscapeSupportCulling(
 	TEXT("landscape.SupportGPUCulling"),
-	0,
+	1,
 	TEXT("Whether to support landscape GPU culling"),
 	ECVF_ReadOnly | ECVF_RenderThreadSafe);
 
