Fix duplicated landscape proxies in cooked builds.

diff --git a/Engine/Source/Runtime/Landscape/Private/LandscapeEdit.cpp b/Engine/Source/Runtime/Landscape/Private/LandscapeEdit.cpp
--- a/Engine/Source/Runtime/Landscape/Private/LandscapeEdit.cpp
+++ b/Engine/Source/Runtime/Landscape/Private/LandscapeEdit.cpp
@@ -4570,7 +4570,7 @@ void ALandscape::PostRegisterAllComponents()
 	};
 
 	// Until it is properly supported, ALandscape with layer brushes will force all of its proxies to be loaded in editor
-	if (GEditor && !GetWorld()->IsGameWorld() && LandscapeGuid.IsValid() && HasValidBrush())
+	if (GEditor && !GetWorld()->IsGameWorld() && LandscapeGuid.IsValid() && HasValidBrush() && !IsRunningCommandlet())
 	{
 		if (UWorldPartition* WorldPartition = GetWorld()->GetWorldPartition())
 		{
