Allow PCG to access the original landscape guid, before instancing modifications.

PCG uses the landscape guid to cache content, but when we modify the guid because of instancing, it can no longer find it.  This makes the original landscape guid available, so that PCG can use it as a cache key.

diff --git a/Engine/Source/Runtime/Landscape/Classes/LandscapeProxy.h b/Engine/Source/Runtime/Landscape/Classes/LandscapeProxy.h
--- a/Engine/Source/Runtime/Landscape/Classes/LandscapeProxy.h
+++ b/Engine/Source/Runtime/Landscape/Classes/LandscapeProxy.h
@@ -411,10 +411,17 @@ protected:
 	UPROPERTY()
 	TObjectPtr<ULandscapeSplinesComponent> SplineComponent;
 
-	/** Guid for LandscapeEditorInfo **/
+	/** Guid for LandscapeEditorInfo, all proxies that belong to the same landscape should have the same LandscapeGuid, even if split across world partitions 
+	  * Note that this value may change when the landscape is instanced (or in PIE) in order to allow multiple instances of the same landscape to exist.
+	  **/
 	UPROPERTY()
 	FGuid LandscapeGuid;
 
+	/** 
+	  * The original unmutated LandscapeGuid on the source asset, before instancing modfications.
+	  **/
+	FGuid OriginalLandscapeGuid;
+
 	/** Use Nanite to render landscape as a mesh on supported platforms. */
 	UPROPERTY(EditAnywhere, BlueprintReadOnly, Category = Nanite)
 	bool bEnableNanite = false;
@@ -885,8 +892,23 @@ public:
 	/** Fills an array with height values **/
 	LANDSCAPE_API void GetHeightValues(int32& SizeX, int32& SizeY, TArray<float>& ArrayValue) const;
 
+	/* Return the landscape guid, used to identify landscape proxies (and splines) that belong to the same landscape, even across world partitions.
+	 * Also used as the world partition grid guid.  This value may be modified when the landscape is instanced, to allow multiple instances of the same
+	 * landscape to exist simultaneously.  If you want the original (uninstanced) value, use GetOriginalLandscapeGuid().
+	 */
 	virtual FGuid GetLandscapeGuid() const override { return LandscapeGuid; }
-	void SetLandscapeGuid(const FGuid& Guid) { LandscapeGuid = Guid; }
+	void SetLandscapeGuid(const FGuid& Guid)
+	{
+		// we probably shouldn't be setting the landscape guid on instanced landscapes
+		check((OriginalLandscapeGuid == LandscapeGuid) || !OriginalLandscapeGuid.IsValid());
+		LandscapeGuid = Guid;
+		OriginalLandscapeGuid = Guid;
+	}
+
+	/* Return the original landscape guid, before it was modified by instancing.
+	 * When not instanced, this value is equal to LandscapeGuid.
+	 */
+	const FGuid& GetOriginalLandscapeGuid() const { return OriginalLandscapeGuid; }
 
 	UFUNCTION(BlueprintCallable, Category = "Landscape|Runtime")
 	virtual ALandscape* GetLandscapeActor() PURE_VIRTUAL(GetLandscapeActor, return nullptr;)
@@ -1394,6 +1416,8 @@ private:
 
 #endif
 
+	void PostLoadFixupLandscapeGuidsIfInstanced();
+
 #if WITH_EDITORONLY_DATA
 public:
 	static const TArray<ALandscapeProxy*>& GetLandscapeProxies() { return LandscapeProxies; }
diff --git a/Engine/Source/Runtime/Landscape/Private/Landscape.cpp b/Engine/Source/Runtime/Landscape/Private/Landscape.cpp
--- a/Engine/Source/Runtime/Landscape/Private/Landscape.cpp
+++ b/Engine/Source/Runtime/Landscape/Private/Landscape.cpp
@@ -3412,41 +3412,41 @@ bool ULandscapeInfo::UpdateLayerInfoMap(ALandscapeProxy* Proxy /*= nullptr*/, bo
  * this avoids guid collisions when you instance a world (and its landscapes) multiple times,
  * while maintaining the same GUID between landscape proxy objects within an instance
  */ 
-namespace LandscapeInstancing
+void ALandscapeProxy::PostLoadFixupLandscapeGuidsIfInstanced()
 {
-	void FixupLandscapeGuidsIfInstanced(ALandscapeProxy* LandscapeProxy)
+	// record the original value before modification
+	check(!OriginalLandscapeGuid.IsValid() || (OriginalLandscapeGuid == LandscapeGuid));
+	OriginalLandscapeGuid = this->LandscapeGuid;
+
+	// we shouldn't be dealing with any instanced landscapes in these cases, early out
+	if (this->IsTemplate() || IsRunningCookCommandlet())
 	{
-		// we shouldn't be dealing with any instanced landscapes in these cases, early out
-		if (LandscapeProxy->IsTemplate() || IsRunningCookCommandlet())
-		{
-			return;
-		}
+		return;
+	}
 
-		UWorldPartition* WorldPartition = FWorldPartitionHelpers::GetWorldPartition(LandscapeProxy);
-		UWorld* OuterWorld = WorldPartition ? WorldPartition->GetTypedOuter<UWorld>() : LandscapeProxy->GetTypedOuter<UWorld>();
+	UWorldPartition* WorldPartition = FWorldPartitionHelpers::GetWorldPartition(this);
+	UWorld* OuterWorld = WorldPartition ? WorldPartition->GetTypedOuter<UWorld>() : this->GetTypedOuter<UWorld>();
 
-		// TODO [chris.tchou] : Note this is not 100% correct, IsInstanced() returns TRUE when using PIE on non-instanced landscapes.
-		// That is generally ok however, as the GUID remaps are still deterministic and landscape still works.
-		// the only issue is if some external system is tracking or interacting with landscape via the GUID, as it changes.
-		if (OuterWorld && OuterWorld->IsInstanced())
-		{
-			FArchiveMD5 Ar;
-			FGuid OldLandscapeGuid = LandscapeProxy->GetLandscapeGuid();
-			Ar << OldLandscapeGuid;
+	// TODO [chris.tchou] : Note this is not 100% correct, IsInstanced() returns TRUE when using PIE on non-instanced landscapes.
+	// That is generally ok however, as the GUID remaps are still deterministic and landscape still works.
+	if (OuterWorld && OuterWorld->IsInstanced())
+	{
+		FArchiveMD5 Ar;
+		FGuid OldLandscapeGuid = this->LandscapeGuid;
+		Ar << OldLandscapeGuid;
 
-			UPackage* OuterWorldPackage = OuterWorld->GetPackage();
-			if (ensure(OuterWorldPackage))
-			{
-				FName PackageName = OuterWorldPackage->GetFName();
-				Ar << PackageName;
-			}
+		UPackage* OuterWorldPackage = OuterWorld->GetPackage();
+		if (ensure(OuterWorldPackage))
+		{
+			FName PackageName = OuterWorldPackage->GetFName();
+			Ar << PackageName;
+		}
 
-			FMD5Hash MD5Hash;
-			Ar.GetHash(MD5Hash);
+		FMD5Hash MD5Hash;
+		Ar.GetHash(MD5Hash);
 
-			FGuid NewLandscapeGuid = MD5HashToGuid(MD5Hash);
-			LandscapeProxy->SetLandscapeGuid(NewLandscapeGuid);
-		}
+		FGuid NewLandscapeGuid = MD5HashToGuid(MD5Hash);
+		this->LandscapeGuid = NewLandscapeGuid;
 	}
 }
 
@@ -3454,7 +3454,7 @@ void ALandscapeProxy::PostLoad()
 {
 	Super::PostLoad();
 
-	LandscapeInstancing::FixupLandscapeGuidsIfInstanced(this);
+	PostLoadFixupLandscapeGuidsIfInstanced();
 
 	// Temporary
 	if (ComponentSizeQuads == 0 && LandscapeComponents.Num() > 0)
@@ -3659,6 +3659,7 @@ void ALandscapeProxy::GetSharedProperties(ALandscapeProxy* Landscape)
 		Modify();
 
 		LandscapeGuid = Landscape->LandscapeGuid;
+		OriginalLandscapeGuid = Landscape->OriginalLandscapeGuid;
 
 		//@todo UE merge, landscape, this needs work
 		RootComponent->SetRelativeScale3D(Landscape->GetRootComponent()->GetComponentToWorld().GetScale3D());
diff --git a/Engine/Source/Runtime/Landscape/Private/LandscapeEdit.cpp b/Engine/Source/Runtime/Landscape/Private/LandscapeEdit.cpp
--- a/Engine/Source/Runtime/Landscape/Private/LandscapeEdit.cpp
+++ b/Engine/Source/Runtime/Landscape/Private/LandscapeEdit.cpp
@@ -739,8 +739,6 @@ void ULandscapeComponent::PostEditUndo()
 
 		for (ALandscapeProxy* ComponentProxy : Proxies)
 		{
-			check(ComponentProxy->GetLandscapeGuid().IsValid());
-
 			// Here we create and register with the LandscapeInfo early (before PostRegisterAllComponents)
 			ULandscapeInfo* LandscapeInfo = ComponentProxy->GetLandscapeInfo();
 			if (!LandscapeInfo->IsRegistered(ComponentProxy))
@@ -5626,9 +5624,11 @@ void ALandscapeStreamingProxy::PostEditChangeProperty(FPropertyChangedEvent& Pro
 
 	if (PropertyName == FName(TEXT("LandscapeActorRef")))
 	{
+		// Landscape Actor reference was changed .. need to update LandscapeGUIDs to match
 		if (LandscapeActorRef && IsValidLandscapeActor(LandscapeActorRef.Get()))
 		{
 			LandscapeGuid = LandscapeActorRef->GetLandscapeGuid();
+			OriginalLandscapeGuid = LandscapeActorRef->GetOriginalLandscapeGuid();
 			if (GIsEditor && GetWorld() && !GetWorld()->IsPlayInEditor())
 			{
 				// TODO - only need to refresh the old and new landscape info
