PSO precaching: Added landscape tile VF to precache

diff --git a/Engine/Source/Runtime/Landscape/Private/Landscape.cpp b/Engine/Source/Runtime/Landscape/Private/Landscape.cpp
--- a/Engine/Source/Runtime/Landscape/Private/Landscape.cpp
+++ b/Engine/Source/Runtime/Landscape/Private/Landscape.cpp
@@ -62,6 +62,7 @@ Landscape.cpp: Terrain rendering
 #include "ComponentRecreateRenderStateContext.h"
 #include "LandscapeWeightmapUsage.h"
 #include "LandscapeSubsystem.h"
+#include "LandscapeCulling.h"
 #include "ContentStreaming.h"
 #include "UObject/ObjectSaveContext.h"
 #include "GlobalShader.h"
@@ -1040,6 +1041,13 @@ void ULandscapeComponent::PostLoad()
 			VertexFactoryDataList.Add(FPSOPrecacheVertexFactoryData(&FLandscapeFixedGridVertexFactory::StaticType));
 		}
 
+		using namespace UE::Landscape;
+		
+		if (Culling::UseCulling(GMaxRHIShaderPlatform))
+		{
+			VertexFactoryDataList.Add(FPSOPrecacheVertexFactoryData(Culling::GetTileVertexFactoryType()));
+		}
+
 		TArray<FMaterialPSOPrecacheRequestID> MaterialPrecacheRequestIDs;
 		for (UMaterialInterface* MaterialInterface : Materials)
 		{
diff --git a/Engine/Source/Runtime/Landscape/Private/LandscapeCulling.cpp b/Engine/Source/Runtime/Landscape/Private/LandscapeCulling.cpp
--- a/Engine/Source/Runtime/Landscape/Private/LandscapeCulling.cpp
+++ b/Engine/Source/Runtime/Landscape/Private/LandscapeCulling.cpp
@@ -110,8 +110,8 @@ void FLandscapeTileVertexFactory::Copy(const FLandscapeTileVertexFactory& Other)
 
 void FLandscapeTileVertexFactory::GetPSOPrecacheVertexFetchElements(EVertexInputStreamType VertexInputStreamType, FVertexDeclarationElementList& Elements)
 {
-	Elements.Add(FVertexElement(0, 0, VET_UByte4, 0, 0, false)); // Position
-	Elements.Add(FVertexElement(1, 0, VET_UByte4, 1, 0, true)); // Tile per-instance data
+	Elements.Add(FVertexElement(0, 0, VET_UByte4, 0, sizeof(FLandscapeVertex), false)); // Position
+	Elements.Add(FVertexElement(1, 0, VET_UByte4, 1, 4u, true)); // Tile per-instance data
 }
 
 bool FLandscapeTileVertexFactory::ShouldCompilePermutation(const FVertexFactoryShaderPermutationParameters& Parameters)
