Landscape Texture Mip Storage Provider, give us an option to customize how Landscape Texture data is stored.

Currently, Landscape will only apply the new storage provider to it's heightmaps at cook time, when it sees that the landscape has been flagged to use the new compressed storage path (enabled under Landscape -> Advanced), and only for Windows platforms (to be expanded as we test other platforms).

The compressed storage path option on the landscape is hidden by default, and can be made visible via CVAR ("landscape.ShowCompressHeightMapsOption")

Makes use of a new All Mip Provider interface to Texture2Ds, that expands on the existing Mip Data Provider (that only supported providing streaming mips).
The new interface assumes you want to completely override the Texture2D mip data.

When Texture2D sees that it has an All Mip Provider, it will skip serializing it's own heavyweight mip data during cook, as it assumes the provider will take care of providing all mip data.

The majority of the implementation is in LandscapeTextureStorageProvider, which implements the AllMipProvider that can store landscape heightmaps in a custom compressed format.  Each mip's height data is stored using delta encoding, and normals are reconstructed.

diff --git a/Engine/Source/Editor/DetailCustomizations/Private/LandscapeProxyUIDetails.cpp b/Engine/Source/Editor/DetailCustomizations/Private/LandscapeProxyUIDetails.cpp
--- a/Engine/Source/Editor/DetailCustomizations/Private/LandscapeProxyUIDetails.cpp
+++ b/Engine/Source/Editor/DetailCustomizations/Private/LandscapeProxyUIDetails.cpp
@@ -50,6 +50,12 @@ class URuntimeVirtualTexture;
 
 #define LOCTEXT_NAMESPACE "FLandscapeProxyUIDetails"
 
+static TAutoConsoleVariable<int32> CVarShowCompressWeightMapsOption(
+	TEXT("landscape.ShowCompressHeightMapsOption"),
+	0,
+	TEXT("Enable editing of the compressed height map option on landscape proxies (experimental)."),
+	ECVF_Default);
+
 FLandscapeProxyUIDetails::FLandscapeProxyUIDetails()
 {
 }
@@ -98,6 +104,11 @@ void FLandscapeProxyUIDetails::CustomizeDetails( IDetailLayoutBuilder& DetailBui
 		DetailBuilder.HideProperty(DetailBuilder.GetProperty(GET_MEMBER_NAME_CHECKED(ALandscape, bAreNewLandscapeActorsSpatiallyLoaded), ALandscape::StaticClass()));
 	}
 
+	if (CVarShowCompressWeightMapsOption->GetInt() == 0)
+	{
+		DetailBuilder.HideProperty(DetailBuilder.GetProperty(GET_MEMBER_NAME_CHECKED(ALandscapeProxy, bUseCompressedHeightmapStorage), ALandscapeProxy::StaticClass()));
+	}
+
 	if (LandscapeActors.Num() == 1)
 	{
 		TWeakObjectPtr<ALandscape> LandscapeActor = LandscapeActors[0];
diff --git a/Engine/Source/Runtime/CoreUObject/Public/Serialization/BulkData.h b/Engine/Source/Runtime/CoreUObject/Public/Serialization/BulkData.h
--- a/Engine/Source/Runtime/CoreUObject/Public/Serialization/BulkData.h
+++ b/Engine/Source/Runtime/CoreUObject/Public/Serialization/BulkData.h
@@ -1164,6 +1164,30 @@ public:
 	{
 		FBulkData::Serialize(Ar, Owner, bAttemptFileMapping, GetElementSize(), FileRegionType);
 	}
+
+	/**
+	 * Serialize function used to serialize this bulk data structure, applying an override of the Bulk Data Flags when saving.
+	 *
+	 * @param Ar					Archive to serialize with
+	 * @param Owner					Object owning the bulk data
+	 * @param SaveOverrideFlags		EBulkDataFlags to use when saving (restores original flags after serialization)
+	 * @param bAttemptFileMapping	If true, attempt to map this instead of loading it into malloc'ed memory
+	 * @param FileRegionType		When cooking, a hint describing the type of data, used by some platforms to improve compression ratios
+	 */
+	void SerializeWithFlags(FArchive& Ar, UObject* Owner, uint32 SaveOverrideFlags, bool bAttemptFileMapping = false, EFileRegionType FileRegionType = EFileRegionType::None)
+	{
+		if (Ar.IsSaving())
+		{
+			const uint32 OriginalBulkDataFlags = GetBulkDataFlags();
+			SetBulkDataFlags(SaveOverrideFlags);		// NOTE this does an OR between existing flags and the override flags
+			FBulkData::Serialize(Ar, Owner, bAttemptFileMapping, GetElementSize(), FileRegionType);
+			ResetBulkDataFlags(OriginalBulkDataFlags);
+		}
+		else
+		{
+			FBulkData::Serialize(Ar, Owner, bAttemptFileMapping, GetElementSize(), FileRegionType);
+		}
+	}
 };
 
 UE_DEPRECATED(5.1, "Use FBulkData/TBulkData");
diff --git a/Engine/Source/Runtime/Engine/Classes/Engine/Texture.h b/Engine/Source/Runtime/Engine/Classes/Engine/Texture.h
--- a/Engine/Source/Runtime/Engine/Classes/Engine/Texture.h
+++ b/Engine/Source/Runtime/Engine/Classes/Engine/Texture.h
@@ -884,10 +884,10 @@ public:
 	 *
 	 * @param Ar Archive to serialize with
 	 * @param Owner Owner texture
-	 * @param bStreamable Store some mips inline, only used during cooking
+	 * @param bStreamable set up serialization to stream some mips
+	 * @param bSerializeMipData if false, no mip bulk data will be serialized.  This should only be false when an alternate All Mip Data Provider is attached.
 	 */
-	void SerializeCooked(FArchive& Ar, class UTexture* Owner, bool bStreamable);
-	
+	void SerializeCooked(FArchive& Ar, class UTexture* Owner, bool bStreamable, const bool bSerializeMipData);
 	
 	inline void SetPackedData(int32 InNumSlices, bool bInHasOptData, bool bInCubeMap)
 	{
@@ -1355,7 +1355,11 @@ public:
 	UPROPERTY()
 	uint8 CompressionYCoCg : 1;
 
-	/** If true, the RHI texture will be created without TexCreate_OfflineProcessed. */
+	/** If true, the RHI texture will be created without TexCreate_OfflineProcessed.
+	  * This controls what format the data will be uploaded to RHI.
+	  * Offline processed textures may have platform specific tiling applied, and/or have their mip tails pre-combined into a single mip's data.
+	  * If NotOffline, then it will expect data to be uploaded in standard per-mip layouts.
+	  */
 	UPROPERTY(transient)
 	uint8 bNotOfflineProcessed : 1;
 
@@ -1521,7 +1525,7 @@ public:
 	/**
 	 * Serializes cooked platform data.
 	 */
-	ENGINE_API void SerializeCookedPlatformData(class FArchive& Ar);
+	ENGINE_API void SerializeCookedPlatformData(class FArchive& Ar, const bool bSerializeMipData = true);
 
 	//~ Begin IInterface_AssetUserData Interface
 	ENGINE_API virtual void AddAssetUserData(UAssetUserData* InUserData) override;
diff --git a/Engine/Source/Runtime/Engine/Classes/Engine/Texture2D.h b/Engine/Source/Runtime/Engine/Classes/Engine/Texture2D.h
--- a/Engine/Source/Runtime/Engine/Classes/Engine/Texture2D.h
+++ b/Engine/Source/Runtime/Engine/Classes/Engine/Texture2D.h
@@ -10,6 +10,7 @@
 #if UE_ENABLE_INCLUDE_ORDER_DEPRECATED_IN_5_2
 #include "TextureResource.h"
 #endif
+#include "Engine/TextureAllMipDataProviderFactory.h"
 #include "Misc/FieldAccessor.h"
 #include "Serialization/BulkData.h"
 #include "Texture2D.generated.h"
@@ -87,6 +88,11 @@ public:
 	/** Get the const derived data for this texture on this platform. */
 	ENGINE_API const FTexturePlatformData* GetPlatformData() const;
 
+	UTextureAllMipDataProviderFactory* GetAllMipProvider() const
+	{
+		return const_cast<UTexture2D*>(this)->GetAssetUserData<UTextureAllMipDataProviderFactory>();
+	}
+
 #if WITH_EDITOR
 	/* cooked platform data for this texture */
 	TMap<FString, FTexturePlatformData*> CookedPlatformData;
@@ -197,6 +203,23 @@ public:
 	 */
 	ENGINE_API void GetMipData(int32 FirstMipToLoad, void** OutMipData);
 
+	/**
+	 * Retrieve initial texel data for mips, starting from FirstMipToLoad, up to the last mip in the texture.
+	 *    This function will only return mips that are currently loaded, or available via derived data.
+	 *    THIS FUNCTION WILL FAIL if you include in the requested range a streaming mip that is NOT currently loaded.
+	 *    Also note that ALL mip bulk data buffers are discarded (either returned or freed).
+	 *    So... in short, this function should only ever be called once.
+	 *
+	 * @param FirstMipToLoad - The first mip index to load.
+	 * @param OutMipData -	A pre-allocated array of pointers, that correspond to [FirstMipToLoad, .... LastMip]
+	 *						Upon successful return, each of those pointers will point to allocated memory containing the corresponding mip's data.
+	 *						Caller takes responsibility to free that memory.
+	 * @param OutMipSize -  A pre-allocated array of int64, that should be the same size as OutMipData (or zero size if caller does not require the sizes to be returned)
+	 *						Upon successful return, each element contains the size of the corresponding mip's data buffer.
+	 * @returns true if the requested mip data has been successfully returned.
+	 */
+	virtual bool GetInitialMipData(int32 InFirstMipToLoad, TArrayView<void*> OutMipData, TArrayView<int64> OutMipSize);
+
 	/**
 	 * Calculates the size of this texture in bytes if it had MipCount miplevels streamed in.
 	 *
diff --git a/Engine/Source/Runtime/Engine/Classes/Engine/TextureAllMipDataProviderFactory.h b/Engine/Source/Runtime/Engine/Classes/Engine/TextureAllMipDataProviderFactory.h
new file mode 100644
--- /dev/null
+++ b/Engine/Source/Runtime/Engine/Classes/Engine/TextureAllMipDataProviderFactory.h
@@ -0,0 +1,55 @@
+// Copyright Epic Games, Inc. All Rights Reserved.
+
+/*=============================================================================
+TextureMipDataProviderFactory.h: base class to create custom FTextureMipDataProvider.
+=============================================================================*/
+
+#pragma once
+#include "CoreMinimal.h"
+#include "Engine/AssetUserData.h"
+#include "Engine/TextureMipDataProviderFactory.h"
+#include "Streaming/StreamableRenderResourceState.h"
+#include "TextureAllMipDataProviderFactory.generated.h"
+
+class UTexture;
+class FTextureMipDataProvider;
+
+/**
+* UTextureAllMipDataProviderFactory defines an interface to create instances of FTextureMipDataProvider.
+* Derived classes from UTextureAllMipDataProviderFactory can be attached to UTexture2D
+* to define a new source for all of the mip data (instead of the default disk file or ddc mips). 
+* Use cases include custom texture compression.
+*/
+
+UCLASS(abstract, hidecategories=Object)
+class ENGINE_API UTextureAllMipDataProviderFactory : public UTextureMipDataProviderFactory
+{
+	GENERATED_UCLASS_BODY()
+
+public:
+
+	/**
+	 * Retrieve initial texel data for mips, starting from FirstMipToLoad, up to the last mip in the texture.
+	 *    This function will only return mips that are currently loaded in Bulk Data.
+	 *    THIS FUNCTION WILL FAIL if you include in the requested range a bulk data mip that is NOT currently loaded.
+	 *    Also note that only inline mips are loaded initially, and that ALL mip data buffers are discarded (either returned or freed).
+	 *    So... in short, this function should only ever be called once.
+	 *
+	 * @param FirstMipToLoad - The first mip index to load.
+	 * @param OutMipData -	A pre-allocated array of pointers, that correspond to [FirstMipToLoad, .... LastMip]
+	 *						Upon successful return, each of those pointers will point to allocated memory containing the corresponding mip's data.
+	 *						Caller takes responsibility to free that memory.
+	 * @param OutMipSize -  A pre-allocated array of int64, that should be the same size as OutMipData (or zero size if caller does not require the sizes to be returned)
+	 *						Upon successful return, each element contains the size of the corresponding mip's data buffer.
+	 * @param DebugContext - A string used for debug tracking and logging. Usually Texture->GetPathName()
+	 * @returns true if the requested mip data has been successfully returned.
+	 */
+	virtual bool GetInitialMipData(int32 FirstMipToLoad, TArrayView<void*> OutMipData, TArrayView<int64> OutMipSize, FStringView DebugContext)
+		PURE_VIRTUAL(UTextureAllMipDataProviderFactory::GetInitialMipData, return false;);
+
+	/**
+	  * Get the initial streaming state (after texture is first loaded)
+	  */
+	virtual FStreamableRenderResourceState GetResourcePostInitState(const UTexture* Owner, bool bAllowStreaming)
+		PURE_VIRTUAL(UTextureAllMipDataProviderFactory::GetResourcePostInitState, return FStreamableRenderResourceState(););
+};
diff --git a/Engine/Source/Runtime/Engine/Private/Rendering/Texture2DResource.cpp b/Engine/Source/Runtime/Engine/Private/Rendering/Texture2DResource.cpp
--- a/Engine/Source/Runtime/Engine/Private/Rendering/Texture2DResource.cpp
+++ b/Engine/Source/Runtime/Engine/Private/Rendering/Texture2DResource.cpp
@@ -33,7 +33,14 @@ FTexture2DResource::FTexture2DResource(UTexture2D* InOwner, const FStreamableRen
 {
 	// Retrieve initial mip data.
 	MipData.AddZeroed(State.MaxNumLODs);
-	InOwner->GetMipData(State.LODCountToAssetFirstLODIdx(State.NumRequestedLODs), MipData.GetData() + State.LODCountToFirstLODIdx(State.NumRequestedLODs));
+	MipDataSize.AddZeroed(State.MaxNumLODs);
+
+	int32 AssetFirstMipToLoad = State.LODCountToAssetFirstLODIdx(State.NumRequestedLODs);	// takes lod bias into account
+	int32 LocalFirstMipToLoad = State.LODCountToFirstLODIdx(State.NumRequestedLODs);		// ignores bias
+	InOwner->GetInitialMipData(
+		AssetFirstMipToLoad,
+		TArrayView<void*>(MipData).RightChop(LocalFirstMipToLoad),
+		TArrayView<int64>(MipDataSize).RightChop(LocalFirstMipToLoad));
 
 	CacheSamplerStateInitializer(InOwner);
 }
@@ -71,6 +78,7 @@ FTexture2DResource::~FTexture2DResource()
 			FMemory::Free( MipData[MipIndex] );
 		}
 		MipData[MipIndex] = NULL;
+		MipDataSize[MipIndex] = 0;
 	}
 }
 
@@ -129,6 +137,7 @@ void FTexture2DResource::CreateTexture()
 		for( int32 MipIndex=0; MipIndex < MipData.Num(); MipIndex++ )
 		{
 			MipData[MipIndex] = nullptr;
+			MipDataSize[MipIndex] = 0;
 		}
 	}
 	else
@@ -264,18 +273,18 @@ void FTexture2DResource::GetData( uint32 MipIndex, void* Dest, uint32 DestPitch,
 	uint32 SrcPitch;
 	uint32 EffectiveSize = CalculateTightPackedMipSize(MipMap.SizeX,MipMap.SizeY,PixelFormat,SrcPitch);
 
-	int64 BulkDataSize = MipMap.BulkData.GetBulkDataSize();
+	int64 DataSize = MipDataSize[MipIndex];
 
 	UE_LOG(LogTextureUpload,Verbose,TEXT("Size: %dx%d , EffectiveSize=%d BulkDataSize=%lld , SrcPitch=%d DestPitch=%d Format=%d, DestSize=%lld"),
 		MipMap.SizeX,MipMap.SizeY,
-		EffectiveSize,BulkDataSize,
+		EffectiveSize,DataSize,
 		SrcPitch,DestPitch, PixelFormat, DestSize);
 
-	if ((uint64)BulkDataSize > DestSize)
+	if ((uint64)DataSize > DestSize)
 	{
 		UE_LOG(LogTextureUpload, Warning, TEXT("DestSize is reported smaller than BulkDataSize in upload! Either RHI is reporting wrong sizes or we are stomping memory! Size: %dx%d, EffectiveSize=%d, BulkDataSize=%lld, SrcPitch=%d, DestPitch=%d, Format=%d, DestSize=%lld"),
 			MipMap.SizeX,MipMap.SizeY,
-			EffectiveSize,BulkDataSize,
+			EffectiveSize,DataSize,
 			SrcPitch,DestPitch, PixelFormat, DestSize);
 	}
 
@@ -283,9 +292,7 @@ void FTexture2DResource::GetData( uint32 MipIndex, void* Dest, uint32 DestPitch,
 	// runtime block size checking, conversion, or the like
 	if (DestPitch == 0)
 	{
-		// check( BulkDataSize >= EffectiveSize );
-
-		FMemory::Memcpy(Dest, MipData[MipIndex], BulkDataSize);
+		FMemory::Memcpy(Dest, MipData[MipIndex], DataSize);
 	}
 	else
 	{
@@ -293,17 +300,17 @@ void FTexture2DResource::GetData( uint32 MipIndex, void* Dest, uint32 DestPitch,
 		// in Editor, Mip doesn't come from BulkData, it may be null
 		// MipData[] was set from Editor data
 		// would be nice to check MipData[MipIndex] size ! but it's not stored
-		if ( BulkDataSize == 0 )
+		if (DataSize == 0)
 		{
-			BulkDataSize = EffectiveSize;
+			DataSize = EffectiveSize;
 		}
 		#endif
 
 	#if !WITH_EDITORONLY_DATA
 		// only checking when !WITH_EDITORONLY_DATA ? because in Editor BulkDataSize == 0 , so not possible to check
-		checkf((int64)EffectiveSize == BulkDataSize, 
-			TEXT("Texture '%s', mip %d, has a BulkDataSize [%d] that doesn't match calculated size [%d]. Texture size %dx%d, format %d"),
-			*TextureName.ToString(), MipIndex, BulkDataSize, EffectiveSize, GetSizeX(), GetSizeY(), (int32)PixelFormat);
+		checkf((int64)EffectiveSize == DataSize, 
+			TEXT("Texture '%s', mip %d, has a DataSize [%d] that doesn't match calculated size [%d]. Texture size %dx%d, format %d"),
+			*TextureName.ToString(), MipIndex, DataSize, EffectiveSize, GetSizeX(), GetSizeY(), (int32)PixelFormat);
 	#endif
 
 		// for platforms that returned 0 pitch from Lock, we need to just use the bulk data directly, never do 
@@ -311,7 +318,7 @@ void FTexture2DResource::GetData( uint32 MipIndex, void* Dest, uint32 DestPitch,
 		if (DestPitch == 0 || DestPitch == SrcPitch )
 		{
 			// checking Dest size before we memcpy would be nice!
-			FMemory::Memcpy(Dest, MipData[MipIndex], BulkDataSize);
+			FMemory::Memcpy(Dest, MipData[MipIndex], DataSize);
 		}
 		else
 		{
@@ -323,4 +330,5 @@ void FTexture2DResource::GetData( uint32 MipIndex, void* Dest, uint32 DestPitch,
 	// Free data retrieved via GetCopy inside constructor.
 	FMemory::Free(MipData[MipIndex]);
 	MipData[MipIndex] = NULL;
+	MipDataSize[MipIndex] = 0;
 }
diff --git a/Engine/Source/Runtime/Engine/Private/Streaming/TextureAllMipDataProviderFactory.cpp b/Engine/Source/Runtime/Engine/Private/Streaming/TextureAllMipDataProviderFactory.cpp
new file mode 100644
--- /dev/null
+++ b/Engine/Source/Runtime/Engine/Private/Streaming/TextureAllMipDataProviderFactory.cpp
@@ -0,0 +1,15 @@
+// Copyright Epic Games, Inc. All Rights Reserved.
+
+/*=============================================================================
+	TextureAllMipDataProviderFactory.cpp: Implementation of TextureAllMipDataProviderFactory.
+=============================================================================*/
+
+#include "Engine/TextureAllMipDataProviderFactory.h"
+
+#include UE_INLINE_GENERATED_CPP_BY_NAME(TextureAllMipDataProviderFactory)
+
+UTextureAllMipDataProviderFactory::UTextureAllMipDataProviderFactory(const FObjectInitializer& ObjectInitializer)
+	: Super(ObjectInitializer)
+{
+}
+
diff --git a/Engine/Source/Runtime/Engine/Private/Texture2D.cpp b/Engine/Source/Runtime/Engine/Private/Texture2D.cpp
--- a/Engine/Source/Runtime/Engine/Private/Texture2D.cpp
+++ b/Engine/Source/Runtime/Engine/Private/Texture2D.cpp
@@ -143,13 +143,21 @@ extern bool TrackTextureEvent( FStreamingRenderAsset* StreamingTexture, UStreama
 static_assert(sizeof(FTexture2DMipMap) <= 80, "FTexture2DMipMap was packed to reduce its memory footprint and fit into an 80 bytes bin of MallocBinned2/3");
 #endif
 
-void FTexture2DMipMap::Serialize(FArchive& Ar, UObject* Owner, int32 MipIdx)
+void FTexture2DMipMap::Serialize(FArchive& Ar, UObject* Owner, int32 MipIdx, bool bSerializeMipData)
 {
+	if (bSerializeMipData)
+	{
 #if WITH_EDITORONLY_DATA
-	BulkData.Serialize(Ar, Owner, MipIdx, false, FileRegionType);
+		BulkData.Serialize(Ar, Owner, MipIdx, false, FileRegionType);
 #else
-	BulkData.Serialize(Ar, Owner, MipIdx, false);
+		BulkData.Serialize(Ar, Owner, MipIdx, false);
 #endif
+	}
+	else if (Ar.IsLoading())
+	{
+		// in case we're deserializing into an existing object, clear out BulkData
+		BulkData.RemoveBulkData();
+	}
 
 	int32 XSize = this->SizeX;
 	int32 YSize = this->SizeY;
@@ -453,7 +461,18 @@ void UTexture2D::Serialize(FArchive& Ar)
 
 	if (Ar.IsCooking() || bCooked)
 	{
-		SerializeCookedPlatformData(Ar);
+		bool bSerializeMipData = true;
+
+		if (Ar.IsSaving())
+		{
+			// if there is an all mip provider, then we don't need to serialize the PlatformData mip data
+			bSerializeMipData = (GetAllMipProvider() == nullptr);
+		}
+
+		// since the binary serialization format depends on this bool, we need to serialize it to know what format to expect at load time
+		Ar << bSerializeMipData;
+		
+		SerializeCookedPlatformData(Ar, bSerializeMipData);
 	}
 }
 
@@ -1014,7 +1033,17 @@ FTextureResource* UTexture2D::CreateResource()
 			else
 			{
 				// Should be as big as the mips we have already directly loaded into GPU mem
-				const FStreamableRenderResourceState PostInitState = GetResourcePostInitState(GetPlatformData(), !bTemporarilyDisableStreaming, ResourceMem ? ResourceMem->GetNumMips() : 0, NumMips);
+				FStreamableRenderResourceState PostInitState;
+				if (UTextureAllMipDataProviderFactory* ProviderFactory = GetAllMipProvider())
+				{
+					// All Mip Providers get to control the initial streaming state
+					PostInitState = ProviderFactory->GetResourcePostInitState(this, !bTemporarilyDisableStreaming);
+				}
+				else
+				{
+					PostInitState = GetResourcePostInitState(GetPlatformData(), !bTemporarilyDisableStreaming, ResourceMem ? ResourceMem->GetNumMips() : 0, NumMips);
+				}
+
 				FTexture2DResource* Texture2DResource = new FTexture2DResource(this, PostInitState);
 				// preallocated memory for the UTexture2D resource is now owned by this resource
 				// and will be freed by the RHI resource or when the FTexture2DResource is deleted
diff --git a/Engine/Source/Runtime/Engine/Private/TextureDerivedData.cpp b/Engine/Source/Runtime/Engine/Private/TextureDerivedData.cpp
--- a/Engine/Source/Runtime/Engine/Private/TextureDerivedData.cpp
+++ b/Engine/Source/Runtime/Engine/Private/TextureDerivedData.cpp
@@ -2414,8 +2414,8 @@ static void SerializePlatformData(
 	FArchive& Ar,
 	FTexturePlatformData* PlatformData,
 	UTexture* Texture,
-	EPlatformDataSerializationFlags Flags
-)
+	EPlatformDataSerializationFlags Flags,
+	const bool bSerializeMipData)
 {
 	// note: if BuildTexture failed, we still get called here,
 	//	just with a default-constructed PlatformData
@@ -2514,6 +2514,10 @@ static void SerializePlatformData(
 		#endif
 
 		#if WITH_EDITOR
+			// TODO [chris.tchou] : we should probably query the All Mip Provider to see what streaming state is.
+			// Otherwise we might disable streaming calculations, even though the AMP expects to stream.
+			// (this feeds into the bHasNoStreamableTextures optimization that skips streaming calculations)
+
 			// Record the use of streaming mips on the owner.
 			if (NumNonStreamingMips < NumMips)
 			{
@@ -2713,20 +2717,11 @@ static void SerializePlatformData(
 				check(Mip.DerivedData);
 			}
 
+			// from the first inline mip onwards, we serialize to inline bulk data
 			for (int32 MipIndex = FirstInlineMip; MipIndex < NumMips; ++MipIndex)
 			{
 				FTexture2DMipMap& Mip = PlatformData->Mips[FirstMipToSerialize + MipIndex];
-				if (Ar.IsSaving())
-				{
-					const uint32 OriginalBulkDataFlags = Mip.BulkData.GetBulkDataFlags();
-					Mip.BulkData.SetBulkDataFlags(BULKDATA_ForceInlinePayload | BULKDATA_SingleUse);
-					Mip.BulkData.Serialize(Ar, Texture, MipIndex);
-					Mip.BulkData.ResetBulkDataFlags(OriginalBulkDataFlags);
-				}
-				else
-				{
-					Mip.BulkData.Serialize(Ar, Texture, MipIndex);
-				}
+				Mip.BulkData.SerializeWithFlags(Ar, Texture, BULKDATA_ForceInlinePayload | BULKDATA_SingleUse);
 			}
 		}
 
@@ -2844,7 +2839,7 @@ static void SerializePlatformData(
 
 	for (int32 MipIndex = 0; MipIndex < NumMips; ++MipIndex)
 	{
-		PlatformData->Mips[FirstMipToSerialize + MipIndex].Serialize(Ar, Texture, MipIndex);
+		PlatformData->Mips[FirstMipToSerialize + MipIndex].Serialize(Ar, Texture, MipIndex, bSerializeMipData);
 	}
 
 	Ar << bIsVirtual;
@@ -2876,7 +2871,8 @@ static void SerializePlatformData(
 
 void FTexturePlatformData::Serialize(FArchive& Ar, UTexture* Owner)
 {
-	SerializePlatformData(Ar, this, Owner, EPlatformDataSerializationFlags::None);
+	check(!Ar.IsCooking()); // this is not the path that handles serialization for cooking
+	SerializePlatformData(Ar, this, Owner, EPlatformDataSerializationFlags::None, /* bSerializeMipData = */ true);
 }
 
 #if WITH_EDITORONLY_DATA
@@ -2894,14 +2890,14 @@ UE::DerivedData::FValueId FTexturePlatformData::MakeMipId(int32 MipIndex)
 
 #endif // WITH_EDITORONLY_DATA
 
-void FTexturePlatformData::SerializeCooked(FArchive& Ar, UTexture* Owner, bool bStreamable)
+void FTexturePlatformData::SerializeCooked(FArchive& Ar, UTexture* Owner, bool bStreamable, const bool bSerializeMipData)
 {
 	EPlatformDataSerializationFlags Flags = EPlatformDataSerializationFlags::Cooked;
 	if (bStreamable)
 	{
 		Flags |= EPlatformDataSerializationFlags::Streamable;
 	}
-	SerializePlatformData(Ar, this, Owner, Flags);
+	SerializePlatformData(Ar, this, Owner, Flags, bSerializeMipData);
 	if (Ar.IsLoading())
 	{
 		// Patch up Size as due to mips being stripped out during cooking it could be wrong.
@@ -2968,7 +2964,26 @@ bool UTexture2DArray::GetMipData(int32 InFirstMipToLoad, TArray<FUniqueBuffer, T
 
 void UTexture2D::GetMipData(int32 FirstMipToLoad, void** OutMipData)
 {
-	if (GetPlatformData()->TryLoadMips(FirstMipToLoad, OutMipData, GetPathName()) == false)
+	// hack convert the unsafe inputs to the "safe" form
+	// here we are hoping that the caller has allocated this number of elements in OutMipData...  :fingers_crossed:
+	int32 NumberOfMipsToLoad = GetPlatformData()->Mips.Num() - FirstMipToLoad;
+	TArrayView<int64> MipSizeView; // empty array - we don't need the sizes returned
+	GetInitialMipData(FirstMipToLoad, MakeArrayView(OutMipData, NumberOfMipsToLoad), MipSizeView);
+}
+
+bool UTexture2D::GetInitialMipData(int32 FirstMipToLoad, TArrayView<void*> OutMipData, TArrayView<int64> OutMipSize)
+{
+	bool bLoaded = false;
+	if (UTextureAllMipDataProviderFactory* ProviderFactory = GetAllMipProvider())
+	{
+		bLoaded = ProviderFactory->GetInitialMipData(FirstMipToLoad, OutMipData, OutMipSize, GetPathName());
+	}
+	else
+	{
+		bLoaded = GetPlatformData()->TryLoadMipsWithSizes(FirstMipToLoad, OutMipData.GetData(), OutMipSize.GetData(), GetPathName());
+	}
+
+	if (!bLoaded)
 	{
 		// Unable to load mips from the cache. Rebuild the texture and try again.
 		UE_LOG(LogTexture,Warning,TEXT("GetMipData failed for %s (%s)"),
@@ -2977,13 +2992,14 @@ void UTexture2D::GetMipData(int32 FirstMipToLoad, void** OutMipData)
 		if (!GetOutermost()->bIsCookedForEditor)
 		{
 			ForceRebuildPlatformData();
-			if (GetPlatformData()->TryLoadMips(FirstMipToLoad, OutMipData, GetPathName()) == false)
+			if (GetPlatformData()->TryLoadMipsWithSizes(FirstMipToLoad, OutMipData.GetData(), OutMipSize.GetData(), GetPathName()) == false)
 			{
 				UE_LOG(LogTexture, Error, TEXT("TryLoadMips still failed after ForceRebuildPlatformData %s."), *GetPathName());
 			}
 		}
 #endif // #if WITH_EDITOR
 	}
+	return bLoaded;
 }
 
 void UTextureCube::GetMipData(int32 FirstMipToLoad, void** OutMipData)
@@ -3614,7 +3630,7 @@ void UTexture::CleanupCachedRunningPlatformData()
 }
 
 
-void UTexture::SerializeCookedPlatformData(FArchive& Ar)
+void UTexture::SerializeCookedPlatformData(FArchive& Ar, const bool bSerializeMipData)
 {
 	if (IsTemplate() )
 	{
@@ -3710,11 +3726,13 @@ void UTexture::SerializeCookedPlatformData(FArchive& Ar)
 			// this iteration is over NumLayers :
 			for (FTexturePlatformData* PlatformDataToSave : PlatformDataToSerialize)
 			{
+				// wait for async build task to complete, if there is one
 				PlatformDataToSave->FinishCache();
 
 				FName PixelFormatName = PixelFormatEnum->GetNameByValue(PlatformDataToSave->PixelFormat);
 				Ar << PixelFormatName;
 
+				// reserve space in the archive to record the skip offset
 				const int64 SkipOffsetLoc = Ar.Tell();
 				int64 SkipOffset = 0;
 				{
@@ -3723,8 +3741,11 @@ void UTexture::SerializeCookedPlatformData(FArchive& Ar)
 
 				// Pass streamable flag for inlining mips
 				bool bTextureIsStreamable = GetTextureIsStreamableOnPlatform(*this, *Ar.CookingTarget());
-				PlatformDataToSave->SerializeCooked(Ar, this, bTextureIsStreamable);
 
+				// serialize the platform data
+				PlatformDataToSave->SerializeCooked(Ar, this, bTextureIsStreamable, bSerializeMipData);
+
+				// go back and patch the skip offset
 				SkipOffset = Ar.Tell() - SkipOffsetLoc;
 				Ar.Seek(SkipOffsetLoc);
 				Ar << SkipOffset;
@@ -3763,7 +3784,7 @@ void UTexture::SerializeCookedPlatformData(FArchive& Ar)
 			{
 				// Extra arg is unused here because we're loading
 				const bool bStreamable = false;
-				RunningPlatformData->SerializeCooked(Ar, this, bStreamable);
+				RunningPlatformData->SerializeCooked(Ar, this, bStreamable, bSerializeMipData);
 			}
 			else
 			{
diff --git a/Engine/Source/Runtime/Engine/Public/Rendering/Texture2DResource.h b/Engine/Source/Runtime/Engine/Public/Rendering/Texture2DResource.h
--- a/Engine/Source/Runtime/Engine/Public/Rendering/Texture2DResource.h
+++ b/Engine/Source/Runtime/Engine/Public/Rendering/Texture2DResource.h
@@ -73,6 +73,7 @@ private:
 
 	/** Local copy/ cache of mip data between creation and first call to InitRHI.							*/
 	TArray<void*, TInlineAllocator<MAX_TEXTURE_MIP_COUNT> > MipData;
+	TArray<int64, TInlineAllocator<MAX_TEXTURE_MIP_COUNT> > MipDataSize;
 
 	/**
 	 * Writes the data for a single mip-level into a destination buffer.
diff --git a/Engine/Source/Runtime/Engine/Public/Streaming/TextureMipDataProvider.h b/Engine/Source/Runtime/Engine/Public/Streaming/TextureMipDataProvider.h
--- a/Engine/Source/Runtime/Engine/Public/Streaming/TextureMipDataProvider.h
+++ b/Engine/Source/Runtime/Engine/Public/Streaming/TextureMipDataProvider.h
@@ -118,7 +118,7 @@ public:
 	FORCEINLINE ETickThread GetNextTickThread() const { return NextTickThread; }
 
 	/**
-	* Initialize data prelimary to the GetMips() step. Can be called several time. Mostly useful to simplify the logic in GetMips().
+	* Initialize data prelimary to the GetMips() step. Can be called several time (it does not have to advance to GetMips immediately). Mostly useful to simplify the logic in GetMips().
 	* This is because GetMips is a chained call between all mip data providers, each taking some mips to handle, and is not compatible with multi step process.
 	* This means that GetMips() must return immediately and can not postpone or delay return by not advancing to the next steps.
 	*
@@ -141,7 +141,7 @@ public:
 	virtual int32 GetMips(const FTextureUpdateContext& Context, int32 StartingMipIndex, const FTextureMipInfoArray& MipInfos, const FTextureUpdateSyncOptions& SyncOptions) = 0;
 
 	/**
-	* Check if each mip handled by this mip data provider have been updated correctly. Must move to CleanUp() when done.
+	* Check if each mip handled by this mip data provider have been updated correctly. Must move to CleanUp() or Done() when done.
 	*
 	* @param SyncOptions - Different sync options to control when the next tick of FTextureStreamIn can be scheduled.
 	*
diff --git a/Engine/Source/Runtime/Engine/Public/TextureResource.h b/Engine/Source/Runtime/Engine/Public/TextureResource.h
--- a/Engine/Source/Runtime/Engine/Public/TextureResource.h
+++ b/Engine/Source/Runtime/Engine/Public/TextureResource.h
@@ -70,7 +70,7 @@ struct FTexture2DMipMap
 	PRAGMA_ENABLE_DEPRECATION_WARNINGS
 
 	/** Serialization. */
-	ENGINE_API void Serialize(FArchive& Ar, UObject* Owner, int32 MipIndex);
+	ENGINE_API void Serialize(FArchive& Ar, UObject* Owner, int32 MipIndex, bool bSerializeMipData);
 
 #if WITH_EDITORONLY_DATA
 	/** The file region type appropriate for the pixel format of this mip-map. */
diff --git a/Engine/Source/Runtime/Landscape/Classes/LandscapeProxy.h b/Engine/Source/Runtime/Landscape/Classes/LandscapeProxy.h
--- a/Engine/Source/Runtime/Landscape/Classes/LandscapeProxy.h
+++ b/Engine/Source/Runtime/Landscape/Classes/LandscapeProxy.h
@@ -809,6 +809,10 @@ public:
 	UPROPERTY()
 	bool bHasLayersContent;
 
+	/** Enable compressed heightmap texture storage. */
+	UPROPERTY(EditAnywhere, Category = Landscape, AdvancedDisplay)
+	bool bUseCompressedHeightmapStorage = false;
+
 #if WITH_EDITOR
 	LANDSCAPE_API static ULandscapeLayerInfoObject* VisibilityLayer;
 
diff --git a/Engine/Source/Runtime/Landscape/Classes/LandscapeTextureStorageProvider.h b/Engine/Source/Runtime/Landscape/Classes/LandscapeTextureStorageProvider.h
new file mode 100644
--- /dev/null
+++ b/Engine/Source/Runtime/Landscape/Classes/LandscapeTextureStorageProvider.h
@@ -0,0 +1,165 @@
+// Copyright Epic Games, Inc. All Rights Reserved.
+
+#pragma once
+
+#include "CoreMinimal.h"
+#include "Engine/TextureMipDataProviderFactory.h"
+#include "Streaming/TextureMipDataProvider.h"
+#include "Containers/ContainersFwd.h"
+#include "Algo/Accumulate.h"
+#include "Serialization/BulkData.h"
+#include "Engine/TextureAllMipDataProviderFactory.h"
+#include "LandscapeTextureStorageProvider.generated.h"
+
+class ULandscapeTextureStorageProviderFactory;
+class UTexture2D;
+
+USTRUCT()
+struct FLandscapeTexture2DMipMap
+{
+	GENERATED_BODY()
+
+	/** Width of the mip-map. */
+	UPROPERTY()
+	int32 SizeX = 0;
+	/** Height of the mip-map. */
+	UPROPERTY()
+	int32 SizeY = 0;
+	/** Whether the bulk data is compressed or not */
+	UPROPERTY()
+	bool bCompressed = false;
+
+	FByteBulkData BulkData;
+
+	void Serialize(FArchive& Ar, UObject* Owner, uint32 SaveOverrideFlags);
+};
+
+class FLandscapeTextureStorageMipProvider : public FTextureMipDataProvider
+{
+public:
+	FLandscapeTextureStorageMipProvider(ULandscapeTextureStorageProviderFactory* InFactory);
+	~FLandscapeTextureStorageMipProvider();
+
+	// ********************************************************
+	// ********* FLandscapeTextureStorageMipProvider implementation **********
+	// ********************************************************
+
+	void Init(const FTextureUpdateContext& Context, const FTextureUpdateSyncOptions& SyncOptions) final override;
+	int32 GetMips(const FTextureUpdateContext& Context, int32 StartingMipIndex, const FTextureMipInfoArray& MipInfos, const FTextureUpdateSyncOptions& SyncOptions) final override;
+	bool PollMips(const FTextureUpdateSyncOptions& SyncOptions) final override;
+	void AbortPollMips() final override;
+	void CleanUp(const FTextureUpdateSyncOptions& SyncOptions) final override;
+	void Cancel(const FTextureUpdateSyncOptions& SyncOptions) final override;
+	ETickThread GetCancelThread() const final override;
+
+private:
+	// the Factory (that actually has the mip data)
+	ULandscapeTextureStorageProviderFactory* Factory = nullptr;
+
+	// A structured with information about which file contains which mips.
+	struct FIORequest
+	{
+		FIoFilenameHash FilenameHash = INVALID_IO_FILENAME_HASH;
+		TUniquePtr<IBulkDataIORequest> BulkDataIORequest;
+		uint8* DestMipData = nullptr;
+	};
+
+	// Pending async requests created in GetMips().
+	TArray<FIORequest, TInlineAllocator<MAX_TEXTURE_MIP_COUNT>> IORequests;
+
+	int32 FirstRequestedMipIndex = INDEX_NONE;
+
+	// The asset name, used to log IO errors.
+	FName TextureName;
+	// Whether async read requests must be created with high priority (executes faster). 
+	bool bPrioritizedIORequest = false;
+	// Whether async read requests where cancelled for any reasons.
+	bool bIORequestCancelled = false;
+	// Whether async read requests were required to abort through AbortPollMips().
+	bool bIORequestAborted = false;
+
+	// A callback to be executed once all IO pending requests are completed.
+	FBulkDataIORequestCallBack AsyncFileCallBack;
+
+	// Helper to create the AsyncFileCallBack -- which can handle IO Request completions (and will trigger the next step to be scheduled via SyncOptions.RescheduleCallback())
+	void CreateAsyncFileCallback(const FTextureUpdateSyncOptions& SyncOptions);
+
+	void ClearIORequests();
+};
+
+
+
+UCLASS(MinimalAPI)		// MinimalAPI -- do not export methods (methods not available in other modules)
+class ULandscapeTextureStorageProviderFactory : public UTextureAllMipDataProviderFactory
+{
+	GENERATED_UCLASS_BODY()
+
+public:
+	
+	int32 NumNonOptionalMips = 0;
+	int32 NumNonStreamingMips = 0;
+	FVector LandscapeGridScale = FVector(ForceInit);
+
+	TArray<FLandscapeTexture2DMipMap> Mips;
+	TObjectPtr<UTexture2D> Texture;
+
+	#if WITH_EDITORONLY_DATA
+	static ULandscapeTextureStorageProviderFactory* ApplyTo(UTexture2D* TargetTexture, const FVector& InLandsapeGridScale);
+	#endif // WITH_EDITORONLY_DATA
+
+	// compress the mip data.
+	static void CompressMipToBulkData(int32 MipIndex, int32 MipSizeX, int32 MipSizeY, uint8* SourceData, int32 SourceDataBytes, FByteBulkData& DestBulkData);
+	static void CopyMipToBulkData(int32 MipIndex, int32 MipSizeX, int32 MipSizeY, uint8* SourceData, int32 SourceDataBytes, FByteBulkData& DestBulkData);
+
+	// decompress the mip data. Currently operates in-place (though this could be changed in the future)
+	void DecompressMip(uint8* SourceData, int64 SourceDataBytes, uint8* DestData, int64 DestDataBytes, int32 MipIndex);
+
+	virtual void Serialize(FArchive& Ar) override;
+
+	/**
+	* Create a FTextureMipDataProvider to handle a single StreamIn mip operation.
+	* The object lifetime will be managed by FRenderAssetUpdate.
+	*
+	* @param InAsset - the texture on which the stream in operation will be performed.
+	*/
+	virtual FTextureMipDataProvider* AllocateMipDataProvider(UTexture* Asset) override
+	{
+		return new FLandscapeTextureStorageMipProvider(this);
+	}
+
+	/**
+	* Returns true if TextureMipDataProviders allocated by this factory can provide MipData by themselves,
+	* even without loading from disk at all, so streaming can be enabled for their textures.
+	*/
+	virtual bool WillProvideMipDataWithoutDisk() const
+	{
+		return true;
+	}
+
+	// load mip data (used to access preloaded inline mips for creating the initial version of the texture)
+	virtual bool GetInitialMipData(int32 FirstMipToLoad, TArrayView<void*> OutMipData, TArrayView<int64> OutMipSize, FStringView DebugContext) override;
+
+	// return the initial streaming state of the texture
+	virtual FStreamableRenderResourceState GetResourcePostInitState(const UTexture* Owner, bool bAllowStreaming) override;
+
+	bool DoesMipDataExist(int32 MipIndex)
+	{
+		// we need to check the bulk data exists (might not for optional mips)
+		return Mips.IsValidIndex(MipIndex) && Mips[MipIndex].BulkData.DoesExist();
+	}
+
+	const FLandscapeTexture2DMipMap* GetMip(int32 MipIndex)
+	{
+		int32 Index = MipIndex;
+		if (Mips.IsValidIndex(Index))
+		{
+			return &Mips[Index];
+		}
+		return nullptr;
+	}
+
+	int32 GetTotalBytes() const
+	{
+		return Algo::TransformAccumulate(Mips, [](const FLandscapeTexture2DMipMap& Mip) -> int32 { return Mip.BulkData.GetBulkDataSize(); }, 0);
+	}
+};
diff --git a/Engine/Source/Runtime/Landscape/Private/Landscape.cpp b/Engine/Source/Runtime/Landscape/Private/Landscape.cpp
--- a/Engine/Source/Runtime/Landscape/Private/Landscape.cpp
+++ b/Engine/Source/Runtime/Landscape/Private/Landscape.cpp
@@ -74,6 +74,7 @@ Landscape.cpp: Terrain rendering
 #include "VisualLogger/VisualLogger.h"
 #include "NaniteSceneProxy.h"
 #include "Misc/ArchiveMD5.h"
+#include "LandscapeTextureStorageProvider.h"
 #include "LandscapeVersion.h"
 #include "UObject/FortniteMainBranchObjectVersion.h"
 #include "UObject/FortniteReleaseBranchCustomObjectVersion.h"
@@ -3080,6 +3081,22 @@ void ALandscapeProxy::PreSave(FObjectPreSaveContext ObjectSaveContext)
 			LandscapeInfo->OnModifiedPackageSaved(GetPackage());
 		}
 	}
+
+	if (bUseCompressedHeightmapStorage && ObjectSaveContext.IsCooking())
+	{
+		FString PlatformName = ObjectSaveContext.GetTargetPlatform()->PlatformName();
+
+		// TODO [chris.tchou] : remove the limitation to only work on windows, once we validate other platforms work as well...
+		if (PlatformName.StartsWith(TEXT("Windows")))
+		{
+			for (ULandscapeComponent* LandscapeComponent : LandscapeComponents)
+			{
+				UTexture2D* Tex = LandscapeComponent->GetHeightmap();
+				FVector LandscapeGridScale = GetRootComponent()->GetRelativeScale3D();
+				ULandscapeTextureStorageProviderFactory::ApplyTo(Tex, LandscapeGridScale);
+			}
+		}
+	}
 #endif // WITH_EDITOR
 }
 
@@ -3707,6 +3724,7 @@ void ALandscapeProxy::GetSharedProperties(ALandscapeProxy* Landscape)
 		VirtualTextureNumLods = Landscape->VirtualTextureNumLods;
 		VirtualTextureRenderPassType = Landscape->VirtualTextureRenderPassType;
 		bEnableNanite = Landscape->bEnableNanite;
+		bUseCompressedHeightmapStorage = Landscape->bUseCompressedHeightmapStorage;
 #if WITH_EDITORONLY_DATA
 		bNaniteSkirtEnabled = Landscape->bNaniteSkirtEnabled;
 		NaniteSkirtDepth = Landscape->NaniteSkirtDepth;
diff --git a/Engine/Source/Runtime/Landscape/Private/LandscapeTextureStorageProvider.cpp b/Engine/Source/Runtime/Landscape/Private/LandscapeTextureStorageProvider.cpp
new file mode 100644
--- /dev/null
+++ b/Engine/Source/Runtime/Landscape/Private/LandscapeTextureStorageProvider.cpp
@@ -0,0 +1,793 @@
+// Copyright Epic Games, Inc. All Rights Reserved.
+
+/*=============================================================================
+LandscapeTextureStorageProvider.cpp: Alternative Texture Storage for Landscape Textures
+=============================================================================*/
+
+#include "LandscapeTextureStorageProvider.h"
+#include "LandscapeDataAccess.h"
+#include "RHIGlobals.h"
+#include "ContentStreaming.h"
+#include "LandscapePrivate.h"
+
+#include UE_INLINE_GENERATED_CPP_BY_NAME(LandscapeTextureStorageProvider)
+
+FLandscapeTextureStorageMipProvider::FLandscapeTextureStorageMipProvider(ULandscapeTextureStorageProviderFactory* InFactory)
+	: FTextureMipDataProvider(InFactory->Texture, ETickState::Init, ETickThread::Async)
+{
+	this->Factory = InFactory;
+
+	if (InFactory->Texture)
+	{
+		TextureName = InFactory->Texture->GetFName();
+	}
+}
+
+FLandscapeTextureStorageMipProvider::~FLandscapeTextureStorageMipProvider()
+{
+}
+
+ULandscapeTextureStorageProviderFactory::ULandscapeTextureStorageProviderFactory(const FObjectInitializer& ObjectInitializer)
+	: Super(ObjectInitializer)
+{
+}
+
+void FLandscapeTexture2DMipMap::Serialize(FArchive& Ar, UObject* Owner, uint32 SaveOverrideFlags)
+{
+	Ar << SizeX;
+	Ar << SizeY;
+	Ar << bCompressed;
+	BulkData.SerializeWithFlags(Ar, Owner, SaveOverrideFlags);
+}
+
+template<typename T, typename F>
+static bool SerializeArray(FArchive& Ar, TArray<T>& Array, F&& SerializeElementFn)
+{
+	int32 Num = Array.Num();
+	Ar << Num;
+	if (Ar.IsLoading())
+	{
+		if (Num < 0)
+		{
+			return false;
+		}
+		else
+		{
+			Array.SetNum(Num);
+			for (int32 Index = 0; Index < Num; ++Index)
+			{
+				SerializeElementFn(Ar, Index, Array[Index]);
+			}
+		}
+	}
+	else
+	{
+		for (int32 Index = 0; Index < Num; ++Index)
+		{
+			SerializeElementFn(Ar, Index, Array[Index]);
+		}
+	}
+	return true;
+}
+
+void ULandscapeTextureStorageProviderFactory::Serialize(FArchive& Ar)
+{
+	Super::Serialize(Ar);
+
+	//  mip 0                                                      mip N
+	// 	high rez <---------------------------------------------> low rez
+	// 	[ Optional Mips ][            Non Optional Mips                ]
+	// 	                 [ Streaming Mips ][ Non Streaming Inline Mips ]
+
+	int32 OptionalMips = Mips.Num() - NumNonOptionalMips;
+	check(OptionalMips >= 0);
+
+	int32 FirstInlineMip = Mips.Num() - NumNonStreamingMips;
+	check(FirstInlineMip >= 0);
+
+	Ar << NumNonOptionalMips;
+	Ar << NumNonStreamingMips;
+	Ar << LandscapeGridScale;
+
+	SerializeArray(Ar, Mips,
+		[this, OptionalMips, FirstInlineMip](FArchive& Ar, int32 Index, FLandscapeTexture2DMipMap& Mip)
+		{
+			// select bulk data flags for optional/streaming/inline mips
+			uint32 BulkDataFlags;
+			if (Index < OptionalMips)
+			{
+				// optional mip
+				BulkDataFlags = BULKDATA_Force_NOT_InlinePayload | BULKDATA_OptionalPayload;
+			}
+			else if (Index < FirstInlineMip)
+			{
+				// streaming mip
+				bool bDuplicateNonOptionalMips = false; // TODO [chris.tchou] : if we add support for optional mips, we might need to calculate this.
+				BulkDataFlags = BULKDATA_Force_NOT_InlinePayload | (bDuplicateNonOptionalMips ? BULKDATA_DuplicateNonOptionalPayload : 0);
+			}
+			else
+			{
+				// non streaming inline mip (can be single use as we only need to upload to GPU once, are never streamed out)
+				BulkDataFlags = BULKDATA_ForceInlinePayload | BULKDATA_SingleUse;
+			}
+			Mip.Serialize(Ar, this, BulkDataFlags);
+		});
+
+	Ar << Texture;
+}
+
+FStreamableRenderResourceState ULandscapeTextureStorageProviderFactory::GetResourcePostInitState(const UTexture* Owner, bool bAllowStreaming)
+{
+	// We are using the non-offline mode to upload these textures currently, so we don't need to consider mip tails.
+	// (RHI will handle it during upload, just less optimal than having them pre-packed)
+	// If we ever want to optimize the GPU upload by using the offline mode, we can add the logic here to take mip tails into account.
+	const int32 PlatformNumMipsInTail = 1;
+
+	const int32 TotalMips = Mips.Num();
+	const int32 ExpectedAssetLODBias = FMath::Clamp<int32>(Owner->GetCachedLODBias() - Owner->NumCinematicMipLevels, 0, TotalMips - 1);
+	const int32 MaxRuntimeMipCount = FMath::Min<int32>(GMaxTextureMipCount, FStreamableRenderResourceState::MAX_LOD_COUNT);
+
+	const int32 NumMips = FMath::Min<int32>(TotalMips - ExpectedAssetLODBias, MaxRuntimeMipCount);
+
+	bool bTextureIsStreamable = true;		// landscape texture storage is always streamable (we should not use it for platforms that are not)
+
+	// clamp non-optional and non-streaming mips to reflect potentially reduced mip count because of bias
+	const int32 BiasedNumNonOptionalMips = FMath::Min<int32>(NumMips, NumNonOptionalMips);
+	const int32 NumOfNonStreamingMips = FMath::Min<int32>(NumMips, NumNonStreamingMips);
+
+	// Optional mips must be streaming mips :
+	check(BiasedNumNonOptionalMips >= NumOfNonStreamingMips);
+
+	if (NumOfNonStreamingMips == NumMips)
+	{
+		bTextureIsStreamable = false;
+	}
+
+	const int32 AssetMipIdxForResourceFirstMip = FMath::Max<int32>(0, TotalMips - NumMips);
+
+	const bool bMakeStreamable = bAllowStreaming;
+	int32 NumRequestedMips = 0;
+	if (!bTextureIsStreamable)
+	{
+		// in Editor , NumOfNonStreamingMips may not be all mips
+		// but once we cook it will be
+		// so check this early to make behavior consistent
+		NumRequestedMips = NumMips;
+	}
+	else if (bMakeStreamable && IStreamingManager::Get().IsRenderAssetStreamingEnabled(EStreamableRenderAssetType::Texture))
+	{
+		NumRequestedMips = NumOfNonStreamingMips;
+	}
+	else
+	{
+		// we are not streaming (bMakeStreamable is false)
+		// but this may select a mip below the top mip
+		// (due to cinematic lod bias)
+		// but only if the texture itself is streamable
+
+		// Adjust CachedLODBias so that it takes into account FStreamableRenderResourceState::AssetLODBias.
+		const int32 ResourceLODBias = FMath::Max<int32>(0, Owner->GetCachedLODBias() - AssetMipIdxForResourceFirstMip);
+
+		// Ensure NumMipsInTail is within valid range to safeguard on the above expressions. 
+		const int32 NumMipsInTail = FMath::Clamp<int32>(PlatformNumMipsInTail, 1, NumMips);
+
+		// Bias is not allowed to shrink the mip count below NumMipsInTail.
+		NumRequestedMips = FMath::Max<int32>(NumMips - ResourceLODBias, NumMipsInTail);
+
+		// If trying to load optional mips, check if the first resource mip is available.
+		if (NumRequestedMips > BiasedNumNonOptionalMips && !DoesMipDataExist(AssetMipIdxForResourceFirstMip))
+		{
+			NumRequestedMips = BiasedNumNonOptionalMips;
+		}
+
+		// Ensure we don't request a top mip in the NonStreamingMips
+		NumRequestedMips = FMath::Max(NumRequestedMips, NumOfNonStreamingMips);
+	}
+
+	const int32 MinRequestMipCount = 0;
+	if (NumRequestedMips < MinRequestMipCount && MinRequestMipCount < NumMips)
+	{
+		NumRequestedMips = MinRequestMipCount;
+	}
+
+	FStreamableRenderResourceState PostInitState;
+	PostInitState.bSupportsStreaming = bMakeStreamable;
+	PostInitState.NumNonStreamingLODs = IntCastChecked<uint8>(NumOfNonStreamingMips);
+	PostInitState.NumNonOptionalLODs = IntCastChecked<uint8>(BiasedNumNonOptionalMips);
+	PostInitState.MaxNumLODs = IntCastChecked<uint8>(NumMips);
+	PostInitState.AssetLODBias = IntCastChecked<uint8>(AssetMipIdxForResourceFirstMip);
+	PostInitState.NumResidentLODs = IntCastChecked<uint8>(NumRequestedMips);
+	PostInitState.NumRequestedLODs = IntCastChecked<uint8>(NumRequestedMips);
+
+	return PostInitState;
+}
+
+bool ULandscapeTextureStorageProviderFactory::GetInitialMipData(int32 FirstMipToLoad, TArrayView<void*> OutMipData, TArrayView<int64> OutMipSize, FStringView DebugContext)
+{
+	TRACE_CPUPROFILER_EVENT_SCOPE(ULandscapeTextureStorageProviderFactory::GetInitialMipData);
+	check(FirstMipToLoad >= 0);
+	int32 NumberOfMipsToLoad = OutMipData.Num();
+	check(NumberOfMipsToLoad > 0);
+	check(OutMipData.GetData());
+
+	const int32 LoadableMips = Mips.Num();
+	check(NumberOfMipsToLoad == LoadableMips - FirstMipToLoad);
+
+	const int32 MipLoadEnd = FirstMipToLoad + NumberOfMipsToLoad;
+	check(MipLoadEnd <= LoadableMips);
+
+	check(OutMipSize.Num() == NumberOfMipsToLoad || OutMipSize.Num() == 0);
+
+	int32 NumMipsCached = 0;
+
+	// Handle the case where we inlined more mips than we intend to upload immediately, by discarding the unneeded mips
+	for (int32 MipIndex = 0; MipIndex < FirstMipToLoad && MipIndex < LoadableMips; ++MipIndex)
+	{
+		FLandscapeTexture2DMipMap& Mip = Mips[MipIndex];
+		if (Mip.BulkData.IsBulkDataLoaded())
+		{
+			// we know inline mips are set up with the discard after first use flag, so simply locking then unlocking will cause them to be deleted
+			Mip.BulkData.Lock(LOCK_READ_ONLY);
+			Mip.BulkData.Unlock();
+		}
+	}
+
+	// Get data for the remaining mips from bulk data.
+	for (int32 MipIndex = FirstMipToLoad; MipIndex < MipLoadEnd; ++MipIndex)
+	{
+		FLandscapeTexture2DMipMap& Mip = Mips[MipIndex];
+		const int64 DestBytes = Mip.SizeX * Mip.SizeY * 4;
+		const int64 BulkDataSize = Mip.BulkData.GetBulkDataSize();
+		if (BulkDataSize > 0)
+		{
+			void* SourceData = nullptr;
+			bool bDiscardInternalCopy = true;
+			Mip.BulkData.GetCopy(&SourceData, bDiscardInternalCopy);
+			check(SourceData);
+			uint8* DestData = static_cast<uint8*>(FMemory::Malloc(DestBytes));
+			DecompressMip((uint8*) SourceData, BulkDataSize, DestData, DestBytes, MipIndex);
+			OutMipData[MipIndex - FirstMipToLoad] = DestData;
+			FMemory::Free(SourceData);
+
+			if (OutMipSize.Num() > 0)
+			{
+				OutMipSize[MipIndex - FirstMipToLoad] = DestBytes;
+			}
+			NumMipsCached++;
+		}
+	}
+
+	if (NumMipsCached != (LoadableMips - FirstMipToLoad))
+	{
+		UE_LOG(LogLandscape, Warning, TEXT("ULandscapeTextureStorageProviderFactory::TryLoadMips failed for %.*s, NumMipsCached: %d, LoadableMips: %d, FirstMipToLoad: %d"),
+			DebugContext.Len(), DebugContext.GetData(),
+			NumMipsCached,
+			LoadableMips,
+			FirstMipToLoad);
+
+		// Unable to cache all mips. Release memory for those that were cached.
+		for (int32 MipIndex = FirstMipToLoad; MipIndex < LoadableMips; ++MipIndex)
+		{
+			FLandscapeTexture2DMipMap& Mip = Mips[MipIndex];
+			UE_LOG(LogLandscape, Verbose, TEXT("  Mip %d, BulkDataSize: %" INT64_FMT),
+				MipIndex,
+				Mip.BulkData.GetBulkDataSize());
+
+			if (OutMipData[MipIndex - FirstMipToLoad])
+			{
+				FMemory::Free(OutMipData[MipIndex - FirstMipToLoad]);
+				OutMipData[MipIndex - FirstMipToLoad] = nullptr;
+			}
+			if (OutMipSize.Num() > 0)
+			{
+				OutMipSize[MipIndex - FirstMipToLoad] = 0;
+			}
+		}
+		return false;
+	}
+	return true;
+}
+
+
+#if WITH_EDITORONLY_DATA
+ULandscapeTextureStorageProviderFactory* ULandscapeTextureStorageProviderFactory::ApplyTo(UTexture2D* TargetTexture, const FVector& InLandsapeGridScale)
+{
+	check(TargetTexture);
+	check(TargetTexture->Source.IsValid())
+
+	check(TargetTexture->Source.GetFormat() == TSF_BGRA8);
+	EPixelFormat Format = PF_B8G8R8A8;
+
+	int32 Width = TargetTexture->Source.GetSizeX();
+	int32 Height = TargetTexture->Source.GetSizeY();
+	int32 MipCount = TargetTexture->Source.GetNumMips();
+
+	uint32 SrcBpp = GPixelFormats[Format].BlockBytes;
+	uint32 SrcPitch = Width * SrcBpp;
+
+	// try to get an existing factory
+	ULandscapeTextureStorageProviderFactory* Factory= TargetTexture->GetAssetUserData<ULandscapeTextureStorageProviderFactory>();
+	if (Factory == nullptr)
+	{
+		// create a new one
+		Factory = NewObject<ULandscapeTextureStorageProviderFactory>(TargetTexture);
+		Factory->Texture = TargetTexture;
+		TargetTexture->AddAssetUserData(Factory);
+	}
+
+	Factory->LandscapeGridScale = InLandsapeGridScale;
+
+	// calculate number of non-streaming mips
+	// TODO [chris.tchou] : we could make this calculation platform specific, like Texture2D does.
+	// We would have to calculate it during serialization, when we know the target platform.
+	{
+		int32 NumNonStreamingMips = 1;
+
+		// TODO [chris.tchou] : we could ensure Mip Tails are not streamed, as it's more overhead to upload.
+		// we would have to query TextureCompressorModule for platform specific info.
+		// Ignoring the mip tail should still work, just less optimal as it does more work at runtime to blit into the mip tail.
+		int32 NumMipsInTail = 0;
+
+		NumNonStreamingMips = FMath::Max(NumNonStreamingMips, NumMipsInTail);
+		NumNonStreamingMips = FMath::Max(NumNonStreamingMips, UTexture2D::GetStaticMinTextureResidentMipCount());
+		NumNonStreamingMips = FMath::Min(NumNonStreamingMips, MipCount);
+		Factory->NumNonStreamingMips = NumNonStreamingMips;
+	}
+
+	// calculate number of non-optional mips
+	{
+		// for now, landscape texture storage does not have any optional mips
+		Factory->NumNonOptionalMips = MipCount;
+	}
+
+	Factory->Mips.Empty();
+	int32 MipWidth = Width;
+	int32 MipHeight = Height;
+	for (int32 MipIndex = 0; MipIndex < MipCount; MipIndex++)
+	{
+		FLandscapeTexture2DMipMap* Mip = new(Factory->Mips) FLandscapeTexture2DMipMap();
+		Mip->SizeX = MipWidth;
+		Mip->SizeY = MipHeight;
+		
+		TArray64<uint8> MipData;
+		TargetTexture->Source.GetMipData(MipData, MipIndex);
+
+		// store small mips uncompressed
+		constexpr int32 UncompressedMipSizeThreshold = 8;
+		if ((MipWidth <= UncompressedMipSizeThreshold) || (MipHeight <= UncompressedMipSizeThreshold))
+		{
+			Mip->bCompressed = false;
+			CopyMipToBulkData(MipIndex, MipWidth, MipHeight, MipData.GetData(), MipData.Num(), Mip->BulkData);
+		}
+		else
+		{
+			Mip->bCompressed = true;
+			CompressMipToBulkData(MipIndex, MipWidth, MipHeight, MipData.GetData(), MipData.Num(), Mip->BulkData);
+		}
+
+		MipWidth >>= 1;
+		MipHeight >>= 1;
+	}
+
+	return Factory;
+}
+#endif // WITH_EDITORONLY_DATA
+
+
+// Helper to configure the AsyncFileCallBack.
+void FLandscapeTextureStorageMipProvider::CreateAsyncFileCallback(const FTextureUpdateSyncOptions& SyncOptions)
+{
+	FThreadSafeCounter* Counter = SyncOptions.Counter;
+	FTextureUpdateSyncOptions::FCallback RescheduleCallback = SyncOptions.RescheduleCallback;
+	check(Counter && RescheduleCallback);
+
+	AsyncFileCallBack = [this, Counter, RescheduleCallback](bool bWasCancelled, IBulkDataIORequest* Req)
+	{
+		// At this point task synchronization would hold the number of pending requests.
+		Counter->Decrement();
+
+		if (bWasCancelled)
+		{
+			bIORequestCancelled = true;
+		}
+
+		if (Counter->GetValue() == 0)
+		{
+			RescheduleCallback();
+		}
+	};
+}
+
+void FLandscapeTextureStorageMipProvider::ClearIORequests()
+{
+	for (FIORequest& IORequest : IORequests)
+	{
+		// If requests are not yet completed, cancel and wait.
+		if (IORequest.BulkDataIORequest && !IORequest.BulkDataIORequest->PollCompletion())
+		{
+			IORequest.BulkDataIORequest->Cancel();
+			IORequest.BulkDataIORequest->WaitCompletion();
+		}
+	}
+	IORequests.Empty();
+}
+
+void FLandscapeTextureStorageMipProvider::Init(const FTextureUpdateContext& Context, const FTextureUpdateSyncOptions& SyncOptions)
+{
+	IORequests.AddDefaulted(CurrentFirstLODIdx);
+
+	// If this resource has optional LODs and we are streaming one of them.
+	if (ResourceState.NumNonOptionalLODs < ResourceState.MaxNumLODs && PendingFirstLODIdx < ResourceState.LODCountToFirstLODIdx(ResourceState.NumNonOptionalLODs))
+	{
+		// Generate the FilenameHash of each optional LOD before the first one requested, so that we can handle properly PAK unmount events.
+		// Note that streamer only stores the hash for the first optional mip.
+		for (int32 MipIdx = 0; MipIdx < PendingFirstLODIdx; ++MipIdx)
+		{
+			const FLandscapeTexture2DMipMap* SourceMip = Factory->GetMip(MipIdx);
+			// const FTexture2DMipMap& OwnerMip = *Context.MipsView[MipIdx];
+			IORequests[MipIdx].FilenameHash = SourceMip->BulkData.GetIoFilenameHash();
+		}
+	}
+
+	// Otherwise validate each streamed in mip.
+	for (int32 MipIdx = PendingFirstLODIdx; MipIdx < CurrentFirstLODIdx; ++MipIdx)
+	{
+		const FLandscapeTexture2DMipMap* SourceMip = Factory->GetMip(MipIdx);
+		if (SourceMip->BulkData.IsStoredCompressedOnDisk())
+		{
+			// Compression at the package level is no longer supported
+			continue;
+		}
+		else if (SourceMip->BulkData.GetBulkDataSize() <= 0)
+		{
+			// Invalid bulk data size.
+			continue;
+		}
+		else
+		{
+			IORequests[MipIdx].FilenameHash = SourceMip->BulkData.GetIoFilenameHash();
+		}
+	}
+
+	AdvanceTo(ETickState::GetMips, ETickThread::Async);
+}
+
+int32 FLandscapeTextureStorageMipProvider::GetMips(const FTextureUpdateContext& Context, int32 StartingMipIndex, const FTextureMipInfoArray& MipInfos, const FTextureUpdateSyncOptions& SyncOptions)\
+{
+	CreateAsyncFileCallback(SyncOptions);	// this just creates it... callback has to be passed to the IO request completion to actually get called...
+	check(SyncOptions.Counter != nullptr);
+
+	FirstRequestedMipIndex = StartingMipIndex;
+	while (StartingMipIndex < CurrentFirstLODIdx && MipInfos.IsValidIndex(StartingMipIndex))
+	{
+		const FTextureMipInfo& DestMip = MipInfos[StartingMipIndex];
+		const FLandscapeTexture2DMipMap* SourceMip = Factory->GetMip(StartingMipIndex);
+		if (SourceMip == nullptr || !DestMip.DestData)
+		{
+			break;
+		}
+
+		// Check the validity of the filename.
+		if (IORequests[StartingMipIndex].FilenameHash == INVALID_IO_FILENAME_HASH)
+		{
+			break;
+		}
+
+		// If Data size is specified, check compatibility for safety		// TODO: size doesn't have to match when compressed...
+		if (DestMip.DataSize && (uint64)SourceMip->BulkData.GetBulkDataSize() > DestMip.DataSize)
+		{
+			break;
+		}
+
+		// Increment the sync counter.  This causes the system to not advance to the next tick, until RescheduleCallback() is called (by AsyncFileCallBack when counter reaches zero)
+		// If a request completes immediately, then it will call the RescheduleCallback,
+		// but that won't do anything because the tick would not try to acquire the lock since it is already locked.
+		SyncOptions.Counter->Increment();
+
+		int64 StreamDataSize = SourceMip->BulkData.GetBulkDataSize();
+		uint8* StreamData = static_cast<uint8*>(FMemory::Malloc(StreamDataSize));
+		IORequests[StartingMipIndex].BulkDataIORequest.Reset(
+			SourceMip->BulkData.CreateStreamingRequest(
+				0,
+				StreamDataSize,
+				(EAsyncIOPriorityAndFlags) FMath::Clamp<int32>(AIOP_Low + (bPrioritizedIORequest ? 1 : 0), AIOP_Low, AIOP_High) | AIOP_FLAG_DONTCACHE,
+				&AsyncFileCallBack,
+				StreamData)
+		);
+
+		// remember the dest mip data buffer (we can't fill it out now, must wait until streaming is complete)
+		IORequests[StartingMipIndex].DestMipData = static_cast<uint8*>(DestMip.DestData);
+
+		StartingMipIndex++;
+	}
+
+	AdvanceTo(ETickState::PollMips, ETickThread::Async);
+	return StartingMipIndex;	// return the mips we handled (if this is not CurrentFirstLODIdx, it will fall back to other providers)
+}
+
+bool FLandscapeTextureStorageMipProvider::PollMips(const FTextureUpdateSyncOptions& SyncOptions)
+{
+	// poll mips will run once all io requests are complete (or cancelled)
+
+	// Notify that some files have possibly been unmounted / missing.
+	if (bIORequestCancelled && !bIORequestAborted)
+	{
+		IRenderAssetStreamingManager& StreamingManager = IStreamingManager::Get().GetRenderAssetStreamingManager();
+		for (FIORequest& IORequest : IORequests)
+		{
+			StreamingManager.MarkMountedStateDirty(IORequest.FilenameHash);
+		}
+		UE_LOG(LogLandscape, Warning, TEXT("[%s] FLandscapeTextureStorageMipProvider Texture stream in request failed due to IO error (Mip %d-%d)."), *TextureName.ToString(), ResourceState.AssetLODBias + PendingFirstLODIdx, ResourceState.AssetLODBias + CurrentFirstLODIdx - 1);
+	}
+
+	if (!bIORequestCancelled && !bIORequestAborted)
+	{
+		// decompress the mips (note that this is using the dest mip data pointer we memorized during GetMips)
+		for (int MipIndex = FirstRequestedMipIndex; MipIndex < CurrentFirstLODIdx; MipIndex++)
+		{
+			const FLandscapeTexture2DMipMap* SourceMip = Factory->GetMip(MipIndex);
+			uint8* SourceData = IORequests[MipIndex].BulkDataIORequest->GetReadResults();
+			int64 DestDataBytes = SourceMip->SizeX * SourceMip->SizeY * 4;
+			uint8* DestData = IORequests[MipIndex].DestMipData;
+			Factory->DecompressMip(SourceData, SourceMip->BulkData.GetBulkDataSize(), DestData, DestDataBytes, MipIndex);
+		}
+	}
+
+	ClearIORequests();
+
+	AdvanceTo(ETickState::Done, ETickThread::None);
+
+	return !bIORequestCancelled;	// return true if successful and it can upload the DestMip data to the GPU
+}
+
+void FLandscapeTextureStorageMipProvider::AbortPollMips()
+{
+	// ... cancel all streaming ops in progress ...
+	for (FIORequest& IORequest : IORequests)
+	{
+		if (IORequest.BulkDataIORequest)
+		{
+			// Calling cancel() here will trigger the AsyncFileCallBack and precipitate the execution of Cancel().
+			IORequest.BulkDataIORequest->Cancel();
+			bIORequestAborted = true;
+		}
+	}
+}
+
+void FLandscapeTextureStorageMipProvider::CleanUp(const FTextureUpdateSyncOptions& SyncOptions)
+{
+	AdvanceTo(ETickState::Done, ETickThread::None);
+}
+
+void FLandscapeTextureStorageMipProvider::Cancel(const FTextureUpdateSyncOptions& SyncOptions)
+{
+	ClearIORequests();
+}
+
+FTextureMipDataProvider::ETickThread FLandscapeTextureStorageMipProvider::GetCancelThread() const
+{
+	return IORequests.Num() ? FTextureMipDataProvider::ETickThread::Async : FTextureMipDataProvider::ETickThread::None;
+}
+
+void ULandscapeTextureStorageProviderFactory::CopyMipToBulkData(int32 MipIndex, int32 MipSizeX, int32 MipSizeY, uint8* SourceData, int32 SourceDataBytes, FByteBulkData& DestBulkData)
+{
+	TRACE_CPUPROFILER_EVENT_SCOPE(ULandscapeTextureStorageProviderFactory::CopyMipToBulkData);
+	DestBulkData.Lock(LOCK_READ_WRITE);
+
+	int32 TotalPixels = MipSizeX * MipSizeY;
+	check(SourceDataBytes == TotalPixels * 4);
+
+	int32 DestBytes = SourceDataBytes;
+	uint8* DestData = DestBulkData.Realloc(DestBytes);
+
+	memcpy(DestData, SourceData, DestBytes);
+
+	DestBulkData.Unlock();
+}
+
+void ULandscapeTextureStorageProviderFactory::CompressMipToBulkData(int32 MipIndex, int32 MipSizeX, int32 MipSizeY, uint8* SourceData, int32 SourceDataBytes, FByteBulkData& DestBulkData)
+{
+	TRACE_CPUPROFILER_EVENT_SCOPE(ULandscapeTextureStorageProviderFactory::CompressMipToBulkData);
+
+	DestBulkData.Lock(LOCK_READ_WRITE);
+
+	int32 TotalPixels = MipSizeX * MipSizeY;
+	check(SourceDataBytes == TotalPixels * 4);
+	check(TotalPixels >= 16);	// shouldn't be used on very small mips
+
+	// DestData consists of a 16 bit height per pixel, then an 8:8 normal per edge pixel
+	int32 DestBytes = (TotalPixels + (MipSizeX + MipSizeY) * 2 - 4) * 2;
+	uint8* DestData = DestBulkData.Realloc(DestBytes);
+
+	// delta encode the heights -- this (usually) greatly reduces the variance in the data, which makes it compress much better on disk when package compression is applied.
+	uint16 LastHeight = 32768;
+	int32 DestOffset = 0;
+	for (int32 SourceOffset = 0; SourceOffset < SourceDataBytes; SourceOffset += 4)
+	{
+		// texture data is stored as BGRA, or [normal x, height low bits, height high bits, normal y]
+		uint16 Height = SourceData[SourceOffset + 2] * 256 + SourceData[SourceOffset + 1];
+		uint16 DeltaHeight = Height - LastHeight;
+		LastHeight = Height;
+
+		// store delta height
+		DestData[DestOffset + 0] = DeltaHeight >> 8;
+		DestData[DestOffset + 1] = DeltaHeight & 0xff;
+		DestOffset += 2;
+	}
+
+	int32 DeltaCount = DestOffset;
+
+	// capture normals along the edge (delta encoded clockwise starting from top left)
+	uint8 LastNormalX = 128;
+	uint8 LastNormalY = 128;
+
+	auto EncodeNormal = [&LastNormalX, &LastNormalY, SourceData, DestData, MipSizeX, &DestOffset](int32 X, int32 Y)
+	{
+		int32 SourceOffset = (Y * MipSizeX + X) * 4;
+		uint8 NormalX = SourceData[SourceOffset + 0];
+		uint8 NormalY = SourceData[SourceOffset + 3];
+		DestData[DestOffset + 0] = NormalX - LastNormalX;
+		DestData[DestOffset + 1] = NormalY - LastNormalY;
+		LastNormalX = NormalX;
+		LastNormalY = NormalY;
+		DestOffset += 2;
+	};
+
+	for (int32 X = 0; X < MipSizeX; X++)				// [0 ... MipSizeX-1], 0
+	{
+		EncodeNormal(X, 0);
+	}
+
+	for (int32 Y = 1; Y < MipSizeY; Y++)				// MipSizeX-1, [1 ... MipSizeY-1]
+	{
+		EncodeNormal(MipSizeX - 1, Y);
+	}
+
+	for (int32 X = MipSizeX-2; X >= 0; X--)				// [MipSizeX-2 ... 0], MipSizeY-1
+	{
+		EncodeNormal(X, MipSizeY - 1);
+	}
+
+	for (int32 Y = MipSizeY-2; Y >= 1; Y--)				// 0, [MipSizeY-2 ... 1]
+	{
+		EncodeNormal(0, Y);
+	}
+
+	check(DestOffset == DestBytes);
+
+	DestBulkData.Unlock();
+}
+
+// Compute the normal of the triangle formed by the 3 points (in winding order).
+static FVector ComputeTriangleNormal(const FVector& InPoint0, const FVector& InPoint1, const FVector& InPoint2)
+{
+	FVector Normal = (InPoint0 - InPoint1).Cross(InPoint1 - InPoint2);
+	Normal.Normalize();
+	return Normal;
+}
+
+static void SampleWorldPositionAtOffset(FVector& OutPoint, const uint8* MipData, int32 X, int32 Y, int32 MipSizeX, const FVector& InLandscapeGridScale)
+{
+	int32 OffsetBytes = (Y * MipSizeX + X) * 4;
+	uint16 HeightData = MipData[OffsetBytes + 2] * 256 + MipData[OffsetBytes + 1];
+
+	// NOTE: since we are using deltas between points to calculate the normal, we don't care about constant offsets in the position, only relative scales
+	OutPoint.Set(
+		X * InLandscapeGridScale.X,
+		Y * InLandscapeGridScale.Y,
+		LandscapeDataAccess::GetLocalHeight(HeightData) * InLandscapeGridScale.Z);
+}
+
+void ULandscapeTextureStorageProviderFactory::DecompressMip(uint8* SourceData, int64 SourceDataBytes, uint8* DestData, int64 DestDataBytes, int32 MipIndex)
+{
+	TRACE_CPUPROFILER_EVENT_SCOPE(ULandscapeTextureStorageProviderFactory::DecompressMip);
+
+	check(SourceData && DestData);
+
+	FLandscapeTexture2DMipMap& Mip = Mips[MipIndex];
+
+	if (!Mip.bCompressed)
+	{
+		// mip is uncompressed, just copy it
+		check(SourceDataBytes == DestDataBytes);
+		memcpy(DestData, SourceData, DestDataBytes);
+		return;
+	}
+
+	check(Mip.bCompressed);
+
+	int32 TotalPixels = Mip.SizeX * Mip.SizeY;
+	check(SourceDataBytes == (TotalPixels + (Mip.SizeX + Mip.SizeY) * 2 - 4) * 2);
+	check(DestDataBytes == TotalPixels * 4);
+
+	// Undo Delta Encode of Heights
+	uint16 LastHeight = 32768;
+	for (int32 PixelIndex = 0; PixelIndex < TotalPixels; PixelIndex++)
+	{
+		int32 SourceOffset = PixelIndex * 2;
+		uint16 DeltaHeight = SourceData[SourceOffset + 0] * 256 + SourceData[SourceOffset + 1];
+
+		// undo delta
+		LastHeight += DeltaHeight;
+
+		// texture data is stored as BGRA, or [normal x, height low bits, height high bits, normal y]
+		int32 DestOffset = PixelIndex * 4;
+		DestData[DestOffset + 0] = 128;
+		DestData[DestOffset + 1] = LastHeight & 0xff;
+		DestData[DestOffset + 2] = LastHeight >> 8;
+		DestData[DestOffset + 3] = 128;
+	}
+
+	// recompute normals in the interior
+	{
+		// we skip computing the edges, as they will be overwritten later (and this way we don't have to handle samples that go off the edge)
+		for (int32 Y = 1; Y < Mip.SizeY - 1; Y++)
+		{
+			for (int32 X = 1; X < Mip.SizeX - 1; X++)
+			{
+				// based on shader code in LandscapeLayersPS.usf
+				FVector TL, TT, CC, LL, RR, BR, BB;
+
+				SampleWorldPositionAtOffset(TL, DestData, X - 1, Y - 1, Mip.SizeX, LandscapeGridScale);
+				SampleWorldPositionAtOffset(TT, DestData, X + 0, Y - 1, Mip.SizeX, LandscapeGridScale);
+				SampleWorldPositionAtOffset(CC, DestData, X + 0, Y + 0, Mip.SizeX, LandscapeGridScale);
+				SampleWorldPositionAtOffset(LL, DestData, X - 1, Y + 0, Mip.SizeX, LandscapeGridScale);
+				SampleWorldPositionAtOffset(RR, DestData, X + 1, Y + 0, Mip.SizeX, LandscapeGridScale);
+				SampleWorldPositionAtOffset(BR, DestData, X + 1, Y + 1, Mip.SizeX, LandscapeGridScale);
+				SampleWorldPositionAtOffset(BB, DestData, X + 0, Y + 1, Mip.SizeX, LandscapeGridScale);
+
+				FVector N0 = ComputeTriangleNormal(CC, LL, TL);
+				FVector N1 = ComputeTriangleNormal(TL, TT, CC);
+				FVector N2 = ComputeTriangleNormal(LL, CC, BB);
+				FVector N3 = ComputeTriangleNormal(RR, CC, TT);
+				FVector N4 = ComputeTriangleNormal(BR, BB, CC);
+				FVector N5 = ComputeTriangleNormal(CC, RR, BR);
+
+				FVector FinalNormal = (N0 + N1 + N2 + N3 + N4 + N5);
+				FinalNormal.Normalize();
+
+				// rescale normal.xy to [0,255] range, and write out as bytes
+				int32 OffsetBytes = (Y * Mip.SizeX + X) * 4;
+				DestData[OffsetBytes + 0] = static_cast<uint8>(FMath::Clamp(((FinalNormal.X + 1.0) * 0.5) * 255.0, 0.0, 255.0));
+				DestData[OffsetBytes + 3] = static_cast<uint8>(FMath::Clamp(((FinalNormal.Y + 1.0) * 0.5) * 255.0, 0.0, 255.0));
+			}
+		}
+	}
+
+	// write out normals along the edge (delta encoded clockwise starting from top left)
+	int32 SourceOffset = TotalPixels * 2;
+	uint8 LastNormalX = 128;
+	uint8 LastNormalY = 128;
+
+	auto DecodeNormal = [&LastNormalX, &LastNormalY, SourceData, DestData, MipSizeX = Mip.SizeX, &SourceOffset](int32 X, int32 Y)
+	{
+		int32 DestOffset = (Y * MipSizeX + X) * 4;
+		LastNormalX += SourceData[SourceOffset + 0];
+		LastNormalY += SourceData[SourceOffset + 1];
+		DestData[DestOffset + 0] = LastNormalX;
+		DestData[DestOffset + 3] = LastNormalY;
+		SourceOffset += 2;
+	};
+
+	for (int32 X = 0; X < Mip.SizeX; X++)				// [0 ... MipSizeX-1], 0
+	{
+		DecodeNormal(X, 0);
+	}
+
+	for (int32 Y = 1; Y < Mip.SizeY; Y++)				// MipSizeX-1, [1 ... MipSizeY-1]
+	{
+		DecodeNormal(Mip.SizeX - 1, Y);
+	}
+
+	for (int32 X = Mip.SizeX - 2; X >= 0; X--)				// [MipSizeX-2 ... 0], MipSizeY-1
+	{
+		DecodeNormal(X, Mip.SizeY - 1);
+	}
+
+	for (int32 Y = Mip.SizeY - 2; Y >= 1; Y--)				// 0, [MipSizeY-2 ... 1]
+	{
+		DecodeNormal(0, Y);
+	}
+
+	check(SourceOffset == SourceDataBytes);
+}
