* Added threshold to dirty landscape textures detection (CVars "landscape.DirtyHeightmapThreshold" / "landscape.DirtyWeightmapThreshold") : when >0 and the texture hash differs, we run a deep-comparison of the old and new texture data and let go of any change whose difference is under the threshold, which avoids dirtying the landscape because of precision issues
* Added CVar "landscape.DumpDiffDetails" (to be used in conjunction with "landscape.DumpHeightmapDiff/DumpWeightmapDiff") to dump additional details about the diff introduced on a given heightmap/weightmap (number of different pixels, total differences, locations and details about differences) in a text file next to the dumped bmp files
* Changed hashing algorithm from CRC32 to CityHash64 (0.8ms->0.3ms per 256x256 texture)

diff --git a/Engine/Source/Runtime/Landscape/Classes/Landscape.h b/Engine/Source/Runtime/Landscape/Classes/Landscape.h
--- a/Engine/Source/Runtime/Landscape/Classes/Landscape.h
+++ b/Engine/Source/Runtime/Landscape/Classes/Landscape.h
@@ -375,6 +375,7 @@ private:
 	int32 PerformLayersWeightmapsGlobalMerge(FUpdateLayersContentContext& InUpdateLayersContentContext, const FEditLayersWeightmapMergeParams& InMergeParams);
 	void ResolveLayersWeightmapTexture(const FTextureToComponentHelper& MapHelper, const TSet<UTexture2D*>& WeightmapsToResolve, bool bIntermediateRender, bool bFlushRender, TArray<FLandscapeEditLayerComponentReadbackResult>& InOutComponentReadbackResults);
 
+	bool HasTextureDataChanged(TArrayView<const FColor> InOldData, TArrayView<const FColor> InNewData, bool bInIsWeightmap, uint64 InPreviousHash, uint64& OutNewHash) const;
 	bool ResolveLayersTexture(FTextureToComponentHelper const& MapHelper, FLandscapeEditLayerReadback* InCPUReadBack, UTexture2D* InOutputTexture, bool bIntermediateRender,	bool bFlushRender,
 		TArray<FLandscapeEditLayerComponentReadbackResult>& InOutComponentReadbackResults, bool bIsWeightmap);
 
diff --git a/Engine/Source/Runtime/Landscape/Private/Landscape.cpp b/Engine/Source/Runtime/Landscape/Private/Landscape.cpp
--- a/Engine/Source/Runtime/Landscape/Private/Landscape.cpp
+++ b/Engine/Source/Runtime/Landscape/Private/Landscape.cpp
@@ -448,7 +448,7 @@ FGuid ALandscapeProxy::GetNaniteContentId() const
 
 	FBufferArchive ContentStateAr;
 
-	int32 LocalNaniteLODIndex = GetNaniteLODIndex(); // Use the getter to validate it's set at the same value as the parent actor's
+	int32 LocalNaniteLODIndex = GetNaniteLODIndex();
 	ContentStateAr << LocalNaniteLODIndex;
 
 	struct FCompareULandscapeComponentBySectionBase
diff --git a/Engine/Source/Runtime/Landscape/Private/LandscapeEditLayers.cpp b/Engine/Source/Runtime/Landscape/Private/LandscapeEditLayers.cpp
--- a/Engine/Source/Runtime/Landscape/Private/LandscapeEditLayers.cpp
+++ b/Engine/Source/Runtime/Landscape/Private/LandscapeEditLayers.cpp
@@ -138,6 +138,21 @@ static TAutoConsoleVariable<int32> CVarLandscapeDumpWeightmapDiff(
 	0,
 	TEXT("This will save images for readback weightmap textures that have changed in the last edit layer blend phase. (= 0 No Diff, 1 = Mip 0 Diff, 2 = All Mips Diff"));
 
+static TAutoConsoleVariable<bool> CVarLandscapeDumpDiffDetails(
+	TEXT("landscape.DumpDiffDetails"),
+	false,
+	TEXT("When dumping diffs for heightmap (landscape.DumpHeightmapDiff) or weightmap (landscape.DumpWeightmapDiff), dumps additional details about the pixels being different"));
+
+TAutoConsoleVariable<int32> CVarLandscapeDirtyHeightmapThreshold(
+	TEXT("landscape.DirtyHeightmapThreshold"),
+	0,
+	TEXT("Threshold to avoid imprecision issues on certain GPUs when detecting when a heightmap changes, i.e. only a difference > than this threshold (N over 16-bits uint height and N over each 8-bits uint normal component) will be detected as a change."));
+
+TAutoConsoleVariable<int32> CVarLandscapeDirtyWeightmapThreshold(
+	TEXT("landscape.DirtyWeightmapThreshold"),
+	0,
+	TEXT("Threshold to avoid imprecision issues on certain GPUs when detecting when a weightmap changes, i.e. only a difference > than this threshold (N over each 8-bits uint weightmap channel)."));
+
 TAutoConsoleVariable<int32> CVarLandscapeShowDirty(
 	TEXT("landscape.ShowDirty"),
 	0,
@@ -163,6 +178,7 @@ TAutoConsoleVariable<int32> CVarLandscapeRemoveEmptyPaintLayersOnEdit(
 	1,
 	TEXT("This will analyze weightmaps on readback and remove unneeded allocations (for unpainted layers)."));
 
+// TODO [jonathan.bard] : find out which nvidia driver version fixed the issue (last known failing : 526.47, first known working : 531.41), bump the minimum shader requirements for the editor and remove this workaround : 
 // COMMENT [jonathan.bard] : this is a workaround for a random D3D12 issue that misses some cache flush when copying edit layers individual heightmaps to the global merge atlas and then using 
 //  the atlas as SRV right after, which should be ensured by the manually inserted barriers but is sometimes not). Using PS-based copies doesn't exhibit the behavior so we'll stick with that for now : 
 TAutoConsoleVariable<bool> CVarLandscapeUsePSCopyWorkaround(
@@ -5075,7 +5091,7 @@ int32 ALandscape::RegenerateLayersHeightmaps(const FUpdateLayersContentContext&
 			{
 				FLandscapeEditLayerReadback* NewCPUReadback = new FLandscapeEditLayerReadback();
 				const uint8* LockedMip = ComponentHeightmap->Source.LockMip(0);
-				const uint32 Hash = FLandscapeEditLayerReadback::CalculateHash(LockedMip, ComponentHeightmap->Source.GetSizeX() * ComponentHeightmap->Source.GetSizeY() * sizeof(FColor));
+				const uint64 Hash = FLandscapeEditLayerReadback::CalculateHash(LockedMip, ComponentHeightmap->Source.GetSizeX() * ComponentHeightmap->Source.GetSizeY() * sizeof(FColor));
 				ComponentHeightmap->Source.UnlockMip(0);
 				NewCPUReadback->SetHash(Hash);
 				Proxy->HeightmapsCPUReadback.Add(ComponentHeightmap, NewCPUReadback);
@@ -5239,10 +5255,11 @@ void ALandscape::UpdateWeightDirtyData(ULandscapeComponent* InLandscapeComponent
 
 void ALandscape::OnDirtyWeightmap(FTextureToComponentHelper const& MapHelper, UTexture2D const* InWeightmap, FColor const* InOldData, FColor const* InNewData, int32 InMipLevel)
 {
-	int32 DumpWeightmapDiff = CVarLandscapeDumpWeightmapDiff.GetValueOnAnyThread();
+	int32 DumpWeightmapDiff = CVarLandscapeDumpWeightmapDiff.GetValueOnGameThread();
 	const bool bDumpDiff = (DumpWeightmapDiff > 0);
 	const bool bDumpDiffAllMips = (DumpWeightmapDiff > 1);
-	const bool bTrackDirty = CVarLandscapeTrackDirty.GetValueOnAnyThread() != 0;
+	const bool bDumpDiffDetails = CVarLandscapeDumpDiffDetails.GetValueOnGameThread();
+	const bool bTrackDirty = CVarLandscapeTrackDirty.GetValueOnGameThread() != 0;
 	ULandscapeSubsystem* LandscapeSubsystem = GetWorld()->GetSubsystem<ULandscapeSubsystem>();
 	check(LandscapeSubsystem != nullptr);
 	const FDateTime CurrentTime = LandscapeSubsystem->GetAppCurrentDateTime();
@@ -5287,7 +5304,44 @@ void ALandscape::OnDirtyWeightmap(FTextureToComponentHelper const& MapHelper, UT
 						FFileHelper::EColorChannel ColorChannel = GetWeightmapColorChannel(AllocInfo);
 						FFileHelper::CreateBitmap(*(FilePattern + "_a(pre).bmp"), SizeU, SizeV, InOldData, /*SubRectangle = */nullptr, &IFileManager::Get(), /*OutFilename = */nullptr, /*bool bInWriteAlpha = */true, ColorChannel);
 						FFileHelper::CreateBitmap(*(FilePattern + "_b(post).bmp"), SizeU, SizeV, InNewData, /*SubRectangle = */nullptr, &IFileManager::Get(), /*OutFilename = */nullptr, /*bool bInWriteAlpha = */true, ColorChannel);
+			
+						if (bDumpDiffDetails)
+						{
+							int32 NumDifferentPixels = 0;
+							uint8 MaxDiff = 0;
+							FStringBuilderBase StrBuilder;
+							for (int32 V = 0; V < SizeV; ++V)
+							{
+								for (int32 U = 0; U < SizeU; ++U)
+								{
+									const FColor* OldDataPtr = InOldData + (V * SizeU + U);
+									const FColor* NewDataPtr = InNewData + (V * SizeU + U);
+									if (*OldDataPtr != *NewDataPtr)
+									{
+										auto ComputeMaxDiff = [&MaxDiff](uint8 OldValue, uint8 NewValue) -> uint8
+										{
+											uint8 Diff = (NewValue > OldValue) ? (NewValue - OldValue) : (OldValue - NewValue);
+											MaxDiff = FMath::Max<uint8>(Diff, MaxDiff);
+											return Diff;
+										};
+
+										StrBuilder.Appendf(TEXT("Pixel (%4u,%4u) : R (%3u -> %3u, absdiff %3u), G (%3u -> %3u, absdiff %3u), B (%3u -> %3u, absdiff %3u), A (%3u -> %3u, absdiff %3u)\n"), U, V, 
+											OldDataPtr->R, NewDataPtr->R, ComputeMaxDiff(OldDataPtr->R, NewDataPtr->R),
+											OldDataPtr->G, NewDataPtr->G, ComputeMaxDiff(OldDataPtr->G, NewDataPtr->G),
+											OldDataPtr->B, NewDataPtr->B, ComputeMaxDiff(OldDataPtr->B, NewDataPtr->B),
+											OldDataPtr->A, NewDataPtr->A, ComputeMaxDiff(OldDataPtr->A, NewDataPtr->A));
+
+										++NumDifferentPixels;
+									}
+								}
+							}
+							StrBuilder.Appendf(TEXT("----------------------------------------\n"));
+							StrBuilder.Appendf(TEXT("Num diffs = %u\n"), NumDifferentPixels);
+							StrBuilder.Appendf(TEXT("Max diff = %u (%1.3f%%)\n"), MaxDiff, 100.0 * static_cast<float>(MaxDiff) / MAX_uint8);
+							FFileHelper::SaveStringToFile(StrBuilder.ToView(), *(FilePattern + "_diff.txt"));
+						}
 					}
+
 				}
 			}
 		}
@@ -5333,10 +5387,11 @@ void ALandscape::UpdateHeightDirtyData(ULandscapeComponent* InLandscapeComponent
 
 void ALandscape::OnDirtyHeightmap(FTextureToComponentHelper const& MapHelper, UTexture2D const* InHeightmap, FColor const* InOldData, FColor const* InNewData, int32 InMipLevel)
 {
-	int32 DumpHeightmapDiff = CVarLandscapeDumpHeightmapDiff.GetValueOnAnyThread();
+	int32 DumpHeightmapDiff = CVarLandscapeDumpHeightmapDiff.GetValueOnGameThread();
 	const bool bDumpDiff = (DumpHeightmapDiff > 0);
 	const bool bDumpDiffAllMips = (DumpHeightmapDiff > 1);
-	const bool bTrackDirty = CVarLandscapeTrackDirty.GetValueOnAnyThread() != 0;
+	const bool bDumpDiffDetails = CVarLandscapeDumpDiffDetails.GetValueOnGameThread();
+	const bool bTrackDirty = CVarLandscapeTrackDirty.GetValueOnGameThread() != 0;
 	ULandscapeSubsystem* LandscapeSubsystem = GetWorld()->GetSubsystem<ULandscapeSubsystem>();
 	check(LandscapeSubsystem != nullptr);
 	const FDateTime CurrentTime = LandscapeSubsystem->GetAppCurrentDateTime();
@@ -5376,9 +5431,138 @@ void ALandscape::OnDirtyHeightmap(FTextureToComponentHelper const& MapHelper, UT
 
 				FFileHelper::CreateBitmap(*(FilePattern + "_a(pre).bmp"), SizeU, SizeV, InOldData, &SubRegion, &IFileManager::Get(), /*OutFilename = */nullptr, /*bool bInWriteAlpha = */true);
 				FFileHelper::CreateBitmap(*(FilePattern + "_b(post).bmp"), SizeU, SizeV, InNewData, &SubRegion, &IFileManager::Get(), /*OutFilename = */nullptr, /*bool bInWriteAlpha = */true);
+
+				if (bDumpDiffDetails)
+				{
+					int32 NumDifferentPixels = 0;
+					uint16 MaxHeightDiff = 0;
+					uint8 MaxNormalDiff = 0;
+					FStringBuilderBase StrBuilder;
+					const FColor* OldDataStartPtr = InOldData + (HeightmapOffsetY * SizeU + HeightmapOffsetX);
+					const FColor* NewDataStartPtr = InNewData + (HeightmapOffsetY * SizeU + HeightmapOffsetX);
+					for (int32 V = 0; V < ComponentWidth; ++V)
+					{
+						for (int32 U = 0; U < ComponentWidth; ++U)
+						{
+							const FColor* OldDataPtr = OldDataStartPtr + (V * SizeU + U);
+							const FColor* NewDataPtr = NewDataStartPtr + (V * SizeU + U);
+							if (*OldDataPtr != *NewDataPtr)
+							{
+								uint16 OldHeight = ((static_cast<uint16>(OldDataPtr->R) << 8) | static_cast<uint16>(OldDataPtr->G));
+								uint16 NewHeight = ((static_cast<uint16>(NewDataPtr->R) << 8) | static_cast<uint16>(NewDataPtr->G));
+								uint16 HeightDiff = (NewHeight > OldHeight) ? (NewHeight - OldHeight) : (OldHeight - NewHeight);
+								MaxHeightDiff = FMath::Max<uint16>(HeightDiff, MaxHeightDiff);
+
+								uint8 OldNormalX = (OldDataPtr->B);
+								uint8 NewNormalX = (NewDataPtr->B);
+								uint8 NormalXDiff = (NewNormalX > OldNormalX) ? (NewNormalX - OldNormalX) : (OldNormalX - NewNormalX);
+								MaxNormalDiff = FMath::Max<uint8>(NormalXDiff, MaxNormalDiff);
+
+								uint8 OldNormalY = (OldDataPtr->A);
+								uint8 NewNormalY = (NewDataPtr->A);
+								uint8 NormalYDiff = (NewNormalY > OldNormalY) ? (NewNormalY - OldNormalY) : (OldNormalY - NewNormalY);
+								MaxNormalDiff = FMath::Max<uint8>(NormalYDiff, MaxNormalDiff);
+
+								StrBuilder.Appendf(TEXT("Pixel (%4u,%4u) : Height (%5u -> %5u, absdiff %5u), NormalX (%3u -> %3u, absdiff %3u), NormalY (%3u -> %3u, absdiff %3u)\n"), U, V, 
+									OldHeight, NewHeight, HeightDiff, 
+									OldNormalX, NewNormalX, NormalXDiff, 
+									OldNormalY, NewNormalY, NormalYDiff);
+
+								++NumDifferentPixels;
+							}
+						}
+					}
+					StrBuilder.Appendf(TEXT("----------------------------------------\n"));
+					StrBuilder.Appendf(TEXT("Num diffs = %u\n"), NumDifferentPixels);
+					StrBuilder.Appendf(TEXT("Max height diff = %u (%1.3f%%)\n"), MaxHeightDiff, 100.0f * static_cast<float>(MaxHeightDiff) / MAX_uint16);
+					StrBuilder.Appendf(TEXT("Max normal diff = %u (%1.3f%%)\n"), MaxNormalDiff, 100.0f * static_cast<float>(MaxNormalDiff) / MAX_uint8);
+					FFileHelper::SaveStringToFile(StrBuilder.ToView(), *(FilePattern + "_diff.txt"));
+				}
+			}
+		}
+	}
+}
+
+bool ALandscape::HasTextureDataChanged(TArrayView<const FColor> InOldData, TArrayView<const FColor> InNewData, bool bInIsWeightmap, uint64 InPreviousHash, uint64& OutNewHash) const
+{
+	TRACE_CPUPROFILER_EVENT_SCOPE(ALandscape::HasTextureDataChanged);
+
+	OutNewHash = FLandscapeEditLayerReadback::CalculateHash((uint8*)InNewData.GetData(), InNewData.Num() * sizeof(FColor));
+	if (OutNewHash != InPreviousHash)
+	{
+		int32 TextureSize = InOldData.Num();
+		check(TextureSize == InNewData.Num());
+		// If necessary, perform a deep comparison to avoid taking into account precision issues as actual changes :
+		if (bInIsWeightmap)
+		{
+			if (int32 DirtyWeightmapThreshold = CVarLandscapeDirtyWeightmapThreshold.GetValueOnGameThread(); DirtyWeightmapThreshold > 0)
+			{
+				TRACE_CPUPROFILER_EVENT_SCOPE(DeepCompareWeightmap);
+				for (int32 Index = 0; Index < TextureSize; ++Index)
+				{
+					const FColor& OldColor = InOldData[Index];
+					const FColor& NewColor = InNewData[Index];
+					if (OldColor != NewColor)
+					{
+						if (uint8 Diff = (NewColor.R > OldColor.R) ? (NewColor.R - OldColor.R) : (OldColor.R - NewColor.R); Diff > DirtyWeightmapThreshold)
+						{
+							return true;
+						}
+						if (uint8 Diff = (NewColor.G > OldColor.G) ? (NewColor.G - OldColor.G) : (OldColor.G - NewColor.G); Diff > DirtyWeightmapThreshold)
+						{
+							return true;
+						}
+						if (uint8 Diff = (NewColor.B > OldColor.B) ? (NewColor.B - OldColor.B) : (OldColor.B - NewColor.B); Diff > DirtyWeightmapThreshold)
+						{
+							return true;
+						}
+						if (uint8 Diff = (NewColor.A > OldColor.A) ? (NewColor.A - OldColor.A) : (OldColor.A - NewColor.A); Diff > DirtyWeightmapThreshold)
+						{
+							return true;
+						}
+					}
+				}
+
+				// No significant difference detected : 
+				return false;
 			}
 		}
+		else if (int32 DirtyHeightmapThreshold = CVarLandscapeDirtyHeightmapThreshold.GetValueOnGameThread(); DirtyHeightmapThreshold > 0)
+		{
+			TRACE_CPUPROFILER_EVENT_SCOPE(DeepCompareHeightmap);
+			for (int32 Index = 0; Index < TextureSize; ++Index)
+			{
+				const FColor& OldColor = InOldData[Index];
+				const FColor& NewColor = InNewData[Index];
+				if (OldColor != NewColor)
+				{
+					uint16 OldHeight = ((static_cast<uint16>(OldColor.R) << 8) | static_cast<uint16>(OldColor.G));
+					uint16 NewHeight = ((static_cast<uint16>(NewColor.R) << 8) | static_cast<uint16>(NewColor.G));
+					if (uint16 Diff = (NewHeight > OldHeight) ? (NewHeight - OldHeight) : (OldHeight - NewHeight); Diff > DirtyHeightmapThreshold)
+					{
+						return true;
+					}
+					if (uint8 Diff = (NewColor.B > OldColor.B) ? (NewColor.B - OldColor.B) : (OldColor.B - NewColor.B); Diff > DirtyHeightmapThreshold)
+					{
+						return true;
+					}
+					if (uint8 Diff = (NewColor.A > OldColor.A) ? (NewColor.A - OldColor.A) : (OldColor.A - NewColor.A); Diff > DirtyHeightmapThreshold)
+					{
+						return true;
+					}
+				}
+			}
+
+			// No significant difference detected : 
+			return false;
+		}
+
+		// Hash is different and we don't need a deep comparison : 
+		return true;
 	}
+
+	// Hash is identical : 
+	return false;
 }
 
 bool ALandscape::ResolveLayersTexture(
@@ -5426,22 +5610,32 @@ bool ALandscape::ResolveLayersTexture(
 		// Copy final result to texture source.
 		TArray<TArray<FColor>> const& OutMipsData = InCPUReadback->GetResult(CompletedReadbackNum - 1);
 
+		int32 NumLockedMips = 0;
 		for (int8 MipIndex = 0; MipIndex < OutMipsData.Num(); ++MipIndex)
 		{
-			if (OutMipsData[MipIndex].Num() > 0)
+			int32 TextureSize = OutMipsData[MipIndex].Num();
+			if (TextureSize > 0)
 			{
 				uint8* TextureData = InOutputTexture->Source.LockMip(MipIndex);
+				++NumLockedMips;
 
-				// Do dirty detection on first mip.
-				// Don't do this for intermediate renders.
-				if (MipIndex == 0 && !bIntermediateRender)
+				// Do dirty detection on first mip :
+				if (MipIndex == 0)
 				{
-					uint32 Hash = 0;
+					// But don't do this for intermediate renders, we'll want the data back on CPU all the time in that case : 
+					if (bIntermediateRender)
+					{
+						TRACE_CPUPROFILER_EVENT_SCOPE(ReadbackToCPU);
+						FMemory::Memcpy(TextureData, OutMipsData[MipIndex].GetData(), OutMipsData[MipIndex].Num() * sizeof(FColor));
+						// Don't waste time copying other mips for intermediate renders : 
+						break;
+					}
+
+					uint64 Hash = 0;
+					if (HasTextureDataChanged(MakeArrayView(reinterpret_cast<const FColor*>(TextureData), TextureSize), MakeArrayView(OutMipsData[MipIndex].GetData(), TextureSize), bIsWeightmap, InCPUReadback->GetHash(), Hash))
 					{
-						TRACE_CPUPROFILER_EVENT_SCOPE(LandscapeLayers_CalculateHash);
-						Hash = FLandscapeEditLayerReadback::CalculateHash((uint8*)OutMipsData[MipIndex].GetData(), OutMipsData[MipIndex].Num() * sizeof(FColor));
+						bChanged |= InCPUReadback->SetHash(Hash);
 					}
-					bChanged |= InCPUReadback->SetHash(Hash);
 
 					if (bChanged)
 					{
@@ -5458,6 +5652,7 @@ bool ALandscape::ResolveLayersTexture(
 					}
 				}
 
+				// Debug utils for when we detect changes :
 				if (bChanged)
 				{
 					if (bIsWeightmap)
@@ -5468,14 +5663,15 @@ bool ALandscape::ResolveLayersTexture(
 					{
 						OnDirtyHeightmap(MapHelper, InOutputTexture, (FColor*)TextureData, OutMipsData[MipIndex].GetData(), MipIndex);
 					}
-				}
 
-				FMemory::Memcpy(TextureData, OutMipsData[MipIndex].GetData(), OutMipsData[MipIndex].Num() * sizeof(FColor));
+					TRACE_CPUPROFILER_EVENT_SCOPE(ReadbackToCPU);
+					FMemory::Memcpy(TextureData, OutMipsData[MipIndex].GetData(), OutMipsData[MipIndex].Num() * sizeof(FColor));
+				}
 			}
 		}
 
 		// Unlock all mips at once because there's a lock counter in FTextureSource that recomputes the content hash when reaching 0 (which means we'd recompute the hash several times over if we Lock/Unlock/Lock/Unloc/... for each mip ):
-		for (int8 MipIndex = 0; MipIndex < OutMipsData.Num(); ++MipIndex)
+		for (int8 MipIndex = 0; MipIndex < NumLockedMips; ++MipIndex)
 		{
 			if (OutMipsData[MipIndex].Num() > 0)
 			{
@@ -6720,7 +6916,7 @@ void ALandscape::PrepareLayersWeightmapsLocalMergeRenderThreadData(const FUpdate
 							// Lazily create the readback objects as required (ReallocateLayersWeightmaps might have created new weightmaps)
 							FLandscapeEditLayerReadback* NewCPUReadback = new FLandscapeEditLayerReadback();
 							const uint8* LockedMip = ComponentWeightmap->Source.LockMipReadOnly(0);
-							const uint32 Hash = FLandscapeEditLayerReadback::CalculateHash(LockedMip, ComponentWeightmap->Source.GetSizeX() * ComponentWeightmap->Source.GetSizeY() * sizeof(FColor));
+							const uint64 Hash = FLandscapeEditLayerReadback::CalculateHash(LockedMip, ComponentWeightmap->Source.GetSizeX() * ComponentWeightmap->Source.GetSizeY() * sizeof(FColor));
 							NewCPUReadback->SetHash(Hash);
 							ComponentWeightmap->Source.UnlockMip(0);
 							CPUReadback = &Proxy->WeightmapsCPUReadback.Add(ComponentWeightmap, NewCPUReadback);
@@ -7171,7 +7367,7 @@ int32 ALandscape::PerformLayersWeightmapsGlobalMerge(FUpdateLayersContentContext
 					{
 						FLandscapeEditLayerReadback* NewCPUReadback = new FLandscapeEditLayerReadback();
 						const uint8* LockedMip = WeightmapTexture->Source.LockMipReadOnly(0);
-						const uint32 Hash = FLandscapeEditLayerReadback::CalculateHash(LockedMip, WeightmapTexture->Source.GetSizeX() * WeightmapTexture->Source.GetSizeY() * sizeof(FColor));
+						const uint64 Hash = FLandscapeEditLayerReadback::CalculateHash(LockedMip, WeightmapTexture->Source.GetSizeX() * WeightmapTexture->Source.GetSizeY() * sizeof(FColor));
 						NewCPUReadback->SetHash(Hash);
 						WeightmapTexture->Source.UnlockMip(0);
 						Proxy->WeightmapsCPUReadback.Add(WeightmapTexture, NewCPUReadback);
diff --git a/Engine/Source/Runtime/Landscape/Private/LandscapeEditReadback.cpp b/Engine/Source/Runtime/Landscape/Private/LandscapeEditReadback.cpp
--- a/Engine/Source/Runtime/Landscape/Private/LandscapeEditReadback.cpp
+++ b/Engine/Source/Runtime/Landscape/Private/LandscapeEditReadback.cpp
@@ -5,6 +5,8 @@
 #include "Engine/Texture2D.h"
 #include "LandscapePrivate.h"
 #include "RenderingThread.h"
+#include "Hash/CityHashHelpers.h"
+
 
 /** Data for a read back task. */
 struct FLandscapeEditReadbackTaskImpl
@@ -262,7 +264,7 @@ static TGlobalResource< FLandscapeEditReadbackTaskPool > GReadbackTaskPool;
 
 
 FLandscapeEditLayerReadback::FLandscapeEditLayerReadback()
-	: Hash(0)
+	: Hash(0ull)
 {}
 
 FLandscapeEditLayerReadback::~FLandscapeEditLayerReadback()
@@ -273,12 +275,14 @@ FLandscapeEditLayerReadback::~FLandscapeEditLayerReadback()
 	}
 }
 
-uint32 FLandscapeEditLayerReadback::CalculateHash(const uint8* InMipData, int32 InSizeInBytes)
+uint64 FLandscapeEditLayerReadback::CalculateHash(const uint8* InMipData, int32 InSizeInBytes)
 {
-	return FCrc::MemCrc32(InMipData, InSizeInBytes);
+	TRACE_CPUPROFILER_EVENT_SCOPE(FLandscapeEditLayerReadback::CalculateHash);
+
+	return CityHash64(reinterpret_cast<const char*>(InMipData), InSizeInBytes);
 }
 
-bool FLandscapeEditLayerReadback::SetHash(uint32 InHash)
+bool FLandscapeEditLayerReadback::SetHash(uint64 InHash)
 {
 	const bool bChanged = InHash != Hash;
 	Hash = InHash;
diff --git a/Engine/Source/Runtime/Landscape/Private/LandscapeEditReadback.h b/Engine/Source/Runtime/Landscape/Private/LandscapeEditReadback.h
--- a/Engine/Source/Runtime/Landscape/Private/LandscapeEditReadback.h
+++ b/Engine/Source/Runtime/Landscape/Private/LandscapeEditReadback.h
@@ -20,12 +20,12 @@ public:
 	~FLandscapeEditLayerReadback();
 
 	/** Hash function to generate a hash used to detect change in read back data. This expects to take raw data from first mip. */
-	static uint32 CalculateHash(const uint8* InMipData, int32 InSizeInBytes);
+	static uint64 CalculateHash(const uint8* InMipData, int32 InSizeInBytes);
 
 	/** Update the stored hash value. Return true if this changes the value. */
-	bool SetHash(uint32 InHash);
+	bool SetHash(uint64 InHash);
 	/** Get the stored hash value. */
-	uint32 GetHash() const { return Hash; }
+	uint64 GetHash() const { return Hash; }
 
 	using FPerChannelLayerNames = TStaticArray<FName, 4>;
 
@@ -68,6 +68,6 @@ public:
 	static void GarbageCollectTasks();
 
 private:
-	uint32 Hash;
+	uint64 Hash;
 	TArray<int32> TaskHandles;
 };
