Fix non-unity compile error

diff --git a/Engine/Source/Runtime/Landscape/Classes/LandscapeGrassType.h b/Engine/Source/Runtime/Landscape/Classes/LandscapeGrassType.h
--- a/Engine/Source/Runtime/Landscape/Classes/LandscapeGrassType.h
+++ b/Engine/Source/Runtime/Landscape/Classes/LandscapeGrassType.h
@@ -8,6 +8,7 @@
 #include "Engine/EngineTypes.h"
 #include "PerPlatformProperties.h"
 #include "PerQualityLevelProperties.h"
+#include "SceneTypes.h"
 #include "LandscapeGrassType.generated.h"
 
 class UStaticMesh;
