Fixed some places where the original landscape guid should have been set as well.

#rnx

diff --git a/Engine/Source/Runtime/Landscape/Private/LandscapeEdit.cpp b/Engine/Source/Runtime/Landscape/Private/LandscapeEdit.cpp
--- a/Engine/Source/Runtime/Landscape/Private/LandscapeEdit.cpp
+++ b/Engine/Source/Runtime/Landscape/Private/LandscapeEdit.cpp
@@ -2775,7 +2775,7 @@ LANDSCAPE_API void ALandscapeProxy::Import(const FGuid& InGuid, int32 InMinX, in
 	ComponentSizeQuads = InNumSubsections * InSubsectionSizeQuads;
 	NumSubsections = InNumSubsections;
 	SubsectionSizeQuads = InSubsectionSizeQuads;
-	LandscapeGuid = InGuid;
+	SetLandscapeGuid(InGuid);
 
 	Modify();
 
@@ -4609,6 +4609,7 @@ void ALandscape::PostEditImport()
 		{
 			// Copy/Paste case, need to generate new GUID
 			LandscapeGuid = FGuid::NewGuid();
+			OriginalLandscapeGuid = LandscapeGuid;
 			break;
 		}
 	}
