Landscape - Skip running ALandscapeProxy::UpdateRenderMethod on tick outside of the editor.

Saves 50us on the game thread.
Can be re-enabled with landscape.UpdateProxyActorRenderMethodOnTickAtRuntime 1.

#rnx

diff --git a/Engine/Source/Runtime/Landscape/Private/LandscapeSubsystem.cpp b/Engine/Source/Runtime/Landscape/Private/LandscapeSubsystem.cpp
--- a/Engine/Source/Runtime/Landscape/Private/LandscapeSubsystem.cpp
+++ b/Engine/Source/Runtime/Landscape/Private/LandscapeSubsystem.cpp
@@ -7,6 +7,7 @@
 #include "Engine/World.h"
 #include "Modules/ModuleManager.h"
 #include "ContentStreaming.h"
+#include "HAL/IConsoleManager.h"
 #include "Landscape.h"
 #include "LandscapeEditTypes.h"
 #include "LandscapeProxy.h"
@@ -53,6 +54,12 @@ static FAutoConsoleVariable CVarMaxAsyncNaniteProxiesPerSecond(
 	6.0f,
 	TEXT("Number of Async nanite proxies to dispatch per second"));
 
+static bool GUpdateProxyActorRenderMethodOnTickAtRuntime = false;
+static FAutoConsoleVariableRef CVarUpdateProxyActorRenderMethodUpdateOnTickAtRuntime(
+	TEXT("landscape.UpdateProxyActorRenderMethodOnTickAtRuntime"),
+	GUpdateProxyActorRenderMethodOnTickAtRuntime,
+	TEXT("Update landscape proxy's rendering method (nanite enabled) when ticked. Always enabled in editor."));
+
 DECLARE_CYCLE_STAT(TEXT("LandscapeSubsystem Tick"), STAT_LandscapeSubsystemTick, STATGROUP_Landscape);
 
 #define LOCTEXT_NAMESPACE "LandscapeSubsystem"
@@ -87,6 +94,16 @@ void ULandscapeSubsystem::Initialize(FSubsystemCollectionBase& Collection)
 		}
 	}
 
+	if (IConsoleVariable* NaniteEnabledCVar = IConsoleManager::Get().FindConsoleVariable(TEXT("r.Nanite")))
+	{
+		NaniteEnabledCVar->OnChangedDelegate().AddUObject(this, &ULandscapeSubsystem::OnNaniteEnabledChanged);
+	}
+
+	if (IConsoleVariable* LandscapeNaniteEnabledCVar = IConsoleManager::Get().FindConsoleVariable(TEXT("landscape.RenderNanite")))
+	{
+		LandscapeNaniteEnabledCVar->OnChangedDelegate().AddUObject(this, &ULandscapeSubsystem::OnNaniteEnabledChanged);
+	}
+
 #if WITH_EDITOR
 	GrassMapsBuilder = new FLandscapeGrassMapsBuilder(GetWorld());
 	PhysicalMaterialBuilder = new FLandscapePhysicalMaterialBuilder(GetWorld());
@@ -112,6 +129,15 @@ void ULandscapeSubsystem::Deinitialize()
 		OnNaniteWorldSettingsChangedHandle.Reset();
 	}
 
+	if (IConsoleVariable* NaniteEnabledCVar = IConsoleManager::Get().FindConsoleVariable(TEXT("r.Nanite")))
+	{
+		NaniteEnabledCVar->OnChangedDelegate().RemoveAll(this);
+	}
+
+	if (IConsoleVariable* LandscapeNaniteEnabledCVar = IConsoleManager::Get().FindConsoleVariable(TEXT("landscape.RenderNanite")))
+	{
+		LandscapeNaniteEnabledCVar->OnChangedDelegate().RemoveAll(this);
+	}
 	
 	
 #if WITH_EDITOR
@@ -314,7 +340,12 @@ void ULandscapeSubsystem::Tick(float DeltaTime)
 				Proxy->TickGrass(*Cameras, InOutNumComponentsCreated);
 			}
 
-			Proxy->UpdateRenderingMethod();
+#if !WITH_EDITOR
+			if (GUpdateProxyActorRenderMethodOnTickAtRuntime)
+#endif
+			{
+				Proxy->UpdateRenderingMethod();
+			}
 		}
 	}
 
@@ -609,6 +640,19 @@ ALandscapeProxy* ULandscapeSubsystem::FindOrAddLandscapeProxy(ULandscapeInfo* La
 	return FLandscapeConfigHelper::FindOrAddLandscapeStreamingProxy(LandscapeInfo, SectionBase);
 }
 
+void ULandscapeSubsystem::OnNaniteEnabledChanged(IConsoleVariable*)
+{
+	QUICK_SCOPE_CYCLE_COUNTER(STAT_Landscape_OnNaniteEnabledChanged);
+
+	for (TWeakObjectPtr<ALandscapeProxy>& ProxyPtr : Proxies)
+	{
+		if (ALandscapeProxy* Proxy = ProxyPtr.Get())
+		{
+			Proxy->UpdateRenderingMethod();
+		}
+	}
+}
+
 void ULandscapeSubsystem::DisplayMessages(FCanvas* Canvas, float& XPos, float& YPos)
 {
 	const int32 FontSizeY = 20;
diff --git a/Engine/Source/Runtime/Landscape/Public/LandscapeSubsystem.h b/Engine/Source/Runtime/Landscape/Public/LandscapeSubsystem.h
--- a/Engine/Source/Runtime/Landscape/Public/LandscapeSubsystem.h
+++ b/Engine/Source/Runtime/Landscape/Public/LandscapeSubsystem.h
@@ -9,6 +9,7 @@
 
 class ALandscapeProxy;
 class AWorldSettings;
+class IConsoleVariable;
 class ULandscapeInfo;
 class FLandscapeNotificationManager;
 struct FDateTime;
@@ -128,6 +129,7 @@ private:
 	// End USubsystem
 
 	void OnNaniteWorldSettingsChanged(AWorldSettings* WorldSettings) { RegenerateGrass(true, true); }
+	void OnNaniteEnabledChanged(IConsoleVariable*);
 
 	bool bIsGrassCreationPrioritized = false;
 	TArray<TWeakObjectPtr<ALandscapeProxy>> Proxies;
