Landscape - make console variables static to avoid perf degredation when cooking as it creates many worlds.

#rnx

diff --git a/Engine/Source/Runtime/Landscape/Private/LandscapeSubsystem.cpp b/Engine/Source/Runtime/Landscape/Private/LandscapeSubsystem.cpp
--- a/Engine/Source/Runtime/Landscape/Private/LandscapeSubsystem.cpp
+++ b/Engine/Source/Runtime/Landscape/Private/LandscapeSubsystem.cpp
@@ -94,12 +94,14 @@ void ULandscapeSubsystem::Initialize(FSubsystemCollectionBase& Collection)
 		}
 	}
 
-	if (IConsoleVariable* NaniteEnabledCVar = IConsoleManager::Get().FindConsoleVariable(TEXT("r.Nanite")))
+	static IConsoleVariable* NaniteEnabledCVar = IConsoleManager::Get().FindConsoleVariable(TEXT("r.Nanite"));
+	if (NaniteEnabledCVar && !NaniteEnabledCVar->OnChangedDelegate().IsBoundToObject(this))
 	{
 		NaniteEnabledCVar->OnChangedDelegate().AddUObject(this, &ULandscapeSubsystem::OnNaniteEnabledChanged);
 	}
 
-	if (IConsoleVariable* LandscapeNaniteEnabledCVar = IConsoleManager::Get().FindConsoleVariable(TEXT("landscape.RenderNanite")))
+	static IConsoleVariable* LandscapeNaniteEnabledCVar = IConsoleManager::Get().FindConsoleVariable(TEXT("landscape.RenderNanite"));
+	if (LandscapeNaniteEnabledCVar && !LandscapeNaniteEnabledCVar->OnChangedDelegate().IsBoundToObject(this))
 	{
 		LandscapeNaniteEnabledCVar->OnChangedDelegate().AddUObject(this, &ULandscapeSubsystem::OnNaniteEnabledChanged);
 	}
@@ -129,12 +131,14 @@ void ULandscapeSubsystem::Deinitialize()
 		OnNaniteWorldSettingsChangedHandle.Reset();
 	}
 
-	if (IConsoleVariable* NaniteEnabledCVar = IConsoleManager::Get().FindConsoleVariable(TEXT("r.Nanite")))
+	static IConsoleVariable* NaniteEnabledCVar = IConsoleManager::Get().FindConsoleVariable(TEXT("r.Nanite"));
+	if (NaniteEnabledCVar)
 	{
 		NaniteEnabledCVar->OnChangedDelegate().RemoveAll(this);
 	}
 
-	if (IConsoleVariable* LandscapeNaniteEnabledCVar = IConsoleManager::Get().FindConsoleVariable(TEXT("landscape.RenderNanite")))
+	static IConsoleVariable* LandscapeNaniteEnabledCVar = IConsoleManager::Get().FindConsoleVariable(TEXT("landscape.RenderNanite"));
+	if (LandscapeNaniteEnabledCVar)
 	{
 		LandscapeNaniteEnabledCVar->OnChangedDelegate().RemoveAll(this);
 	}
