Tentative fix for live crash in ULandscapeInfo::GetLandscapeProxy

diff --git a/Engine/Source/Runtime/Landscape/Private/Landscape.cpp b/Engine/Source/Runtime/Landscape/Private/Landscape.cpp
--- a/Engine/Source/Runtime/Landscape/Private/Landscape.cpp
+++ b/Engine/Source/Runtime/Landscape/Private/Landscape.cpp
@@ -4197,8 +4197,9 @@ ALandscapeProxy* ULandscapeInfo::GetLandscapeProxy() const
 	if (LandscapeActor.IsValid())
 	{
 		ALandscape* Landscape = LandscapeActor.Get();
-		if (Landscape != nullptr &&
-			Landscape->GetRootComponent()->IsRegistered())
+		USceneComponent* LandscapeRootComponent = (Landscape != nullptr) ? Landscape->GetRootComponent() : nullptr;
+
+		if ((LandscapeRootComponent != nullptr) && (LandscapeRootComponent->IsRegistered()))
 		{
 			return Landscape;
 		}
@@ -4214,8 +4215,9 @@ ALandscapeProxy* ULandscapeInfo::GetLandscapeProxy() const
 	for (TWeakObjectPtr<ALandscapeStreamingProxy> ProxyPtr : StreamingProxies)
 	{
 		ALandscapeStreamingProxy* Proxy = ProxyPtr.Get();
-		if (Proxy != nullptr &&
-			Proxy->GetRootComponent()->IsRegistered())
+		USceneComponent* ProxyRootComponent = (Proxy != nullptr) ? Proxy->GetRootComponent() : nullptr;
+
+		if ((ProxyRootComponent != nullptr) && (ProxyRootComponent->IsRegistered()))
 		{
 			return Proxy;
 		}
