Fixed smooth/flatten tools that were not appropriately copying back changes on all mips

diff --git a/Engine/Source/Runtime/Landscape/Private/LandscapeEditLayers.cpp b/Engine/Source/Runtime/Landscape/Private/LandscapeEditLayers.cpp
--- a/Engine/Source/Runtime/Landscape/Private/LandscapeEditLayers.cpp
+++ b/Engine/Source/Runtime/Landscape/Private/LandscapeEditLayers.cpp
@@ -5610,30 +5610,24 @@ bool ALandscape::ResolveLayersTexture(
 		// Copy final result to texture source.
 		TArray<TArray<FColor>> const& OutMipsData = InCPUReadback->GetResult(CompletedReadbackNum - 1);
 
-		int32 NumLockedMips = 0;
 		for (int8 MipIndex = 0; MipIndex < OutMipsData.Num(); ++MipIndex)
 		{
 			int32 TextureSize = OutMipsData[MipIndex].Num();
 			if (TextureSize > 0)
 			{
 				uint8* TextureData = InOutputTexture->Source.LockMip(MipIndex);
-				++NumLockedMips;
 
-				// Do dirty detection on first mip :
-				if (MipIndex == 0)
+				// Do dirty detection on first mip.
+				// Don't do this for intermediate renders.
+				if (MipIndex == 0 && !bIntermediateRender)
 				{
-					// But don't do this for intermediate renders, we'll want the data back on CPU all the time in that case : 
-					if (bIntermediateRender)
-					{
-						TRACE_CPUPROFILER_EVENT_SCOPE(ReadbackToCPU);
-						FMemory::Memcpy(TextureData, OutMipsData[MipIndex].GetData(), OutMipsData[MipIndex].Num() * sizeof(FColor));
-					}
-
 					uint64 Hash = 0;
 					if (HasTextureDataChanged(MakeArrayView(reinterpret_cast<const FColor*>(TextureData), TextureSize), MakeArrayView(OutMipsData[MipIndex].GetData(), TextureSize), bIsWeightmap, InCPUReadback->GetHash(), Hash))
 					{
-						bChanged |= InCPUReadback->SetHash(Hash);
+						TRACE_CPUPROFILER_EVENT_SCOPE(LandscapeLayers_CalculateHash);
+						Hash = FLandscapeEditLayerReadback::CalculateHash((uint8*)OutMipsData[MipIndex].GetData(), TextureSize * sizeof(FColor));
 					}
+					bChanged |= InCPUReadback->SetHash(Hash);
 
 					if (bChanged)
 					{
@@ -5650,7 +5644,6 @@ bool ALandscape::ResolveLayersTexture(
 					}
 				}
 
-				// Debug utils for when we detect changes :
 				if (bChanged)
 				{
 					if (bIsWeightmap)
@@ -5661,15 +5654,15 @@ bool ALandscape::ResolveLayersTexture(
 					{
 						OnDirtyHeightmap(MapHelper, InOutputTexture, (FColor*)TextureData, OutMipsData[MipIndex].GetData(), MipIndex);
 					}
-
-					TRACE_CPUPROFILER_EVENT_SCOPE(ReadbackToCPU);
-					FMemory::Memcpy(TextureData, OutMipsData[MipIndex].GetData(), OutMipsData[MipIndex].Num() * sizeof(FColor));
 				}
+
+				TRACE_CPUPROFILER_EVENT_SCOPE(ReadbackToCPU);
+				FMemory::Memcpy(TextureData, OutMipsData[MipIndex].GetData(), OutMipsData[MipIndex].Num() * sizeof(FColor));
 			}
 		}
 
 		// Unlock all mips at once because there's a lock counter in FTextureSource that recomputes the content hash when reaching 0 (which means we'd recompute the hash several times over if we Lock/Unlock/Lock/Unloc/... for each mip ):
-		for (int8 MipIndex = 0; MipIndex < NumLockedMips; ++MipIndex)
+		for (int8 MipIndex = 0; MipIndex < OutMipsData.Num(); ++MipIndex)
 		{
 			if (OutMipsData[MipIndex].Num() > 0)
 			{
