Landscape : Added new CVar to control heightmap detailed change tracking

Split landscape.DirtyHeightmapThreshold into  landscape.DirtyHeightmapHeightThreshold and landscape.DirtyHeightmapNormalThreshold. Apart from this being useful for treating normals & heights separately this change will allow for us to experiment with disabling invalidating on normal changes by setting  landscape.DirtyHeightmapNormalThreshold to something greater than 255

diff --git a/Engine/Source/Runtime/Landscape/Private/LandscapeEditLayers.cpp b/Engine/Source/Runtime/Landscape/Private/LandscapeEditLayers.cpp
--- a/Engine/Source/Runtime/Landscape/Private/LandscapeEditLayers.cpp
+++ b/Engine/Source/Runtime/Landscape/Private/LandscapeEditLayers.cpp
@@ -143,10 +143,15 @@ static TAutoConsoleVariable<bool> CVarLandscapeDumpDiffDetails(
 	false,
 	TEXT("When dumping diffs for heightmap (landscape.DumpHeightmapDiff) or weightmap (landscape.DumpWeightmapDiff), dumps additional details about the pixels being different"));
 
-TAutoConsoleVariable<int32> CVarLandscapeDirtyHeightmapThreshold(
-	TEXT("landscape.DirtyHeightmapThreshold"),
+TAutoConsoleVariable<int32> CVarLandscapeDirtyHeightmapHeightThreshold(
+	TEXT("landscape.DirtyHeightmapHeightThreshold"),
 	0,
-	TEXT("Threshold to avoid imprecision issues on certain GPUs when detecting when a heightmap changes, i.e. only a difference > than this threshold (N over 16-bits uint height and N over each 8-bits uint normal component) will be detected as a change."));
+	TEXT("Threshold to avoid imprecision issues on certain GPUs when detecting when a heightmap height changes, i.e. only a height difference > than this threshold (N over 16-bits uint height) will be detected as a change."));
+
+TAutoConsoleVariable<int32> CVarLandscapeDirtyHeightmapNormalThreshold(
+	TEXT("landscape.DirtyHeightmapNormalThreshold"),
+	0,
+	TEXT("Threshold to avoid imprecision issues on certain GPUs when detecting when a heightmap normal changes, i.e. only a normal channel difference > than this threshold (N over each 8-bits uint B & A channels independently) will be detected as a change."));
 
 TAutoConsoleVariable<int32> CVarLandscapeDirtyWeightmapThreshold(
 	TEXT("landscape.DirtyWeightmapThreshold"),
@@ -5492,6 +5497,7 @@ bool ALandscape::HasTextureDataChanged(TArrayView<const FColor> InOldData, TArra
 	{
 		int32 TextureSize = InOldData.Num();
 		check(TextureSize == InNewData.Num());
+
 		// If necessary, perform a deep comparison to avoid taking into account precision issues as actual changes :
 		if (bInIsWeightmap)
 		{
@@ -5527,34 +5533,39 @@ bool ALandscape::HasTextureDataChanged(TArrayView<const FColor> InOldData, TArra
 				return false;
 			}
 		}
-		else if (int32 DirtyHeightmapThreshold = CVarLandscapeDirtyHeightmapThreshold.GetValueOnGameThread(); DirtyHeightmapThreshold > 0)
+		else 
 		{
-			TRACE_CPUPROFILER_EVENT_SCOPE(DeepCompareHeightmap);
-			for (int32 Index = 0; Index < TextureSize; ++Index)
+			int32 DirtyHeightmapHeightThreshold = CVarLandscapeDirtyHeightmapHeightThreshold.GetValueOnGameThread();
+			int32 DirtyHeightmapNormalThreshold = CVarLandscapeDirtyHeightmapNormalThreshold.GetValueOnGameThread();
+			if (DirtyHeightmapHeightThreshold > 0 || DirtyHeightmapNormalThreshold > 0)
 			{
-				const FColor& OldColor = InOldData[Index];
-				const FColor& NewColor = InNewData[Index];
-				if (OldColor != NewColor)
+				TRACE_CPUPROFILER_EVENT_SCOPE(DeepCompareHeightmap);
+				for (int32 Index = 0; Index < TextureSize; ++Index)
 				{
-					uint16 OldHeight = ((static_cast<uint16>(OldColor.R) << 8) | static_cast<uint16>(OldColor.G));
-					uint16 NewHeight = ((static_cast<uint16>(NewColor.R) << 8) | static_cast<uint16>(NewColor.G));
-					if (uint16 Diff = (NewHeight > OldHeight) ? (NewHeight - OldHeight) : (OldHeight - NewHeight); Diff > DirtyHeightmapThreshold)
-					{
-						return true;
-					}
-					if (uint8 Diff = (NewColor.B > OldColor.B) ? (NewColor.B - OldColor.B) : (OldColor.B - NewColor.B); Diff > DirtyHeightmapThreshold)
-					{
-						return true;
-					}
-					if (uint8 Diff = (NewColor.A > OldColor.A) ? (NewColor.A - OldColor.A) : (OldColor.A - NewColor.A); Diff > DirtyHeightmapThreshold)
+					const FColor& OldColor = InOldData[Index];
+					const FColor& NewColor = InNewData[Index];
+					if (OldColor != NewColor)
 					{
-						return true;
+						uint16 OldHeight = ((static_cast<uint16>(OldColor.R) << 8) | static_cast<uint16>(OldColor.G));
+						uint16 NewHeight = ((static_cast<uint16>(NewColor.R) << 8) | static_cast<uint16>(NewColor.G));
+						if (uint16 Diff = (NewHeight > OldHeight) ? (NewHeight - OldHeight) : (OldHeight - NewHeight); Diff > DirtyHeightmapHeightThreshold)
+						{
+							return true;
+						}
+						if (uint8 Diff = (NewColor.B > OldColor.B) ? (NewColor.B - OldColor.B) : (OldColor.B - NewColor.B); Diff > DirtyHeightmapNormalThreshold)
+						{
+							return true;
+						}
+						if (uint8 Diff = (NewColor.A > OldColor.A) ? (NewColor.A - OldColor.A) : (OldColor.A - NewColor.A); Diff > DirtyHeightmapNormalThreshold)
+						{
+							return true;
+						}
 					}
 				}
-			}
 
-			// No significant difference detected : 
-			return false;
+				// No significant difference detected : 
+				return false;
+			}
 		}
 
 		// Hash is different and we don't need a deep comparison : 
