Fixed rare case of desynchronization between CachedLocalBox and heightmap data in case of landscape undo

diff --git a/Engine/Source/Runtime/Landscape/Private/LandscapeEdit.cpp b/Engine/Source/Runtime/Landscape/Private/LandscapeEdit.cpp
--- a/Engine/Source/Runtime/Landscape/Private/LandscapeEdit.cpp
+++ b/Engine/Source/Runtime/Landscape/Private/LandscapeEdit.cpp
@@ -704,6 +704,11 @@ void ULandscapeComponent::PostEditUndo()
 		
 	if (Proxy->HasLayersContent())
 	{
+		// Updating the bounds after undo will re-sync CachedLocalBox with the heightmap data. This is necessary because ULandscapeComponent might be part of 
+		//  the transaction buffer, but the final heightmaps aren't, so undoing might desynchronize them. If the heightmap update that is requested right after 
+		//  ends up modifying heightmap data, then CachedLocalBox will get updated properly again so this is a 'safe' way to fix the issue
+		UpdateCachedBounds();
+
 		const bool bUpdateAll = true;
 		RequestHeightmapUpdate(bUpdateAll);
 		RequestWeightmapUpdate(bUpdateAll);
