Landscape: updated CanUpdatePhysicalMaterial to ensure we have a SceneProxy to render as this is a prerequisite for updating the physical material.

diff --git a/Engine/Source/Runtime/Landscape/Private/LandscapeEdit.cpp b/Engine/Source/Runtime/Landscape/Private/LandscapeEdit.cpp
--- a/Engine/Source/Runtime/Landscape/Private/LandscapeEdit.cpp
+++ b/Engine/Source/Runtime/Landscape/Private/LandscapeEdit.cpp
@@ -1810,6 +1810,11 @@ bool ULandscapeComponent::CanUpdatePhysicalMaterial()
 		return false;
 	}
 
+	if (SceneProxy == nullptr)
+	{
+		return false;
+	}
+
 	return GetCollisionComponent() != nullptr;
 }
 
