[PCG] Fix a bug with the Landscape Blend Layer not sampling correctly.

diff --git a/Engine/Plugins/Experimental/PCG/Source/PCG/Private/Grid/PCGLandscapeCache.cpp b/Engine/Plugins/Experimental/PCG/Source/PCG/Private/Grid/PCGLandscapeCache.cpp
--- a/Engine/Plugins/Experimental/PCG/Source/PCG/Private/Grid/PCGLandscapeCache.cpp
+++ b/Engine/Plugins/Experimental/PCG/Source/PCG/Private/Grid/PCGLandscapeCache.cpp
@@ -110,7 +110,7 @@ FPCGLandscapeCacheEntry* FPCGLandscapeCacheEntry::CreateCacheEntry(ULandscapeInf
 				continue;
 			}
 
-			if (CDI.GetWeightmapTextureData(LayerInfo, LayerCache, /*bUseEditingLayer=*/false))
+			if (CDI.GetWeightmapTextureData(LayerInfo, LayerCache, /*bUseEditingLayer=*/false, /*bRemoveSubsectionDuplicates=*/true))
 			{
 				Result->LayerDataNames.Add(Layer.LayerName);
 				Result->LayerData.Emplace(std::move(LayerCache));
diff --git a/Engine/Source/Runtime/Landscape/Private/LandscapeDataAccess.cpp b/Engine/Source/Runtime/Landscape/Private/LandscapeDataAccess.cpp
--- a/Engine/Source/Runtime/Landscape/Private/LandscapeDataAccess.cpp
+++ b/Engine/Source/Runtime/Landscape/Private/LandscapeDataAccess.cpp
@@ -79,15 +79,15 @@ LANDSCAPE_API void FLandscapeComponentDataInterface::GetHeightmapTextureData(TAr
 	}
 }
 
-LANDSCAPE_API bool FLandscapeComponentDataInterface::GetWeightmapTextureData(ULandscapeLayerInfoObject* LayerInfo, TArray<uint8>& OutData, bool InUseEditingWeightmap)
+LANDSCAPE_API bool FLandscapeComponentDataInterface::GetWeightmapTextureData(ULandscapeLayerInfoObject* InLayerInfo, TArray<uint8>& OutData, bool bInUseEditingWeightmap, bool bInRemoveSubsectionDuplicates)
 {
 	int32 LayerIdx = INDEX_NONE;
-	const TArray<FWeightmapLayerAllocationInfo>& ComponentWeightmapLayerAllocations = Component->GetWeightmapLayerAllocations(InUseEditingWeightmap);
-	const TArray<UTexture2D*>& ComponentWeightmapTextures = Component->GetWeightmapTextures(InUseEditingWeightmap);
+	const TArray<FWeightmapLayerAllocationInfo>& ComponentWeightmapLayerAllocations = Component->GetWeightmapLayerAllocations(bInUseEditingWeightmap);
+	const TArray<UTexture2D*>& ComponentWeightmapTextures = Component->GetWeightmapTextures(bInUseEditingWeightmap);
 
 	for (int32 Idx = 0; Idx < ComponentWeightmapLayerAllocations.Num(); Idx++)
 	{
-		if (ComponentWeightmapLayerAllocations[Idx].LayerInfo == LayerInfo)
+		if (ComponentWeightmapLayerAllocations[Idx].LayerInfo == InLayerInfo)
 		{
 			LayerIdx = Idx;
 			break;
@@ -106,7 +106,11 @@ LANDSCAPE_API bool FLandscapeComponentDataInterface::GetWeightmapTextureData(ULa
 		return false;
 	}
 
-	int32 WeightmapSize = ((Component->SubsectionSizeQuads + 1) * Component->NumSubsections) >> MipLevel;
+	// If requested to skip the duplicate row/col of texture data
+	int32 WeightmapSize = bInRemoveSubsectionDuplicates ?
+		((Component->SubsectionSizeQuads * Component->NumSubsections) + 1) >> MipLevel :
+		((Component->SubsectionSizeQuads + 1) * Component->NumSubsections) >> MipLevel;
+	
 	OutData.Empty(FMath::Square(WeightmapSize));
 	OutData.AddUninitialized(FMath::Square(WeightmapSize));
 
@@ -119,7 +123,8 @@ LANDSCAPE_API bool FLandscapeComponentDataInterface::GetWeightmapTextureData(ULa
 
 	for (int32 i = 0; i < FMath::Square(WeightmapSize); i++)
 	{
-		OutData[i] = SrcTextureData[i * 4];
+		// If removing subsection duplicates, convert vertex to texel index
+		OutData[i] = bInRemoveSubsectionDuplicates ? SrcTextureData[VertexIndexToTexel(i) * sizeof(FColor)] : SrcTextureData[i * sizeof(FColor)];
 	}
 
 	DataInterface.UnlockMip(ComponentWeightmapTextures[ComponentWeightmapLayerAllocations[LayerIdx].WeightmapTextureIndex], MipLevel);
diff --git a/Engine/Source/Runtime/Landscape/Public/LandscapeDataAccess.h b/Engine/Source/Runtime/Landscape/Public/LandscapeDataAccess.h
--- a/Engine/Source/Runtime/Landscape/Public/LandscapeDataAccess.h
+++ b/Engine/Source/Runtime/Landscape/Public/LandscapeDataAccess.h
@@ -297,7 +297,7 @@ struct FLandscapeComponentDataInterface : public FLandscapeComponentDataInterfac
 	/* Return the raw heightmap data exactly same size for Heightmap texture which belong to only this component */
 	LANDSCAPE_API void GetHeightmapTextureData(TArray<FColor>& OutData, bool bOkToFail = false);
 
-	LANDSCAPE_API bool GetWeightmapTextureData(ULandscapeLayerInfoObject* LayerInfo, TArray<uint8>& OutData, bool InUseEditingWeightmap = false);
+	LANDSCAPE_API bool GetWeightmapTextureData(ULandscapeLayerInfoObject* InLayerInfo, TArray<uint8>& OutData, bool bInUseEditingWeightmap = false, bool bInRemoveSubsectionDuplicates = false);
 
 	FColor* GetHeightData(int32 LocalX, int32 LocalY) const
 	{
