Landscape: collision mesh shows the chaos material on the vertex colors.

diff --git a/Engine/Source/Runtime/Landscape/Private/LandscapeCollision.cpp b/Engine/Source/Runtime/Landscape/Private/LandscapeCollision.cpp
--- a/Engine/Source/Runtime/Landscape/Private/LandscapeCollision.cpp
+++ b/Engine/Source/Runtime/Landscape/Private/LandscapeCollision.cpp
@@ -62,6 +62,7 @@
 #include "Chaos/ChaosArchive.h"
 #include "PhysicsProxy/SingleParticlePhysicsProxy.h"
 #include "Chaos/Framework/PhysicsSolverBase.h"
+#include "Chaos/Defines.h"
 #include "PBDRigidsSolver.h"
 
 using namespace PhysicsInterfaceTypes;
@@ -69,14 +70,31 @@ using namespace PhysicsInterfaceTypes;
 // Global switch for whether to read/write to DDC for landscape cooked data
 // It's a lot faster to compute than to request from DDC, so always skip.
 bool GLandscapeCollisionSkipDDC = true;
-static int32 CVarLandscapeShowCollisionMeshCurrentValue = static_cast<int>(EHeightfieldSource::Simple);
-static TAutoConsoleVariable<int32> CVarLandscapeShowCollisionMesh(
-	TEXT("landscape.ShowCollisionMesh"),
-	CVarLandscapeShowCollisionMeshCurrentValue,
+
+static int32 CVarLandscapeCollisionMeshShowValue = static_cast<int>(EHeightfieldSource::Simple);
+static TAutoConsoleVariable<int32> CVarLandscapeCollisionMeshShow(
+	TEXT("landscape.CollisionMesh.Show"),
+	CVarLandscapeCollisionMeshShowValue,
 	TEXT("Selects which heightfield to visualize when ShowFlags.Collision is used. 0 to disable, 1 for simple, 2 for complex, 3 for editor only."),
 	ECVF_RenderThreadSafe
 );
 
+static float CVarLandscapeCollisionMeshHeightOffsetValue = 0.f;
+static TAutoConsoleVariable<float> CVarLandscapeCollisionMeshHeightOffset(
+	TEXT("landscape.CollisionMesh.HeightOffset"),
+	CVarLandscapeCollisionMeshHeightOffsetValue,
+	TEXT("Offsets the collision mesh wireframe to assist in viewing from distances where the lower landscape lods might hide it."),
+	ECVF_RenderThreadSafe
+);
+
+static bool CVarLandscapeCollisionMeshShowPhysicalMaterialValue = false;
+static TAutoConsoleVariable<bool> CVarLandscapeCollisionMeshShowPhysicalMaterial(
+	TEXT("landscape.CollisionMesh.ShowPhysicalMaterial"),
+	CVarLandscapeCollisionMeshShowPhysicalMaterialValue,
+	TEXT("When enabled, vertex colors of the collision mesh are chosen based on the physical material"),
+	ECVF_RenderThreadSafe
+);
+
 #if ENABLE_COOK_STATS
 namespace LandscapeCollisionCookStats
 {
@@ -471,14 +489,26 @@ void ULandscapeHeightfieldCollisionComponent::ApplyWorldOffset(const FVector& In
 	}
 }
 
-// Callback to flag scene proxy as dirty when cvar changes
-static void OnLandscapeShowCollisionMeshChanged()
+template<typename Type>
+bool LandscapeCollisionCVarChanged(const TAutoConsoleVariable<Type>& CVar, Type& OldValue)
 {
-	const int32 Value = CVarLandscapeShowCollisionMesh.GetValueOnAnyThread();
-	if (CVarLandscapeShowCollisionMeshCurrentValue != Value)
+	const Type NewValue = CVar.GetValueOnAnyThread();
+
+	if (NewValue != OldValue)
 	{
-		CVarLandscapeShowCollisionMeshCurrentValue = Value;
+		OldValue = NewValue;
+		return true;
+	}
+	return false;
+}
 
+// Callbacks to flag scene proxy as dirty when cvars changes
+static void OnLandscapeShowCollisionMeshChanged()
+{
+	if (LandscapeCollisionCVarChanged(CVarLandscapeCollisionMeshShow, CVarLandscapeCollisionMeshShowValue)
+		|| LandscapeCollisionCVarChanged(CVarLandscapeCollisionMeshHeightOffset, CVarLandscapeCollisionMeshHeightOffsetValue)
+		|| LandscapeCollisionCVarChanged(CVarLandscapeCollisionMeshShowPhysicalMaterial, CVarLandscapeCollisionMeshShowPhysicalMaterialValue))
+	{
 		for (ULandscapeHeightfieldCollisionComponent* LandscapeHeightfieldCollisionComponent : TObjectRange<ULandscapeHeightfieldCollisionComponent>(RF_ClassDefaultObject | RF_ArchetypeObject, true, EInternalObjectFlags::Garbage))
 		{
 			LandscapeHeightfieldCollisionComponent->MarkRenderStateDirty();
@@ -486,8 +516,8 @@ static void OnLandscapeShowCollisionMeshChanged()
 	}
 }
 
-// Registers callback with cvar
-static FAutoConsoleVariableSink OnLandscapeCollisionHeightfieldChangedSink(FConsoleCommandDelegate::CreateStatic(&OnLandscapeShowCollisionMeshChanged));
+// Registers callback with cvars
+static FAutoConsoleVariableSink OnLandscapeCollisionMeshChanged(FConsoleCommandDelegate::CreateStatic(&OnLandscapeShowCollisionMeshChanged));
 
 #if !(UE_BUILD_SHIPPING || UE_BUILD_TEST) && WITH_EDITORONLY_DATA
 FPrimitiveSceneProxy* ULandscapeHeightfieldCollisionComponent::CreateSceneProxy()
@@ -503,7 +533,7 @@ FPrimitiveSceneProxy* ULandscapeHeightfieldCollisionComponent::CreateSceneProxy(
 
 		// Constructor exists to populate Vertices and Indices arrays which are
 		// used to construct the collision mesh inside GetDynamicMeshElements
-		FLandscapeHeightfieldCollisionComponentSceneProxy(const ULandscapeHeightfieldCollisionComponent* InComponent, const Chaos::FHeightField& InHeightfield, const FLinearColor& InWireframeColor)
+		FLandscapeHeightfieldCollisionComponentSceneProxy(const ULandscapeHeightfieldCollisionComponent* InComponent, const TArray<Chaos::FMaterialHandle>& InUsedChaosMaterials, const Chaos::FHeightField& InHeightfield, const FLinearColor& InWireframeColor)
 			: FPrimitiveSceneProxy(InComponent)
 		{
 			const Chaos::FHeightField::FData<uint16>& GeomData = InHeightfield.GeomData; 
@@ -512,10 +542,40 @@ FPrimitiveSceneProxy* ULandscapeHeightfieldCollisionComponent::CreateSceneProxy(
 			const int32 NumVerts = NumRows * NumCols;
 			const int32 NumTris = (NumRows - 1) * (NumCols - 1) * 2;
 			Vertices.SetNumUninitialized(NumVerts);
+
+			TArray<FColor, TInlineAllocator<16>> MaterialIndexColors;
+			MaterialIndexColors.Reserve(InUsedChaosMaterials.Num());
+			for (Chaos::FMaterialHandle MaterialHandle : InUsedChaosMaterials)
+			{
+				Chaos::FChaosPhysicsMaterial* ChaosMaterial = MaterialHandle.Get();
+				UPhysicalMaterial* PhysicalMaterial = (ChaosMaterial != nullptr) ? FChaosUserData::Get<UPhysicalMaterial>(ChaosMaterial->UserData) : nullptr;
+				uint32 Hash = PointerHash(PhysicalMaterial);
+				uint8* HashElement = reinterpret_cast<uint8*>(&Hash);
+				MaterialIndexColors.Emplace(HashElement[0], HashElement[1], HashElement[2]);
+			}
+
+			const float HeightOffset = CVarLandscapeCollisionMeshHeightOffset.GetValueOnGameThread();
+
 			for (int32 I = 0; I < NumVerts; I++)
 			{
 				const Chaos::FVec3 Point = GeomData.GetPointScaled(I);
-				Vertices[I].Position = FVector3f(static_cast<float>(Point.X), static_cast<float>(Point.Y), static_cast<float>(Point.Z));
+				uint8 MaterialIndex = InHeightfield.GetMaterialIndex(I % NumCols, I / NumCols);
+				Vertices[I].Position = FVector3f(static_cast<float>(Point.X), static_cast<float>(Point.Y), static_cast<float>(Point.Z + HeightOffset));
+
+				// Material indices are not defined for the last row/column in each component since they are per-triangle and not per-vertex.
+				// To show something intuitive for the user, we simply extend the previous vertex onto the final one.
+				if (I % NumCols == NumCols - 1)
+				{
+					Vertices[I].Color = Vertices[I - 1].Color;
+				}
+				else if (I / NumCols == NumCols - 1)
+				{
+					Vertices[I].Color = Vertices[I - NumRows - 1].Color;
+				}
+				else
+				{
+					Vertices[I].Color = (MaterialIndex == 255) ? FColor::Black  : MaterialIndexColors[MaterialIndex];
+				}
 			}
 			Indices.SetNumUninitialized(NumTris * 3);
 
@@ -573,8 +633,12 @@ FPrimitiveSceneProxy* ULandscapeHeightfieldCollisionComponent::CreateSceneProxy(
 			}
 
 			WireframeMaterialInstance.Reset(new FColoredMaterialRenderProxy(
-				GEngine->WireframeMaterial ? GEngine->WireframeMaterial->GetRenderProxy() : NULL,
+				GEngine->VertexColorMaterial ? GEngine->WireframeMaterial->GetRenderProxy() : nullptr,
 				InWireframeColor));
+
+			VertexColorMaterialInstance.Reset(new FColoredMaterialRenderProxy(
+				GEngine->VertexColorMaterial ? GEngine->VertexColorMaterial->GetRenderProxy() : nullptr,
+				FColor::White));
 		}
 
 		virtual void GetDynamicMeshElements(const TArray<const FSceneView*>& Views, const FSceneViewFamily& ViewFamily, uint32 VisibilityMap, FMeshElementCollector& Collector) const override
@@ -584,7 +648,11 @@ FPrimitiveSceneProxy* ULandscapeHeightfieldCollisionComponent::CreateSceneProxy(
 
 			const bool bDrawCollision = ViewFamily.EngineShowFlags.Collision && IsCollisionEnabled();
 
-			if (bDrawCollision && AllowDebugViewmodes() && WireframeMaterialInstance.IsValid())
+			const bool bShowPhysicalMaterial = CVarLandscapeCollisionMeshShowPhysicalMaterial.GetValueOnRenderThread();
+
+			const TUniquePtr<FColoredMaterialRenderProxy>& MaterialToUse = bShowPhysicalMaterial ? VertexColorMaterialInstance : WireframeMaterialInstance;
+
+			if (bDrawCollision && AllowDebugViewmodes() && MaterialToUse.IsValid())
 			{
 				for (int32 ViewIndex = 0; ViewIndex < Views.Num(); ViewIndex++)
 				{
@@ -596,7 +664,13 @@ FPrimitiveSceneProxy* ULandscapeHeightfieldCollisionComponent::CreateSceneProxy(
 						MeshBuilder.AddVertices(Vertices);
 						MeshBuilder.AddTriangles(Indices);
 
-						MeshBuilder.GetMesh(LocalToWorldNoScale, WireframeMaterialInstance.Get(), SDPG_World, false, false, ViewIndex, Collector);
+						FDynamicMeshBuilderSettings Settings;
+						Settings.bWireframe = true;
+						Settings.bDisableBackfaceCulling = false;
+						Settings.bReceivesDecals = false;
+						Settings.CastShadow = false;
+						Settings.bUseSelectionOutline = false;
+						MeshBuilder.GetMesh(LocalToWorldNoScale, MaterialToUse.Get(), SDPG_World, Settings, nullptr, ViewIndex, Collector);
 					}
 				}
 			}
@@ -610,6 +684,7 @@ FPrimitiveSceneProxy* ULandscapeHeightfieldCollisionComponent::CreateSceneProxy(
 			FPrimitiveViewRelevance Result;
 			Result.bDrawRelevance = IsShown(View) || bShowForCollision;
 			Result.bDynamicRelevance = true;
+			Result.bStaticRelevance = false;
 			Result.bShadowRelevance = false;
 			Result.bEditorPrimitiveRelevance = UseEditorCompositing(View);
 			return Result;
@@ -622,6 +697,7 @@ FPrimitiveSceneProxy* ULandscapeHeightfieldCollisionComponent::CreateSceneProxy(
 		TArray<FDynamicMeshVertex> Vertices;
 		TArray<uint32> Indices;
 		TUniquePtr<FColoredMaterialRenderProxy> WireframeMaterialInstance = nullptr;
+		TUniquePtr<FColoredMaterialRenderProxy> VertexColorMaterialInstance = nullptr;
 	};
 
 	FLandscapeHeightfieldCollisionComponentSceneProxy* Proxy = nullptr;
@@ -636,7 +712,7 @@ FPrimitiveSceneProxy* ULandscapeHeightfieldCollisionComponent::CreateSceneProxy(
 		const Chaos::FHeightField* LocalHeightfield = nullptr;
 		FLinearColor WireframeColor;
 
-		switch (static_cast<EHeightfieldSource>(CVarLandscapeShowCollisionMesh.GetValueOnGameThread()))
+		switch (static_cast<EHeightfieldSource>(CVarLandscapeCollisionMeshShow.GetValueOnGameThread()))
 		{
 		case EHeightfieldSource::None:
 			WireframeColor = FColor(0, 0, 0, 0);
@@ -678,7 +754,7 @@ FPrimitiveSceneProxy* ULandscapeHeightfieldCollisionComponent::CreateSceneProxy(
 
 		if (LocalHeightfield != nullptr)
 		{
-			Proxy = new FLandscapeHeightfieldCollisionComponentSceneProxy(this, *LocalHeightfield, WireframeColor);
+			Proxy = new FLandscapeHeightfieldCollisionComponentSceneProxy(this, HeightfieldRef->UsedChaosMaterials, *LocalHeightfield, WireframeColor);
 		}
 	}
 
