[Landscape Splines] Method to convert points along Landscape Spline into points on a Spline Component

diff --git a/Engine/Source/Runtime/Landscape/Classes/LandscapeSplinesComponent.h b/Engine/Source/Runtime/Landscape/Classes/LandscapeSplinesComponent.h
--- a/Engine/Source/Runtime/Landscape/Classes/LandscapeSplinesComponent.h
+++ b/Engine/Source/Runtime/Landscape/Classes/LandscapeSplinesComponent.h
@@ -7,6 +7,7 @@
 #include "Misc/Guid.h"
 #include "LandscapeInfo.h"
 #include "Components/PrimitiveComponent.h"
+#include "Components/SplineComponent.h"
 #include "LandscapeSplinesComponent.generated.h"
 
 class ALandscapeProxy;
@@ -214,6 +215,13 @@ public:
 
 	// Iterates throug a copy of the ControlPoints list.
 	LANDSCAPE_API void ForEachControlPoint(TFunctionRef<void(ULandscapeSplineControlPoint*)> Func);
+
+	/**
+	 * Converts points along the landscape spline into points on a spline component
+	 * Works best with landscape splines that contain no intersections or splits
+	 * @param SplineComponent Spline component to be populated
+	**/
+	LANDSCAPE_API void CopyToSplineComponent(USplineComponent* SplineComponent);
 #endif
 
 	//~ Begin UObject Interface
diff --git a/Engine/Source/Runtime/Landscape/Private/LandscapeSplines.cpp b/Engine/Source/Runtime/Landscape/Private/LandscapeSplines.cpp
--- a/Engine/Source/Runtime/Landscape/Private/LandscapeSplines.cpp
+++ b/Engine/Source/Runtime/Landscape/Private/LandscapeSplines.cpp
@@ -63,6 +63,13 @@ static FAutoConsoleVariableRef CVarSplinesAlwaysUseBlockAll(
 	TEXT("Force splines to always use the BlockAll collision profile instead of whatever is stored in the CollisionProfileName property")
 );
 
+int32 LandscapeSplineToSplineComponentMaxIterations = 200;
+static FAutoConsoleVariableRef CVarLandscapeSplineToSplineComponentMaxIterations(
+	TEXT("Landscape.Splines.ApplyToSplineComponentMaxIterations"),
+	LandscapeSplineToSplineComponentMaxIterations,
+	TEXT("Max possible iterations when converting a landscape spline into a spline component")
+);
+
 //////////////////////////////////////////////////////////////////////////
 // LANDSCAPE SPLINES SCENE PROXY
 
@@ -751,6 +758,140 @@ void ULandscapeSplinesComponent::ForEachControlPoint(TFunctionRef<void(ULandscap
 	}
 }
 
+void ULandscapeSplinesComponent::CopyToSplineComponent(USplineComponent* SplineComponent)
+{
+	if (SplineComponent == nullptr)
+	{
+		return;
+	}
+
+	if (Segments.IsEmpty())
+	{
+		return;
+	}
+	
+	SplineComponent->ClearSplinePoints(false);
+
+	// Incremented index for spline component points
+	int32 CurrentPointIndex = 0;
+
+	// Traverse up to this many segments
+	// For safety.  Ensures we don't get caught in an infinite loop.
+	const int32 MAX_ITERATIONS = LandscapeSplineToSplineComponentMaxIterations;
+	int32 NumIterations = 0;
+
+	// List of segments we've visited.  Used to ensure we don't visit the same segment twice.
+	TSet<ULandscapeSplineSegment*> VisitedSegments;
+
+	// The last control point we used
+	// 2 control points per segment
+	// 2 segments can share the same control point - one being the arrive control point and one being the exit
+	ULandscapeSplineControlPoint* LastControlPoint = nullptr;
+
+	// Arrive Tangent to apply to the next spline point.  Populated as the previous spline mesh's End Tangent
+	FVector ArriveTangent;
+	
+	ULandscapeSplineSegment* CurrentSegment = Segments[0];
+	do
+	{
+		if (!ensure(NumIterations < MAX_ITERATIONS))
+		{
+			UE_LOG(LogLandscape, Error, TEXT("%s Reached Max Iterations. Exiting."), ANSI_TO_TCHAR(__FUNCTION__));
+			break;
+		}
+		
+		if (!ensure(CurrentSegment != nullptr))
+		{
+			UE_LOG(LogLandscape, Error, TEXT("%s Current Segment is nullptr. Exiting."), ANSI_TO_TCHAR(__FUNCTION__));
+			break;
+		}
+
+		if (!ensure(!VisitedSegments.Contains(CurrentSegment)))
+		{
+			UE_LOG(LogLandscape, Error, TEXT("%s Current Segment was already visited. Exiting."), ANSI_TO_TCHAR(__FUNCTION__));
+			break;
+		}
+
+		VisitedSegments.Add(CurrentSegment);
+
+		// No guarantee if the control point we're looking at is the segments start or end control point
+		// If it's the end control point, we must reverse the direction we traverse over this segment
+		// Connections[0] == Start; Connections[1] == End;
+		bool bReverseTraversal = CurrentSegment->Connections[1].ControlPoint == LastControlPoint;
+
+		// Walk down the spline meshes
+		TArray<USplineMeshComponent*> SegmentSplineMeshComponents = CurrentSegment->GetLocalMeshComponents();
+		if (!bReverseTraversal)
+		{
+			for (int Index = 0; Index < SegmentSplineMeshComponents.Num(); ++Index)
+			{
+				USplineMeshComponent* SplineMesh = SegmentSplineMeshComponents[Index];
+				if (SplineMesh == nullptr)
+				{
+					continue;
+				}
+
+				// Create a spline component point from each spline mesh in the segment
+				SplineComponent->AddSplinePoint(SplineMesh->GetStartPosition() + SplineMesh->GetComponentLocation(), ESplineCoordinateSpace::World, false);
+				SplineComponent->SetTangentsAtSplinePoint(CurrentPointIndex, ArriveTangent, SplineMesh->GetStartTangent(), ESplineCoordinateSpace::World, false);
+
+				++CurrentPointIndex;
+				ArriveTangent = SplineMesh->GetEndTangent();
+			}
+		}
+		else
+		{
+			// Same as the true block, but reversed direction
+			for (int Index = SegmentSplineMeshComponents.Num() - 1; Index >= 0; --Index)
+			{
+				USplineMeshComponent* SplineMesh = SegmentSplineMeshComponents[Index];
+				if (SplineMesh == nullptr)
+				{
+					continue;
+				}
+
+				// Swap start with end since we're going in reverse
+				SplineComponent->AddSplinePoint(SplineMesh->GetEndPosition() + SplineMesh->GetComponentLocation(), ESplineCoordinateSpace::World, false);
+				SplineComponent->SetTangentsAtSplinePoint(CurrentPointIndex, ArriveTangent, SplineMesh->GetEndTangent(), ESplineCoordinateSpace::World, false);
+
+				++CurrentPointIndex;
+				ArriveTangent = SplineMesh->GetStartTangent();
+			}
+		}
+		
+		// Search for the next segment
+		FLandscapeSplineSegmentConnection& Connection = bReverseTraversal ? CurrentSegment->Connections[0] : CurrentSegment->Connections[1];
+		bool bFoundNextSegment = false;
+		if (Connection.ControlPoint)
+		{
+			for (FLandscapeSplineConnection& ControlPointConnection : Connection.ControlPoint->ConnectedSegments)
+			{
+				if (ControlPointConnection.Segment == nullptr || VisitedSegments.Contains(ControlPointConnection.Segment))
+				{
+					continue;
+				}
+				
+				CurrentSegment = ControlPointConnection.Segment;
+				LastControlPoint = Connection.ControlPoint;
+				bFoundNextSegment = true;
+				break;
+			}
+		}
+
+		if (!bFoundNextSegment)
+		{
+			UE_LOG(LogLandscape, Verbose, TEXT("%s Could not find next segment or all segments have been traversed."), ANSI_TO_TCHAR(__FUNCTION__));
+			break;
+		}
+
+		++NumIterations;
+	}
+	while (NumIterations < MAX_ITERATIONS);
+
+	// Finish by updating the spline with all of our work
+	SplineComponent->UpdateSpline();
+}
+
 bool ULandscapeSplinesComponent::IsUsingLayerInfo(const ULandscapeLayerInfoObject* LayerInfo) const
 {
 	for (ULandscapeSplineControlPoint* ControlPoint : ControlPoints)
