Fixed undo of (Nanite-enabled) landscape proxy deletion that didn't properly invalidate the Nanite component on redo

diff --git a/Engine/Source/Runtime/Landscape/Private/Landscape.cpp b/Engine/Source/Runtime/Landscape/Private/Landscape.cpp
--- a/Engine/Source/Runtime/Landscape/Private/Landscape.cpp
+++ b/Engine/Source/Runtime/Landscape/Private/Landscape.cpp
@@ -3666,6 +3666,9 @@ void ALandscapeProxy::Destroyed()
 		NumTexturesToStreamForVisibleGrassMapRender = 0;
 	}
 
+	// Destroy the Nanite component when we get destroyed so that we don't restore a garbage Nanite component (it's non-transactional and will get regenerated anyway)
+	InvalidateNaniteRepresentation(/* bInCheckContentId = */false);
+
 	// unregister feature level changed handler for grass
 	if (FeatureLevelChangedDelegateHandle.IsValid())
 	{
diff --git a/Engine/Source/Runtime/Landscape/Private/LandscapeEdit.cpp b/Engine/Source/Runtime/Landscape/Private/LandscapeEdit.cpp
--- a/Engine/Source/Runtime/Landscape/Private/LandscapeEdit.cpp
+++ b/Engine/Source/Runtime/Landscape/Private/LandscapeEdit.cpp
@@ -671,7 +671,7 @@ void ULandscapeComponent::PreEditUndo()
 	
 	// Count the number of components involved in the undo/redo
 	UndoRedoModifiedComponentCount++;
-	check(UndoRedoModifiedComponents.Num() == 0);
+	check(UndoRedoModifiedComponents.Num() == 0);	
 }
 
 void ULandscapeComponent::PostEditUndo()
@@ -712,19 +712,18 @@ void ULandscapeComponent::PostEditUndo()
 		const bool bUpdateAll = true;
 		RequestHeightmapUpdate(bUpdateAll);
 		RequestWeightmapUpdate(bUpdateAll);
-
-		// Nanite data is non-transactional at the moment, so we need to invalidate the data now (it will get updated eventually after textures have been resolved): 
-		Proxy->InvalidateNaniteRepresentation(/* bInCheckContentId = */true);
 	}
 	else
 	{
 		TSet<ULandscapeComponent*> Components;
 		Components.Add(this);
 		Proxy->FlushGrassComponents(&Components);
-
-		Proxy->InvalidateOrUpdateNaniteRepresentation(/* bInCheckContentId = */true, /*InTargetPlatform = */nullptr);
 	}
 
+	// Nanite data is non-transactional so we simply want to regenerate it on undo (depending on whether auto-update is enabled, it will either trigger its 
+	//  rebuild or will simply invalidate it) : 
+	Proxy->InvalidateOrUpdateNaniteRepresentation(/* bInCheckContentId = */true, /*InTargetPlatform = */nullptr);
+
 	check(!UndoRedoModifiedComponents.Contains(this));
 	UndoRedoModifiedComponents.Add(this);
 
