Removed workaround for landscape tools on nVidia hardware now that the minimum recommended driver version was bumped with a version of the driver where the fix doesn't occur

diff --git a/Engine/Source/Runtime/Landscape/Private/LandscapeEditLayers.cpp b/Engine/Source/Runtime/Landscape/Private/LandscapeEditLayers.cpp
--- a/Engine/Source/Runtime/Landscape/Private/LandscapeEditLayers.cpp
+++ b/Engine/Source/Runtime/Landscape/Private/LandscapeEditLayers.cpp
@@ -183,14 +183,6 @@ TAutoConsoleVariable<int32> CVarLandscapeRemoveEmptyPaintLayersOnEdit(
 	1,
 	TEXT("This will analyze weightmaps on readback and remove unneeded allocations (for unpainted layers)."));
 
-// TODO [jonathan.bard] : find out which nvidia driver version fixed the issue (last known failing : 526.47, first known working : 531.41), bump the minimum recommended driver for the editor and remove this workaround : 
-// COMMENT [jonathan.bard] : this is a workaround for a random D3D12 issue that misses some cache flush when copying edit layers individual heightmaps to the global merge atlas and then using 
-//  the atlas as SRV right after, which should be ensured by the manually inserted barriers but is sometimes not). Using PS-based copies doesn't exhibit the behavior so we'll stick with that for now : 
-TAutoConsoleVariable<bool> CVarLandscapeUsePSCopyWorkaround(
-	TEXT("landscape.UsePSCopyWorkaround"),
-	1,
-	TEXT("This is a workaround for a random D3D12-only issue when using a COPY_DEST RT as a SRV."));
-
 void OnLandscapeEditLayersLocalMergeChanged(IConsoleVariable* CVar)
 {
 	for (TObjectIterator<UWorld> It; It; ++It)
@@ -1526,7 +1518,6 @@ struct FLandscapeLayersCopyTextureParams
 	uint32 DestArrayIndex = 0;
 	ERHIAccess SourceAccess = ERHIAccess::SRVMask;
 	ERHIAccess DestAccess = ERHIAccess::SRVMask;
-	bool bUsePS = false;
 };
 
 class FLandscapeLayersCopyTexture_RenderThread
@@ -1538,19 +1529,6 @@ public:
 
 	const FLandscapeLayersCopyTextureParams& GetParams() const { return Params; }
 	void Copy(FRHICommandListImmediate& InRHICmdList)
-	{
-		if (Params.bUsePS)
-		{
-			CopyInternalPS(InRHICmdList);
-		}
-		else
-		{
-			CopyInternal(InRHICmdList);
-		}
-	}
-
-private:
-	void CopyInternal(FRHICommandListImmediate& InRHICmdList)
 	{
 		SCOPED_GPU_STAT(InRHICmdList, LandscapeLayers_CopyTexture);
 		SCOPED_DRAW_EVENTF(InRHICmdList, LandscapeLayers, TEXT("LandscapeLayers_Copy %s -> %s, Mip (%d -> %d), Array Index (%d -> %d)"), *Params.SourceResourceDebugName, *Params.DestResourceDebugName, Params.SourceMip, Params.DestMip, Params.SourceArrayIndex, Params.DestArrayIndex);
@@ -1586,52 +1564,6 @@ private:
 		InRHICmdList.Transition(FRHITransitionInfo(Params.DestResource->TextureRHI, ERHIAccess::CopyDest, Params.DestAccess));
 	}
 
-	void CopyInternalPS(FRHICommandListImmediate& InRHICmdList)
-	{
-		SCOPED_GPU_STAT(InRHICmdList, LandscapeLayers_CopyTexturePS);
-		SCOPED_DRAW_EVENTF(InRHICmdList, LandscapeLayers, TEXT("LandscapeLayers_CopyPS %s -> %s, Mip (%d -> %d), Array Index (%d -> %d)"), *Params.SourceResourceDebugName, *Params.DestResourceDebugName, Params.SourceMip, Params.DestMip, Params.SourceArrayIndex, Params.DestArrayIndex);
-
-		FIntPoint SourceSize(Params.SourceResource->GetSizeX() >> Params.SourceMip, Params.SourceResource->GetSizeY() >> Params.SourceMip);
-		FIntPoint DestSize(Params.DestResource->GetSizeX() >> Params.DestMip, Params.DestResource->GetSizeY() >> Params.DestMip);
-
-		// If CopySize is passed, used that as the size (and don't adjust with the mip level : consider that the user has computed it properly) : 
-		FIntPoint Size;
-		Size.X = (Params.CopySize.X > 0) ? Params.CopySize.X : SourceSize.X;
-		Size.Y = (Params.CopySize.Y > 0) ? Params.CopySize.Y : SourceSize.Y;
-		check((Params.SourceArrayIndex == 0) && (Params.DestArrayIndex == 0) && (Params.SourceMip == 0) && (Params.DestMip == 0)); // The PS version of copy is not supported on texture arrays and mips for now
-
-		InRHICmdList.Transition(FRHITransitionInfo(Params.SourceResource->TextureRHI, Params.SourceAccess, ERHIAccess::SRVGraphics));
-		InRHICmdList.Transition(FRHITransitionInfo(Params.DestResource->TextureRHI, Params.DestAccess, ERHIAccess::RTV));
-
-		FRHIRenderPassInfo RPInfo(Params.DestResource->TextureRHI, ERenderTargetActions::DontLoad_Store);
-		InRHICmdList.BeginRenderPass(RPInfo, TEXT("CopyTexture"));
-
-		FGlobalShaderMap* GlobalShaderMap = GetGlobalShaderMap(GMaxRHIFeatureLevel);
-		TShaderMapRef< FLandscapeCopyTextureVS > VertexShader(GlobalShaderMap);
-		TShaderMapRef< FLandscapeCopyTexturePS > PixelShader(GlobalShaderMap);
-
-		FGraphicsPipelineStateInitializer GraphicsPSOInit;
-		InRHICmdList.ApplyCachedRenderTargets(GraphicsPSOInit);
-		GraphicsPSOInit.BlendState = TStaticBlendState<>::GetRHI();
-		GraphicsPSOInit.RasterizerState = TStaticRasterizerState<>::GetRHI();
-		GraphicsPSOInit.DepthStencilState = TStaticDepthStencilState<false, CF_Always>::GetRHI();
-		GraphicsPSOInit.PrimitiveType = PT_TriangleList;
-		GraphicsPSOInit.BoundShaderState.VertexDeclarationRHI = GetVertexDeclarationFVector4();
-		GraphicsPSOInit.BoundShaderState.VertexShaderRHI = VertexShader.GetVertexShader();
-		GraphicsPSOInit.BoundShaderState.PixelShaderRHI = PixelShader.GetPixelShader();
-		SetGraphicsPipelineState(InRHICmdList, GraphicsPSOInit, 0);
-
-		SetShaderParametersLegacyPS(InRHICmdList, PixelShader, Params.SourceResource->TextureRHI, Params.SourcePosition, Size);
-
-		InRHICmdList.SetViewport((float)Params.DestPosition.X, (float)Params.DestPosition.Y, 0.0f, (float)(Params.DestPosition.X + Size.X), (float)(Params.DestPosition.Y + Size.Y), 1.0f);
-		InRHICmdList.DrawIndexedPrimitive(GTwoTrianglesIndexBuffer.IndexBufferRHI, 0, 0, 4, 0, 2, 1);
-
-		InRHICmdList.EndRenderPass();
-
-		InRHICmdList.Transition(FRHITransitionInfo(Params.SourceResource->TextureRHI, ERHIAccess::SRVGraphics, Params.SourceAccess));
-		InRHICmdList.Transition(FRHITransitionInfo(Params.DestResource->TextureRHI, ERHIAccess::RTV, Params.DestAccess));
-	}
-
 private:
 	FLandscapeLayersCopyTextureParams Params;
 };
@@ -4870,8 +4802,6 @@ int32 ALandscape::PerformLayersHeightmapsGlobalMerge(const FUpdateLayersContentC
 	UTextureRenderTarget2D* LandscapeScratchRT2 = HeightmapRTList[(int32)EHeightmapRTType::HeightmapRT_Scratch2];
 	UTextureRenderTarget2D* LandscapeScratchRT3 = HeightmapRTList[(int32)EHeightmapRTType::HeightmapRT_Scratch3];
 
-	const bool bUsePSCopyWorkaround = CVarLandscapeUsePSCopyWorkaround.GetValueOnGameThread();
-
 	for (FLandscapeLayer& Layer : LandscapeLayers)
 	{
 		// Draw each Layer's heightmaps to a Combined RT Atlas in LandscapeScratchRT1
@@ -4905,7 +4835,6 @@ int32 ALandscape::PerformLayersHeightmapsGlobalMerge(const FUpdateLayersContentC
 				CopyTextureParams.CopySize = LayerHeightmap.SectionRect.Size();
 				// Copy from the heightmap's top-left corner to the composited texture's position :
 				CopyTextureParams.DestPosition = LayerHeightmap.SectionRect.Min;
-				CopyTextureParams.bUsePS = bUsePSCopyWorkaround;
 			}
 			ExecuteCopyLayersTexture(MoveTemp(DeferredCopyTextures));
 		}
