New LOD Group to fix geometry seams between neighboring Landscapes, caused by mismatched LOD levels.

Landscapes blend LODs across components by referencing a spatial map of the LODs that is uploaded to the vertex shader.  The spatial map however only contains data for a single Landscape.  Here we introduce the concept of LODGroups to the landscapes, and build a combined spatial map for all landscapes within an LOD Group, by placing all of those landscapes into a single LandscapeRenderer.  This lets adjacent landscapes within an LOD Group match their LOD levels on the boundary between them.

diff --git a/Engine/Source/Runtime/Landscape/Classes/Landscape.h b/Engine/Source/Runtime/Landscape/Classes/Landscape.h
--- a/Engine/Source/Runtime/Landscape/Classes/Landscape.h
+++ b/Engine/Source/Runtime/Landscape/Classes/Landscape.h
@@ -260,6 +260,9 @@ public:
 	LANDSCAPE_API bool IsUpToDate() const;
 	LANDSCAPE_API void TickLayers(float DeltaTime);
 
+	LANDSCAPE_API void SetLODGroupKey(uint32 InLODGroupKey);
+	LANDSCAPE_API uint32 GetLODGroupKey();
+
 	UFUNCTION(BlueprintCallable, Category = "Landscape|Runtime")
 	LANDSCAPE_API void RenderHeightmap(const FTransform& InWorldTransform, const FBox2D& InExtents, UTextureRenderTarget2D* OutRenderTarget);
 
@@ -364,7 +367,7 @@ private:
 	void UpdateLayersContent(bool bInWaitForStreaming = false, bool bInSkipMonitorLandscapeEdModeChanges = false, bool bIntermediateRender = false, bool bFlushRender = false);
 	void MonitorShaderCompilation();
 	void MonitorLandscapeEdModeChanges();
-	
+
 	int32 RegenerateLayersHeightmaps(const FUpdateLayersContentContext& InUpdateLayersContentContext);
 	int32 PerformLayersHeightmapsLocalMerge(const FUpdateLayersContentContext& InUpdateLayersContentContext, const FEditLayersHeightmapMergeParams& InMergeParams);
 	int32 PerformLayersHeightmapsGlobalMerge(const FUpdateLayersContentContext& InUpdateLayersContentContext, const FEditLayersHeightmapMergeParams& InMergeParams);
@@ -440,6 +443,9 @@ private:
 	static bool IsMaterialResourceCompiled(FMaterialResource* InMaterialResource, bool bInWaitForCompilation);
 #endif // WITH_EDITOR
 
+private:
+	void MarkAllLandscapeRenderStateDirty();
+
 public:
 
 #if WITH_EDITORONLY_DATA
diff --git a/Engine/Source/Runtime/Landscape/Classes/LandscapeMeshProxyComponent.h b/Engine/Source/Runtime/Landscape/Classes/LandscapeMeshProxyComponent.h
--- a/Engine/Source/Runtime/Landscape/Classes/LandscapeMeshProxyComponent.h
+++ b/Engine/Source/Runtime/Landscape/Classes/LandscapeMeshProxyComponent.h
@@ -33,6 +33,9 @@ private:
 	UPROPERTY()
 	int8 ProxyLOD;
 
+	UPROPERTY()
+	uint32 LODGroupKey;
+
 public:
 	LANDSCAPE_API void InitializeForLandscape(ALandscapeProxy* Landscape, int8 InProxyLOD);
 	
diff --git a/Engine/Source/Runtime/Landscape/Classes/LandscapeProxy.h b/Engine/Source/Runtime/Landscape/Classes/LandscapeProxy.h
--- a/Engine/Source/Runtime/Landscape/Classes/LandscapeProxy.h
+++ b/Engine/Source/Runtime/Landscape/Classes/LandscapeProxy.h
@@ -463,6 +463,9 @@ public:
 	UPROPERTY(EditAnywhere, Category = "LOD Distribution", meta = (DisplayName = "LOD 0 Screen Size", ClampMin = "0.1", ClampMax = "10.0", UIMin = "1.0", UIMax = "10.0"))
 	float LOD0ScreenSize;
 
+	UPROPERTY(EditAnywhere, Category = LOD, AdvancedDisplay, meta = (ToolTip = "Specifies the LOD Group (Zero is No Group). All landscapes in the same group calculate their LOD together, allowing matching border LODs to fix geometry seams."))
+	uint32 LODGroupKey;
+
 	/** The distribution setting used to change the LOD 0 generation, 1.75 is the normal distribution, numbers influence directly the LOD0 proportion on screen. */
 	UPROPERTY(EditAnywhere, Category = "LOD Distribution", meta = (DisplayName = "LOD 0", ClampMin = "1.0", ClampMax = "10.0", UIMin = "1.0", UIMax = "10.0"))
 	float LOD0DistributionSetting;
diff --git a/Engine/Source/Runtime/Landscape/Private/Landscape.cpp b/Engine/Source/Runtime/Landscape/Private/Landscape.cpp
--- a/Engine/Source/Runtime/Landscape/Private/Landscape.cpp
+++ b/Engine/Source/Runtime/Landscape/Private/Landscape.cpp
@@ -1487,6 +1487,29 @@ void ALandscapeStreamingProxy::SetLandscapeActor(ALandscape* InLandscape)
 	LandscapeActorRef = InLandscape;
 }
 
+void ALandscape::SetLODGroupKey(uint32 InLODGroupKey)
+{
+	LODGroupKey = InLODGroupKey;
+	MarkAllLandscapeRenderStateDirty();
+}
+
+uint32 ALandscape::GetLODGroupKey()
+{
+	return LODGroupKey;
+}
+
+void ALandscape::MarkAllLandscapeRenderStateDirty()
+{
+	if (ULandscapeInfo* Info = GetLandscapeInfo())
+	{
+		Info->ForEachLandscapeProxy([](ALandscapeProxy* Proxy)
+			{
+				Proxy->MarkComponentsRenderStateDirty();
+				return true;
+			});
+	}
+}
+
 ULandscapeInfo* ALandscapeProxy::CreateLandscapeInfo(bool bMapCheck, bool bUpdateAllAddCollisions)
 {
 	ULandscapeInfo* LandscapeInfo = ULandscapeInfo::FindOrCreate(GetWorld(), LandscapeGuid);
@@ -3712,6 +3735,7 @@ void ALandscapeProxy::GetSharedProperties(ALandscapeProxy* Landscape)
 		LODDistributionSetting = Landscape->LODDistributionSetting;
 		LOD0DistributionSetting = Landscape->LOD0DistributionSetting;
 		LOD0ScreenSize = Landscape->LOD0ScreenSize;
+		LODGroupKey = Landscape->LODGroupKey;
 		NegativeZBoundsExtension = Landscape->NegativeZBoundsExtension;
 		PositiveZBoundsExtension = Landscape->PositiveZBoundsExtension;
 		CollisionMipLevel = Landscape->CollisionMipLevel;
@@ -3871,6 +3895,18 @@ void ALandscapeProxy::FixupSharedData(ALandscape* Landscape, bool bMapCheck)
 		bUpdated = true;
 	}
 
+	if (LODGroupKey != Landscape->LODGroupKey)
+	{
+		LODGroupKey = Landscape->LODGroupKey;
+		bUpdated = true;
+	}
+
+	if (bUseCompressedHeightmapStorage)
+	{
+		bUseCompressedHeightmapStorage = Landscape->bUseCompressedHeightmapStorage;
+		bUpdated = true;
+	}
+
 	if (TargetDisplayOrder != Landscape->TargetDisplayOrder)
 	{
 		TargetDisplayOrder = Landscape->TargetDisplayOrder;
@@ -4606,6 +4642,7 @@ void ULandscapeInfo::RegisterActor(ALandscapeProxy* Proxy, bool bMapCheck, bool
 #if WITH_EDITOR
 					StreamingProxy->FixupSharedData(Landscape, bMapCheck);
 #endif // WITH_EDITOR
+					StreamingProxy->LODGroupKey = Landscape->LODGroupKey;
 				}
 			}
 		}
@@ -4652,6 +4689,7 @@ void ULandscapeInfo::RegisterActor(ALandscapeProxy* Proxy, bool bMapCheck, bool
 		if (LandscapeActor.IsValid())	// don't overwrite the proxy's landscape actor if we don't have one registered yet
 		{
 			StreamingProxy->SetLandscapeActor(LandscapeActor.Get());
+			StreamingProxy->LODGroupKey = LandscapeActor.Get()->LODGroupKey;
 		}
 #if WITH_EDITOR
 		StreamingProxy->FixupSharedData(LandscapeActor.Get(), bMapCheck);
@@ -5206,6 +5244,7 @@ ULandscapeMeshProxyComponent::ULandscapeMeshProxyComponent(const FObjectInitiali
 void ULandscapeMeshProxyComponent::InitializeForLandscape(ALandscapeProxy* Landscape, int8 InProxyLOD)
 {
 	LandscapeGuid = Landscape->GetLandscapeGuid();
+	LODGroupKey = Landscape->GetLandscapeActor()->GetLODGroupKey();
 
 	for (ULandscapeComponent* Component : Landscape->LandscapeComponents)
 	{
diff --git a/Engine/Source/Runtime/Landscape/Private/LandscapeCulling.cpp b/Engine/Source/Runtime/Landscape/Private/LandscapeCulling.cpp
--- a/Engine/Source/Runtime/Landscape/Private/LandscapeCulling.cpp
+++ b/Engine/Source/Runtime/Landscape/Private/LandscapeCulling.cpp
@@ -406,7 +406,7 @@ IMPLEMENT_GLOBAL_SHADER(FBuildLandscapeTileDataCS, "/Engine/Private/Landscape/La
 struct FViewStateIntermediateData
 {
 	FRDGBufferRef SectionsBufferRDG;
-	TArray<FIntPoint, TInlineAllocator<8>> SectionBases;
+	TArray<FIntPoint, TInlineAllocator<8>> SectionRenderCoords;
 };
 
 struct FDispatchIntermediates
@@ -418,19 +418,19 @@ struct FDispatchIntermediates
 struct FArgumentsKey
 {
 	const void* ViewPtr;
-	FIntPoint SectionBase;
+	FIntPoint RenderCoord;
 	uint32 ViewStateKey;
 
 	bool operator==(const FArgumentsKey& Other) const
 	{
-		return ViewPtr == Other.ViewPtr && SectionBase == Other.SectionBase && ViewStateKey == Other.ViewStateKey;
+		return ViewPtr == Other.ViewPtr && RenderCoord == Other.RenderCoord && ViewStateKey == Other.ViewStateKey;
 	}
 };
 
 uint32 GetTypeHash(const FArgumentsKey& Key)
 {
 	uint32 Hash = PointerHash(Key.ViewPtr);
-	Hash = HashCombine(Hash, GetTypeHash(Key.SectionBase));
+	Hash = HashCombine(Hash, GetTypeHash(Key.RenderCoord));
 	Hash = HashCombine(Hash, ::GetTypeHash(Key.ViewStateKey));
 	return Hash;
 }
@@ -573,7 +573,7 @@ void UnregisterLandscape(uint32 LandscapeKey)
 	}
 }
 
-static float ComputeNeighborsMaxLOD(const FLandscapeRenderSystem& RenderSystem, const TResourceArray<float>& SectionLODValues, FIntPoint SectionBase)
+static float ComputeNeighborsMaxLOD(const FLandscapeRenderSystem& RenderSystem, const TResourceArray<float>& SectionLODValues, FIntPoint RenderCoord)
 {
 	static const FIntPoint NeightborBaseOffsets[4]
 	{
@@ -583,7 +583,7 @@ static float ComputeNeighborsMaxLOD(const FLandscapeRenderSystem& RenderSystem,
 	float NeighborsMaxLOD = 0.f;
 	for (int32 i = 0; i < UE_ARRAY_COUNT(NeightborBaseOffsets); ++i)
 	{
-		int32 SectionIdx = RenderSystem.GetSectionLinearIndex(SectionBase + NeightborBaseOffsets[i]);
+		int32 SectionIdx = RenderSystem.GetSectionLinearIndex(RenderCoord + NeightborBaseOffsets[i]);
 		if (SectionLODValues.IsValidIndex(SectionIdx))
 		{
 			NeighborsMaxLOD = FMath::Max(SectionLODValues[SectionIdx], NeighborsMaxLOD);
@@ -628,7 +628,7 @@ static void ComputeSectionIntermediateData(FRDGBuilder& GraphBuilder, TArrayView
 					FMatrix SectionLocalToWorld;
 					SectionInfo->GetSectionBoundsAndLocalToWorld(SectionLocalBounds, SectionLocalToWorld);
 					const FLargeWorldRenderPosition SectionAbsoluteOrigin(SectionLocalToWorld.GetOrigin());
-					const int32 NeighborsMaxLOD = static_cast<int32>(FMath::RoundFromZero(ComputeNeighborsMaxLOD(RenderSystem, SectionLODValues, SectionInfo->ComponentBase)));
+					const int32 NeighborsMaxLOD = static_cast<int32>(FMath::RoundFromZero(ComputeNeighborsMaxLOD(RenderSystem, SectionLODValues, SectionInfo->RenderCoord)));
 
 					Section.LocalToRelativeWorld = FLargeWorldRenderScalar::MakeToRelativeWorldMatrix(SectionAbsoluteOrigin.GetTileOffset(), SectionLocalToWorld);
 					Section.TilePosition = SectionAbsoluteOrigin.GetTile();
@@ -637,7 +637,7 @@ static void ComputeSectionIntermediateData(FRDGBuilder& GraphBuilder, TArrayView
 					// How many quads to add to each tile extent to compensate for a neighbors LOD
 					Section.NeighborLODExtent = static_cast<float>((1 << NeighborsMaxLOD) - 1); 
 
-					ViewStateIntermediates.SectionBases.Add(SectionInfo->ComponentBase);
+					ViewStateIntermediates.SectionRenderCoords.Add(SectionInfo->RenderCoord);
 				}
 			}
 
@@ -694,7 +694,7 @@ static void DispatchCulling(FRDGBuilder& GraphBuilder, TArrayView<const FSceneVi
 			FViewStateIntermediateData& ViewStateIntermediates = CullingEntry.IntermediateData[ViewStateIdx];
 
 			uint32 NumCullingViews = CullingViewsData.Num();
-			uint32 NumCullingSections = ViewStateIntermediates.SectionBases.Num();
+			uint32 NumCullingSections = ViewStateIntermediates.SectionRenderCoords.Num();
 			uint32 NumCullingItems = NumCullingSections * NumCullingViews;
 			if (NumCullingItems == 0)
 			{
@@ -761,10 +761,11 @@ static void DispatchCulling(FRDGBuilder& GraphBuilder, TArrayView<const FSceneVi
 			for (int32 ViewIdx = 0; ViewIdx < FilteredViews.Num(); ++ViewIdx)
 			{
 				const FSceneView* ViewPtr = FilteredViews[ViewIdx];
-				for (int32 SectionIdx = 0; SectionIdx < ViewStateIntermediates.SectionBases.Num(); ++SectionIdx)
+				for (int32 SectionIdx = 0; SectionIdx < ViewStateIntermediates.SectionRenderCoords.Num(); ++SectionIdx)
 				{
-					FIntPoint SectionBase = ViewStateIntermediates.SectionBases[SectionIdx];
-					FArguments& Args = CullingEntry.CullingArguments.Add(FArgumentsKey{ ViewPtr, SectionBase, ViewStateKey });
+					FIntPoint RenderCoord = ViewStateIntermediates.SectionRenderCoords[SectionIdx];
+					check(RenderCoord.X > INT32_MIN);
+					FArguments& Args = CullingEntry.CullingArguments.Add(FArgumentsKey{ ViewPtr, RenderCoord, ViewStateKey });
 					Args.IndirectArgsBuffer = IndirectArgsBufferRHI;
 					Args.TileDataVertexBuffer = TileDataBufferRHI;
 					// offsets are in bytes
@@ -812,7 +813,7 @@ void InitShadowViews(FRDGBuilder& GraphBuilder, TArrayView<const FSceneView*> Sh
 	DispatchCulling(GraphBuilder, ShadowDepthViews, ShadowViewMatrices, false);
 }
 
-bool GetViewArguments(const FSceneView& View, uint32 LandscapeKey, FIntPoint SectionBase, int32 LODIndex, FArguments& Args)
+bool GetViewArguments(const FSceneView& View, uint32 LandscapeKey, FIntPoint RenderCoord, int32 LODIndex, FArguments& Args)
 {
 	if (LODIndex != 0 || 
 		GCullingSystem.Landscapes.Num() == 0 ||
@@ -829,7 +830,9 @@ bool GetViewArguments(const FSceneView& View, uint32 LandscapeKey, FIntPoint Sec
 	{
 		return false;
 	}
-	FArgumentsKey ArgumentsKey{ &View, SectionBase, View.GetViewKey() };
+
+	check(RenderCoord.X > INT32_MIN);
+	FArgumentsKey ArgumentsKey{ &View, RenderCoord, View.GetViewKey() };
 	const FArguments* ArgsPtr = CullingEntryPtr->CullingArguments.Find(ArgumentsKey);
 	if (ArgsPtr)
 	{
diff --git a/Engine/Source/Runtime/Landscape/Private/LandscapeEdit.cpp b/Engine/Source/Runtime/Landscape/Private/LandscapeEdit.cpp
--- a/Engine/Source/Runtime/Landscape/Private/LandscapeEdit.cpp
+++ b/Engine/Source/Runtime/Landscape/Private/LandscapeEdit.cpp
@@ -93,12 +93,6 @@ DEFINE_LOG_CATEGORY(LogLandscapeBP);
 
 #define LOCTEXT_NAMESPACE "Landscape"
 
-static TAutoConsoleVariable<int32> CVarMobileCompressLanscapeWeightMaps(
-    TEXT("r.Mobile.CompressLandscapeWeightMaps"),
-	0,
-    TEXT("Whether to compress the terrain weight maps for mobile."),
-	ECVF_ReadOnly);
-
 static TAutoConsoleVariable<int32> CVarLandscapeApplyPhysicalMaterialChangesImmediately(
     TEXT("landscape.ApplyPhysicalMaterialChangesImmediately"),
 	1,
@@ -5401,11 +5395,13 @@ bool ALandscapeProxy::CanEditChange(const FProperty* InProperty) const
 	{
 		FName PropertyName = InProperty ? InProperty->GetFName() : NAME_None;
 
-		if ((PropertyName == GET_MEMBER_NAME_CHECKED(ALandscapeProxy, MaxLODLevel))
+		if (   (PropertyName == GET_MEMBER_NAME_CHECKED(ALandscapeProxy, MaxLODLevel))
 			|| (PropertyName == GET_MEMBER_NAME_CHECKED(ALandscapeProxy, ComponentScreenSizeToUseSubSections))
 			|| (PropertyName == GET_MEMBER_NAME_CHECKED(ALandscapeProxy, LODDistributionSetting))
 			|| (PropertyName == GET_MEMBER_NAME_CHECKED(ALandscapeProxy, LOD0DistributionSetting))
 			|| (PropertyName == GET_MEMBER_NAME_CHECKED(ALandscapeProxy, LOD0ScreenSize))
+			|| (PropertyName == GET_MEMBER_NAME_CHECKED(ALandscapeProxy, LODGroupKey))
+			|| (PropertyName == GET_MEMBER_NAME_CHECKED(ALandscapeProxy, bUseCompressedHeightmapStorage))
 			|| (PropertyName == GET_MEMBER_NAME_CHECKED(ALandscapeProxy, TargetDisplayOrder))
 			|| (PropertyName == GET_MEMBER_NAME_CHECKED(ALandscapeProxy, TargetDisplayOrderList))
 			|| (PropertyName == GET_MEMBER_NAME_CHECKED(ALandscapeProxy, bEnableNanite))
@@ -5534,6 +5530,7 @@ void ALandscapeProxy::PostEditChangeProperty(FPropertyChangedEvent& PropertyChan
 		|| (PropertyName == GET_MEMBER_NAME_CHECKED(ALandscapeProxy, bCastDynamicShadow))
 		|| (PropertyName == GET_MEMBER_NAME_CHECKED(ALandscapeProxy, ShadowCacheInvalidationBehavior))
 		|| (PropertyName == GET_MEMBER_NAME_CHECKED(ALandscapeProxy, bCastStaticShadow))
+		|| (PropertyName == GET_MEMBER_NAME_CHECKED(ALandscapeProxy, bCastContactShadow))
 		|| (PropertyName == GET_MEMBER_NAME_CHECKED(ALandscapeProxy, bCastFarShadow))
 		|| (PropertyName == GET_MEMBER_NAME_CHECKED(ALandscapeProxy, bCastHiddenShadow))
 		|| (PropertyName == GET_MEMBER_NAME_CHECKED(ALandscapeProxy, bCastShadowAsTwoSided))
@@ -5963,6 +5960,7 @@ void ALandscape::PostEditChangeProperty(FPropertyChangedEvent& PropertyChangedEv
 	bool bNeedsRecalcBoundingBox = false;
 	bool bChangedLighting = false;
 	bool bPropagateToProxies = false;
+	bool bMarkAllLandscapeRenderStateDirty = false;
 	bool bNaniteToggled = false;
 
 	ULandscapeInfo* Info = GetLandscapeInfo();
@@ -6117,6 +6115,12 @@ void ALandscape::PostEditChangeProperty(FPropertyChangedEvent& PropertyChangedEv
 		UpdateRenderingMethod();
 		MarkComponentsRenderStateDirty();
 	}
+	else if (PropertyName == GET_MEMBER_NAME_CHECKED(ALandscape, LODGroupKey))
+	{
+		// need to invoke GetSharedProperties to copy this change down to the proxies
+		bPropagateToProxies = true;
+		bMarkAllLandscapeRenderStateDirty = true;
+	}
 
 	// Must do this *after* clamping values
 	Super::PostEditChangeProperty(PropertyChangedEvent);
@@ -6139,6 +6143,11 @@ void ALandscape::PostEditChangeProperty(FPropertyChangedEvent& PropertyChangedEv
 			}
 		}
 
+		if (bMarkAllLandscapeRenderStateDirty)
+		{
+			MarkAllLandscapeRenderStateDirty();
+		}
+
 		// Update normals if DrawScale3D is changed
 		if (MemberPropertyName == FName(TEXT("RelativeScale3D")))
 		{
diff --git a/Engine/Source/Runtime/Landscape/Private/LandscapeRender.cpp b/Engine/Source/Runtime/Landscape/Private/LandscapeRender.cpp
--- a/Engine/Source/Runtime/Landscape/Private/LandscapeRender.cpp
+++ b/Engine/Source/Runtime/Landscape/Private/LandscapeRender.cpp
@@ -454,6 +454,7 @@ FLandscapeRenderSystem::FLandscapeRenderSystem()
 	, Size(EForceInit::ForceInitToZero)
 	, ReferenceCount(0)
 	, ForcedLODOverride(-1)
+	, SectionsRemovedSinceLastCompact(0)
 {
 	SectionLODBiases.SetAllowCPUAccess(true);
 
@@ -512,26 +513,61 @@ void FLandscapeRenderSystem::CreateResources_Internal(FLandscapeSectionInfo* Sec
 	check(IsInRenderingThread());
 	check(SectionInfo != nullptr);
 	check(!SectionInfo->bRegistered);
+	check(!SectionInfo->bResourcesCreated);
 
-	if (!SectionInfos.IsEmpty())
+	SectionInfo->RenderCoord = SectionInfo->ComponentBase;
+	SectionInfo->bResourcesCreated = true;
+
+	if (SectionInfo->LODGroupKey != 0)
 	{
-		// Calculate new bounding rect of landscape components
-		FIntPoint OriginalMin = Min;
-		FIntPoint OriginalMax = Min + Size - FIntPoint(1, 1);
-		FIntPoint NewMin(FMath::Min(Min.X, SectionInfo->ComponentBase.X), FMath::Min(Min.Y, SectionInfo->ComponentBase.Y));
-		FIntPoint NewMax(FMath::Max(OriginalMax.X, SectionInfo->ComponentBase.X), FMath::Max(OriginalMax.Y, SectionInfo->ComponentBase.Y));
+		// Record and check settings (resolution, scale and orientation) that need to match across landscapes in an LOD Group
+		int32 SectionComponentResolution = SectionInfo->GetComponentResolution();
+		if (SectionComponentResolution > 0)
+		{
+			FBoxSphereBounds LocalBounds;
+			FMatrix LocalToWorld;
+			SectionInfo->GetSectionBoundsAndLocalToWorld(LocalBounds, LocalToWorld);
 
-		FIntPoint SizeRequired = (NewMax - NewMin) + FIntPoint(1, 1);
+			FVector SectionCenterWorldSpace = LocalToWorld.TransformPosition(LocalBounds.Origin);
+			FVector SectionXVector = LocalToWorld.TransformVector(FVector::XAxisVector) * SectionComponentResolution;
+			FVector SectionYVector = LocalToWorld.TransformVector(FVector::YAxisVector) * SectionComponentResolution;
 
-		if (NewMin != Min || Size != SizeRequired)
-		{
-			ResizeAndMoveTo(NewMin, SizeRequired);
+			if (ComponentResolution < 0)
+			{
+				// the first component with a resolution -- record its information
+				ComponentResolution = SectionComponentResolution;
+				ComponentOrigin = SectionCenterWorldSpace;
+				ComponentXVector = SectionXVector;
+				ComponentYVector = SectionYVector;
+			}
+			else
+			{
+				// validate matching resolution
+				bool bResolutionMatches = (ComponentResolution == SectionComponentResolution);
+				bool bXVectorMatches = (SectionXVector - ComponentXVector).IsNearlyZero();
+				bool bYVectorMatches = (SectionYVector - ComponentYVector).IsNearlyZero();
+				if (!(bResolutionMatches && bXVectorMatches && bYVectorMatches))
+				{
+					UE_LOG(LogLandscape, Warning, TEXT("Landscapes in LOD Group %d do not have matching resolution (%d == %d), scale (%f == %f, %f == %f) and/or rotation; geometry seam artifacts may appear."),
+						SectionInfo->LODGroupKey,
+						ComponentResolution, SectionComponentResolution,
+						ComponentXVector.Length(), SectionXVector.Length(),
+						ComponentYVector.Length(), SectionYVector.Length());
+				}
+			}
+			
+			// project onto the Component X/Y plane to calculate the render coordinates
+			FVector Delta = SectionCenterWorldSpace - ComponentOrigin;
+			SectionInfo->RenderCoord.X = FMath::RoundToInt32(Delta.Dot(ComponentXVector) / ComponentXVector.SquaredLength());
+			SectionInfo->RenderCoord.Y = FMath::RoundToInt32(Delta.Dot(ComponentYVector) / ComponentXVector.SquaredLength());
 		}
 	}
-	else
-	{
-		ResizeAndMoveTo(SectionInfo->ComponentBase, FIntPoint(1, 1));
-	}
+
+	// we changed the RenderCoord, need to update the uniform buffer
+	SectionInfo->OnRenderCoordsChanged();
+
+	check(SectionInfo->RenderCoord.X > INT32_MIN);
+	ResizeToInclude(SectionInfo->RenderCoord);
 
 	ReferenceCount++;
 }
@@ -541,8 +577,23 @@ void FLandscapeRenderSystem::DestroyResources_Internal(FLandscapeSectionInfo* Se
 	check(IsInRenderingThread());
 	check(SectionInfo != nullptr);
 	check(!SectionInfo->bRegistered);
+	check(SectionInfo->bResourcesCreated);
 
 	ReferenceCount--;
+
+	SectionInfo->bResourcesCreated = false;
+
+	// try to compact the map every once in a while
+	SectionsRemovedSinceLastCompact++;
+	if (SectionsRemovedSinceLastCompact >= 128)
+	{
+		// if there are at least 128 free entries, run a compact step
+		if (Size.X * Size.Y - ReferenceCount >= 128)
+		{
+			CompactMap();
+		}
+		SectionsRemovedSinceLastCompact = 0;
+	}
 }
 
 void FLandscapeRenderSystem::RegisterSection(FLandscapeSectionInfo* SectionInfo)
@@ -558,10 +609,11 @@ void FLandscapeRenderSystem::RegisterSection(FLandscapeSectionInfo* SectionInfo)
 	// properly restore a previously registered section info.
 
 	FLandscapeRenderSystem*& LandscapeRenderSystem = LandscapeRenderSystems.FindChecked(SectionInfo->LandscapeKey);
-	FLandscapeSectionInfo* ExistingSection = LandscapeRenderSystem->GetSectionInfo(SectionInfo->ComponentBase);
+
+	FLandscapeSectionInfo* ExistingSection = LandscapeRenderSystem->GetSectionInfo(SectionInfo->RenderCoord);
 	if (ExistingSection == nullptr)
 	{
-		LandscapeRenderSystem->SetSectionInfo(SectionInfo->ComponentBase, SectionInfo);
+		LandscapeRenderSystem->SetSectionInfo(SectionInfo->RenderCoord, SectionInfo);
 	}
 	else
 	{
@@ -587,7 +639,7 @@ void FLandscapeRenderSystem::RegisterSection(FLandscapeSectionInfo* SectionInfo)
 		else if (CurrentSection == ExistingSection)
 		{
 			// Set as head
-			LandscapeRenderSystem->SetSectionInfo(SectionInfo->ComponentBase, SectionInfo);
+			LandscapeRenderSystem->SetSectionInfo(SectionInfo->RenderCoord, SectionInfo);
 		}
 	}
 
@@ -602,10 +654,11 @@ void FLandscapeRenderSystem::UnregisterSection(FLandscapeSectionInfo* SectionInf
 	if (SectionInfo->bRegistered)
 	{
 		FLandscapeRenderSystem* LandscapeRenderSystem = LandscapeRenderSystems.FindChecked(SectionInfo->LandscapeKey);
-		FLandscapeSectionInfo* ExistingSection = LandscapeRenderSystem->GetSectionInfo(SectionInfo->ComponentBase);
+
+		FLandscapeSectionInfo* ExistingSection = LandscapeRenderSystem->GetSectionInfo(SectionInfo->RenderCoord);
 		if (ExistingSection == SectionInfo)
 		{
-			LandscapeRenderSystem->SetSectionInfo(SectionInfo->ComponentBase, SectionInfo->GetNextLink());
+			LandscapeRenderSystem->SetSectionInfo(SectionInfo->RenderCoord, SectionInfo->GetNextLink());
 		}
 
 		SectionInfo->Unlink();
@@ -614,37 +667,137 @@ void FLandscapeRenderSystem::UnregisterSection(FLandscapeSectionInfo* SectionInf
 	}
 }
 
-void FLandscapeRenderSystem::ResizeAndMoveTo(FIntPoint NewMin, FIntPoint NewSize)
+void FLandscapeRenderSystem::ResizeAndMoveTo(FIntPoint NewMin, FIntPoint NewMax)
 {
-	SectionLODBiasBuffer.SafeRelease();
+	// chunk size reduces the number of times we need to reallocate the section map when resizing it. should be a power of two
+	static constexpr int32 SectionMapChunkSize = 4;
 
-	TResourceArray<float> NewSectionLODBiases;
-	TArray<FLandscapeSectionInfo*> NewSectionInfos;
+	// round down to multiple of SectionMapChunkSize
+	NewMin.X = (NewMin.X) & ~(SectionMapChunkSize - 1);
+	NewMin.Y = (NewMin.Y) & ~(SectionMapChunkSize - 1);
 
-	NewSectionLODBiases.AddZeroed(NewSize.X * NewSize.Y);
-	NewSectionInfos.AddZeroed(NewSize.X * NewSize.Y);
+	// round up to multiple of SectionMapChunkSize (note that before this operation NewMax is inclusive, but after it is exclusive)
+	NewMax.X = (NewMax.X + SectionMapChunkSize) & ~(SectionMapChunkSize - 1);
+	NewMax.Y = (NewMax.Y + SectionMapChunkSize) & ~(SectionMapChunkSize - 1);
 
-	for (int32 Y = 0; Y < Size.Y; Y++)
+	FIntPoint NewSize = (NewMax - NewMin);
+	if (NewMin != Min || Size != NewSize)
 	{
-		for (int32 X = 0; X < Size.X; X++)
-		{
-			int32 LinearIndex = Y * Size.X + X;
-			int32 NewLinearIndex = (Y + (Min.Y - NewMin.Y)) * NewSize.X + (X + (Min.X - NewMin.X));
+		check((Size.X & (SectionMapChunkSize-1)) == 0);
+		check((Size.Y & (SectionMapChunkSize-1)) == 0);
+
+		SectionLODBiasBuffer.SafeRelease();
+
+		TResourceArray<float> NewSectionLODBiases;
+		TArray<FLandscapeSectionInfo*> NewSectionInfos;
+
+		NewSectionLODBiases.AddZeroed(NewSize.X * NewSize.Y);
+		NewSectionInfos.AddZeroed(NewSize.X * NewSize.Y);
 
-			if (NewLinearIndex >= 0 && NewLinearIndex < NewSize.X * NewSize.Y)
+		for (int32 OldY = 0; OldY < Size.Y; OldY++)
+		{
+			int32 Y = OldY + Min.Y;
+			int32 NewY = Y - NewMin.Y;
+			if ((NewY >= 0) && (NewY < NewSize.Y))
 			{
-				NewSectionLODBiases[NewLinearIndex] = SectionLODBiases[LinearIndex];
-				NewSectionInfos[NewLinearIndex] = SectionInfos[LinearIndex];
+				int32 OldYBase = OldY * Size.X;
+				int32 NewYBase = NewY * NewSize.X;
+				for (int32 OldX = 0; OldX < Size.X; OldX++)
+				{
+					int32 X = OldX + Min.X;
+					int32 NewX = X - NewMin.X;
+					if ((NewX >= 0) && (NewX < NewSize.X))
+					{
+						int32 OldLinearIndex = OldYBase + OldX;
+						int32 NewLinearIndex = NewYBase + NewX;
+						NewSectionLODBiases[NewLinearIndex] = SectionLODBiases[OldLinearIndex];
+						NewSectionInfos[NewLinearIndex] = SectionInfos[OldLinearIndex];
+					}
+				}
 			}
 		}
+
+		Min = NewMin;
+		Size = NewSize;
+		SectionLODBiases = MoveTemp(NewSectionLODBiases);
+		SectionInfos = MoveTemp(NewSectionInfos);
+
+		SectionLODBiases.SetAllowCPUAccess(true);
 	}
+}
 
-	Min = NewMin;
-	Size = NewSize;
-	SectionLODBiases = MoveTemp(NewSectionLODBiases);
-	SectionInfos = MoveTemp(NewSectionInfos);
+void FLandscapeRenderSystem::ResizeToInclude(const FIntPoint& NewCoord)
+{
+	if (!SectionInfos.IsEmpty())
+	{
+		// Calculate new bounding rect of landscape components
+		FIntPoint OriginalMin = Min;
+		FIntPoint OriginalMax = Min + Size - FIntPoint(1, 1);
 
-	SectionLODBiases.SetAllowCPUAccess(true);
+		FIntPoint NewMin(FMath::Min(OriginalMin.X, NewCoord.X), FMath::Min(OriginalMin.Y, NewCoord.Y));
+		FIntPoint NewMax(FMath::Max(OriginalMax.X, NewCoord.X), FMath::Max(OriginalMax.Y, NewCoord.Y));
+
+		ResizeAndMoveTo(NewMin, NewMax);
+	}
+	else
+	{
+		ResizeAndMoveTo(NewCoord, NewCoord);
+	}
+}
+
+void FLandscapeRenderSystem::CompactMap()
+{
+	FIntPoint NewMin = Min;
+	FIntPoint NewMax = Min + Size - FIntPoint(1, 1);
+
+	// Shrink Min X 
+	while ((NewMin.X < NewMax.X) &&
+		!AnySectionsInRangeInclusive(FIntPoint(NewMin.X, NewMin.Y), FIntPoint(NewMin.X, NewMax.Y)))
+	{
+		NewMin.X++;
+	}
+
+	// Shrink Max X 
+	while ((NewMin.X < NewMax.X) &&
+		!AnySectionsInRangeInclusive(FIntPoint(NewMax.X, NewMin.Y), FIntPoint(NewMax.X, NewMax.Y)))
+	{
+		NewMax.X--;
+	}
+
+	// Shrink Min Y
+	while ((NewMin.Y < NewMax.Y) &&
+		!AnySectionsInRangeInclusive(FIntPoint(NewMin.X, NewMin.Y), FIntPoint(NewMax.X, NewMin.Y)))
+	{
+		NewMin.Y++;
+	}
+
+	// Shrink Max Y
+	while ((NewMin.Y < NewMax.Y) &&
+		!AnySectionsInRangeInclusive(FIntPoint(NewMin.X, NewMax.Y), FIntPoint(NewMax.X, NewMax.Y)))
+	{
+		NewMax.Y--;
+	}
+
+	ResizeAndMoveTo(NewMin, NewMax);
+}
+
+bool FLandscapeRenderSystem::AnySectionsInRangeInclusive(FIntPoint RangeMin, FIntPoint RangeMax)
+{
+	// convert to range relative to Min
+	RangeMin -= Min;
+	RangeMax -= Min;
+	for (int32 Y = RangeMin.Y; Y <= RangeMax.Y; Y++)
+	{
+		for (int32 X = RangeMin.X; X <= RangeMax.X; X++)
+		{
+			int32 LinearIndex = Y * Size.X + X;
+			if (SectionInfos[LinearIndex] != nullptr)
+			{
+				return true;
+			}
+		}
+	}
+	return false;
 }
 
 const TResourceArray<float>& FLandscapeRenderSystem::ComputeSectionsLODForView(const FSceneView& InView)
@@ -902,7 +1055,7 @@ bool FLandscapeVisibilityHelper::OnRemoveFromWorld()
 
 FLandscapeComponentSceneProxy::FLandscapeComponentSceneProxy(ULandscapeComponent* InComponent)
 	: FPrimitiveSceneProxy(InComponent, NAME_LandscapeResourceNameForDebugging)
-	, FLandscapeSectionInfo(InComponent->GetWorld(), InComponent->GetLandscapeProxy()->GetLandscapeGuid(), InComponent->GetSectionBase() / InComponent->ComponentSizeQuads)
+	, FLandscapeSectionInfo(InComponent->GetWorld(), InComponent->GetLandscapeProxy()->GetLandscapeGuid(), InComponent->GetSectionBase() / InComponent->ComponentSizeQuads, InComponent->GetLandscapeProxy()->LODGroupKey)
 	, MaxLOD(static_cast<int8>(FMath::CeilLogTwo(InComponent->SubsectionSizeQuads + 1) - 1))
 	, NumWeightmapLayerAllocations(static_cast<int8>(InComponent->GetWeightmapLayerAllocations().Num()))
 	, StaticLightingLOD(static_cast<uint8>(InComponent->GetLandscapeProxy()->StaticLightingLOD))
@@ -1207,7 +1360,7 @@ void FLandscapeComponentSceneProxy::CreateRenderThreadResources()
 
 	if (VisibilityHelper.ShouldBeVisible())
 	{
-		RegisterSection();
+		FLandscapeRenderSystem::RegisterSection(this);
 	}
 
 	auto FeatureLevel = GetScene().GetFeatureLevel();
@@ -1749,6 +1902,14 @@ namespace DebugColorMask
 
 void FLandscapeComponentSceneProxy::OnTransformChanged()
 {
+	// resource creation will call OnTransformChanged(), so don't bother updating everything here if resources haven't been created yet
+	if (!bResourcesCreated)
+	{
+		return;
+	}
+	
+	check(RenderCoord.X > INT32_MIN);
+
 	// Set Lightmap ScaleBias
 	int32 PatchExpandCountX = 0;
 	int32 PatchExpandCountY = 0;
@@ -1773,8 +1934,8 @@ void FLandscapeComponentSceneProxy::OnTransformChanged()
 
 	// Set FLandscapeUniformVSParameters for this subsection
 	FLandscapeUniformShaderParameters LandscapeParams;
-	LandscapeParams.ComponentBaseX = ComponentBase.X;
-	LandscapeParams.ComponentBaseY = ComponentBase.Y;
+	LandscapeParams.ComponentBaseX = RenderCoord.X;
+	LandscapeParams.ComponentBaseY = RenderCoord.Y;
 	LandscapeParams.SubsectionSizeVerts = SubsectionSizeVerts;
 	LandscapeParams.NumSubsections = NumSubsections;
 	LandscapeParams.LastLOD = LastLOD;
@@ -2073,7 +2234,7 @@ void FLandscapeComponentSceneProxy::GetDynamicMeshElements(const TArray<const FS
 
 			const FSceneView* View = Views[ViewIndex];
 
-			int32 LODToRender = static_cast<int32>(RenderSystem.GetSectionLODValue(*View, ComponentBase));
+			int32 LODToRender = static_cast<int32>(RenderSystem.GetSectionLODValue(*View, RenderCoord));
 
 			FMeshBatch& Mesh = Collector.AllocateMesh();
 			GetStaticMeshElement(LODToRender, false, Mesh, ParameterArray.ElementParams);
@@ -2459,7 +2620,8 @@ void FLandscapeComponentSceneProxy::GetDynamicMeshElements(const TArray<const FS
 void FLandscapeComponentSceneProxy::ApplyViewDependentMeshArguments(const FSceneView& View, FMeshBatch& ViewDependentMeshBatch) const
 {
 	Culling::FArguments CullingArgs{};
-	if (Culling::GetViewArguments(View, LandscapeKey, ComponentBase, ViewDependentMeshBatch.LODIndex, CullingArgs))
+	check(RenderCoord.X > INT32_MIN);
+	if (Culling::GetViewArguments(View, LandscapeKey, RenderCoord, ViewDependentMeshBatch.LODIndex, CullingArgs))
 	{
 		ViewDependentMeshBatch.Elements[0].NumPrimitives = 0;
 		ViewDependentMeshBatch.Elements[0].NumInstances = 0;
@@ -2485,7 +2647,7 @@ void FLandscapeComponentSceneProxy::GetDynamicRayTracingInstances(FRayTracingMat
 	}
 	FLandscapeRayTracingState* RayTracingState = RayTracingImpl.Get()->FindOrCreateRayTracingState(SceneView.State, NumSubsections, SubsectionSizeVerts);
 
-	int32 LODToRender = static_cast<int32>(RenderSystem.GetSectionLODValue(SceneView, ComponentBase));
+	int32 LODToRender = static_cast<int32>(RenderSystem.GetSectionLODValue(SceneView, RenderCoord));
 
 	FLandscapeElementParamArray& ParameterArray = Context.RayTracingMeshResourceCollector.AllocateOneFrameResource<FLandscapeElementParamArray>();
 	ParameterArray.ElementParams.AddDefaulted(NumSubsections * NumSubsections);
@@ -2570,13 +2732,13 @@ void FLandscapeComponentSceneProxy::GetDynamicRayTracingInstances(FRayTracingMat
 					RayTracingState->Sections[SubSectionIdx].CurrentLOD = CurrentLOD;
 					RayTracingState->Sections[SubSectionIdx].RayTracingDynamicVertexBuffer.Release();
 				}
-				if (RayTracingState->Sections[SubSectionIdx].HeightmapLODBias != RenderSystem.GetSectionLODBias(ComponentBase))
+				if (RayTracingState->Sections[SubSectionIdx].HeightmapLODBias != RenderSystem.GetSectionLODBias(RenderCoord))
 				{
 					bNeedsRayTracingGeometryUpdate = true;
-					RayTracingState->Sections[SubSectionIdx].HeightmapLODBias = RenderSystem.GetSectionLODBias(ComponentBase);
+					RayTracingState->Sections[SubSectionIdx].HeightmapLODBias = RenderSystem.GetSectionLODBias(RenderCoord);
 				}
 
-				const float PendingFractionalLOD = RenderSystem.GetSectionLODValue(SceneView, ComponentBase);
+				const float PendingFractionalLOD = RenderSystem.GetSectionLODValue(SceneView, RenderCoord);
 				const float FractionLODAbsoluteDifference = FMath::Abs(RayTracingState->Sections[SubSectionIdx].FractionalLOD - PendingFractionalLOD);
 				if (FractionLODAbsoluteDifference > GLandscapeRayTracingGeometryFractionalLODUpdateThreshold)
 				{
@@ -2976,7 +3138,8 @@ void FLandscapeVertexFactoryVertexShaderParameters::GetElementShaderBindings(
 		// TileVF is used only for LOD0
 		int32 LODIndex = 0; 
 		Culling::FArguments CullingArgs{};
-		if (Culling::GetViewArguments(*InView, SceneProxy->LandscapeKey, SceneProxy->ComponentBase, LODIndex, CullingArgs))
+
+		if (Culling::GetViewArguments(*InView, SceneProxy->LandscapeKey, SceneProxy->RenderCoord, LODIndex, CullingArgs))
 		{
 			check(VertexStreams.IsValidIndex(1));
 			VertexStreams[1].Offset = CullingArgs.TileDataOffset;
@@ -3965,6 +4128,12 @@ float FLandscapeComponentSceneProxy::ComputeLODBias() const
 	return ComputedLODBias;
 }
 
+void FLandscapeComponentSceneProxy::OnRenderCoordsChanged()
+{
+	// need to rebuild the uniforms that contain render coords
+	OnTransformChanged();
+}
+
 double FLandscapeComponentSceneProxy::ComputeSectionResolution() const
 {
 	// ComponentMaxExtend is the max(length,width) of the component, in world units
@@ -3973,6 +4142,11 @@ double FLandscapeComponentSceneProxy::ComputeSectionResolution() const
 	return ComponentFullExtent / ComponentQuads;
 }
 
+int32 FLandscapeComponentSceneProxy::GetComponentResolution() const
+{
+	return ComponentSizeVerts;
+}
+
 void FLandscapeComponentSceneProxy::GetSectionBoundsAndLocalToWorld(FBoxSphereBounds& OutLocalBounds, FMatrix& OutLocalToWorld) const
 {
 	OutLocalBounds = GetLocalBounds();
@@ -3982,31 +4156,24 @@ void FLandscapeComponentSceneProxy::GetSectionBoundsAndLocalToWorld(FBoxSphereBo
 //
 // FLandscapeSectionInfo
 //
-FLandscapeSectionInfo::FLandscapeSectionInfo(const UWorld* InWorld, const FGuid& InLandscapeGuid, const FIntPoint& InSectionBase)
-	: LandscapeKey(HashCombine(GetTypeHash(InWorld), GetTypeHash(InLandscapeGuid)))
-	, ComponentBase(InSectionBase)
+FLandscapeSectionInfo::FLandscapeSectionInfo(const UWorld* InWorld, const FGuid& InLandscapeGuid, const FIntPoint& InComponentBase, uint32 LODGroupKey)
+	: LandscapeKey(HashCombine(GetTypeHash(InWorld), (LODGroupKey != 0) ? LODGroupKey : GetTypeHash(InLandscapeGuid)))    // use LODGroupKey instead of LandscapeGUID when LODGroupKey is non-zero
+	, LODGroupKey(LODGroupKey)
+	, RenderCoord(INT32_MIN, INT32_MIN)
+	, ComponentBase(InComponentBase)
+	, bResourcesCreated(false)
 	, bRegistered(false)
 {
 }
 
-void FLandscapeSectionInfo::RegisterSection()
-{
-	FLandscapeRenderSystem::RegisterSection(this);
-}
-
-void FLandscapeSectionInfo::UnregisterSection()
-{
-	FLandscapeRenderSystem::UnregisterSection(this);
-}
-
 //
 // FLandscapeProxySectionInfo
 //
 class FLandscapeProxySectionInfo : public FLandscapeSectionInfo
 {
 public:
-	FLandscapeProxySectionInfo(const UWorld* InWorld, const FGuid& InLandscapeGuid, const FIntPoint& InSectionBase, int8 InProxyLOD)
-		: FLandscapeSectionInfo(InWorld, InLandscapeGuid, InSectionBase)
+	FLandscapeProxySectionInfo(const UWorld* InWorld, const FGuid& InLandscapeGuid, const FIntPoint& InComponentBase, int8 InProxyLOD, uint32 LODGroupKey)
+		: FLandscapeSectionInfo(InWorld, InLandscapeGuid, InComponentBase, LODGroupKey)
 		, ProxyLOD(InProxyLOD)
 	{
 	}
@@ -4032,6 +4199,10 @@ public:
 		LocalToWorld = FMatrix::Identity;
 	}
 
+	virtual void OnRenderCoordsChanged()
+	{
+	}
+
 private:
 	int8 ProxyLOD;
 };
@@ -4039,7 +4210,7 @@ private:
 //
 // FLandscapeMeshProxySceneProxy
 //
-FLandscapeMeshProxySceneProxy::FLandscapeMeshProxySceneProxy(UStaticMeshComponent* InComponent, const FGuid& InLandscapeGuid, const TArray<FIntPoint>& InProxySectionsBases, int8 InProxyLOD)
+FLandscapeMeshProxySceneProxy::FLandscapeMeshProxySceneProxy(UStaticMeshComponent* InComponent, const FGuid& InLandscapeGuid, const TArray<FIntPoint>& InProxySectionsBases, int8 InProxyLOD, uint32 InLODGroupKey)
 	: FStaticMeshSceneProxy(InComponent, false)
 {
 	VisibilityHelper.Init(InComponent, this);
@@ -4052,7 +4223,7 @@ FLandscapeMeshProxySceneProxy::FLandscapeMeshProxySceneProxy(UStaticMeshComponen
 	ProxySectionsInfos.Empty(InProxySectionsBases.Num());
 	for (FIntPoint SectionBase : InProxySectionsBases)
 	{
-		ProxySectionsInfos.Emplace(MakeUnique<FLandscapeProxySectionInfo>(InComponent->GetWorld(), InLandscapeGuid, SectionBase, InProxyLOD));
+		ProxySectionsInfos.Emplace(MakeUnique<FLandscapeProxySectionInfo>(InComponent->GetWorld(), InLandscapeGuid, SectionBase, InProxyLOD, InLODGroupKey));
 	}
 }
 
@@ -4060,7 +4231,7 @@ void FLandscapeMeshProxySceneProxy::RegisterSections()
 {
 	for (auto& Info : ProxySectionsInfos)
 	{
-		Info->RegisterSection();
+		FLandscapeRenderSystem::RegisterSection(Info.Get());
 	}
 }
 
@@ -4068,7 +4239,7 @@ void FLandscapeMeshProxySceneProxy::UnregisterSections()
 {
 	for (auto& Info : ProxySectionsInfos)
 	{
-		Info->UnregisterSection();
+		FLandscapeRenderSystem::UnregisterSection(Info.Get());
 	}
 }
 
@@ -4137,7 +4308,7 @@ FPrimitiveSceneProxy* ULandscapeMeshProxyComponent::CreateSceneProxy()
 		return nullptr;
 	}
 
-	return new FLandscapeMeshProxySceneProxy(this, LandscapeGuid, ProxyComponentBases, ProxyLOD);
+	return new FLandscapeMeshProxySceneProxy(this, LandscapeGuid, ProxyComponentBases, ProxyLOD, LODGroupKey);
 }
 
 class FLandscapeNaniteSceneProxy : public ::Nanite::FSceneProxy
diff --git a/Engine/Source/Runtime/Landscape/Public/LandscapeRender.h b/Engine/Source/Runtime/Landscape/Public/LandscapeRender.h
--- a/Engine/Source/Runtime/Landscape/Public/LandscapeRender.h
+++ b/Engine/Source/Runtime/Landscape/Public/LandscapeRender.h
@@ -394,12 +394,9 @@ public:
 class FLandscapeSectionInfo : public TIntrusiveLinkedList<FLandscapeSectionInfo>
 {
 public:
-	FLandscapeSectionInfo(const UWorld* InWorld, const FGuid& InLandscapeGuid, const FIntPoint& InSectionBase);
+	FLandscapeSectionInfo(const UWorld* InWorld, const FGuid& InLandscapeGuid, const FIntPoint& InComponentBase, uint32 LODGroupKey);
 	virtual ~FLandscapeSectionInfo() = default;
 
-	void RegisterSection();
-	void UnregisterSection();
-
 	virtual float ComputeLODForView(const FSceneView& InView) const = 0;
 	virtual float ComputeLODBias() const = 0;
 	virtual int32 GetSectionPriority() const { return INDEX_NONE; }
@@ -408,9 +405,20 @@ public:
 	virtual double ComputeSectionResolution() const { return -1.0; }
 	
 	virtual void GetSectionBoundsAndLocalToWorld(FBoxSphereBounds& LocalBounds, FMatrix& LocalToWorld) const = 0;
+
+	/* return the resolution of a component, in vertices (-1 for any sections that are not grid based, i.e. mesh sections) */
+	virtual int32 GetComponentResolution() const { return -1; }
+
+	/* Used to notify derived classes when render coords are calculated */
+	virtual void OnRenderCoordsChanged() = 0;
+
 public:
-	uint32 LandscapeKey;
-	FIntPoint ComponentBase;
+	uint32 LandscapeKey;					// a hash of the world and (LandscapeGUID or LOD Group Key)
+	uint32 LODGroupKey;						// LOD Group Key (0 if no group)
+	FIntPoint RenderCoord;					// coordinate in the RenderSystem
+	FIntPoint ComponentBase;				// component base coordinate (relative to the ALandscape actor)
+
+	bool bResourcesCreated;
 	bool bRegistered;
 };
 
@@ -473,6 +481,15 @@ struct FLandscapeRenderSystem
 	/** Forced LOD level which overrides the ForcedLOD level of all the sections under this LandscapeRenderSystem. */
 	int8 ForcedLODOverride;
 
+	//  Resolution, Origin and Size, for use in LOD Groups to verify that all landscapes are of matching resolutions, orientation and scale
+	int32 ComponentResolution = -1;
+	FVector ComponentOrigin = FVector::ZeroVector;		// world space position of the center of the origin component (render coord 0,0)
+	FVector ComponentXVector = FVector::ZeroVector;		// world space vector in the direction of component local X
+	FVector ComponentYVector = FVector::ZeroVector;		// world space vector in the direction of component local Y
+
+	// Counter used to reduce how often we call compact on the map when removing sections
+	int32 SectionsRemovedSinceLastCompact;
+
 	FLandscapeRenderSystem();
 	~FLandscapeRenderSystem();
 
@@ -482,30 +499,37 @@ struct FLandscapeRenderSystem
 	static void RegisterSection(FLandscapeSectionInfo* SectionInfo);
 	static void UnregisterSection(FLandscapeSectionInfo* SectionInfo);
 
-	int32 GetSectionLinearIndex(FIntPoint InSectionBase) const
+	int32 GetSectionLinearIndex(FIntPoint InRenderCoord) const
 	{
-		return (InSectionBase.Y - Min.Y) * Size.X + InSectionBase.X - Min.X;
+		check(InRenderCoord.X >= Min.X && InRenderCoord.X < Min.X + Size.X);
+		check(InRenderCoord.Y >= Min.Y && InRenderCoord.Y < Min.Y + Size.Y);
+		int32 LinearIndex = (InRenderCoord.Y - Min.Y) * Size.X + InRenderCoord.X - Min.X;
+		return LinearIndex;
 	}
-	void ResizeAndMoveTo(FIntPoint NewMin, FIntPoint NewSize);
+	
+	void ResizeAndMoveTo(FIntPoint NewMin, FIntPoint NewMax);
+	void ResizeToInclude(const FIntPoint& NewCoord);
+	void CompactMap();
+	bool AnySectionsInRangeInclusive(FIntPoint RangeMin, FIntPoint RangeMax);
 
-	void SetSectionInfo(FIntPoint InSectionBase, FLandscapeSectionInfo* InSectionInfo)
+	void SetSectionInfo(FIntPoint InRenderCoord, FLandscapeSectionInfo* InSectionInfo)
 	{
-		SectionInfos[GetSectionLinearIndex(InSectionBase)] = InSectionInfo;
+		SectionInfos[GetSectionLinearIndex(InRenderCoord)] = InSectionInfo;
 	}
 
-	FLandscapeSectionInfo* GetSectionInfo(FIntPoint InSectionBase)
+	FLandscapeSectionInfo* GetSectionInfo(FIntPoint InRenderCoord)
 	{
-		return SectionInfos[GetSectionLinearIndex(InSectionBase)];
+		return SectionInfos[GetSectionLinearIndex(InRenderCoord)];
 	}
 
-	float GetSectionLODValue(const FSceneView& SceneView, FIntPoint InSectionBase) const
+	float GetSectionLODValue(const FSceneView& SceneView, FIntPoint InRenderCoord) const
 	{
-		return CachedSectionLODValues[SceneView.GetViewKey()][GetSectionLinearIndex(InSectionBase)];
+		return CachedSectionLODValues[SceneView.GetViewKey()][GetSectionLinearIndex(InRenderCoord)];
 	}
 
-	float GetSectionLODBias(FIntPoint InSectionBase) const
+	float GetSectionLODBias(FIntPoint InRenderCoord) const
 	{
-		return SectionLODBiases[GetSectionLinearIndex(InSectionBase)];
+		return SectionLODBiases[GetSectionLinearIndex(InRenderCoord)];
 	}
 
 	const TResourceArray<float>& ComputeSectionsLODForView(const FSceneView& InView);
@@ -618,7 +642,7 @@ class FLandscapeMeshProxySceneProxy final : public FStaticMeshSceneProxy
 public:
 	SIZE_T GetTypeHash() const override;
 
-	FLandscapeMeshProxySceneProxy(UStaticMeshComponent* InComponent, const FGuid& InLandscapeGuid, const TArray<FIntPoint>& InProxySectionsBases, int8 InProxyLOD);
+	FLandscapeMeshProxySceneProxy(UStaticMeshComponent* InComponent, const FGuid& InLandscapeGuid, const TArray<FIntPoint>& InProxySectionsBases, int8 InProxyLOD, uint32 InLODGroupKey);
 	virtual void CreateRenderThreadResources() override;
 	virtual void DestroyRenderThreadResources() override;
 	virtual bool OnLevelAddedToWorld_RenderThread() override;
@@ -688,14 +712,14 @@ public:
 	static TMap<uint32, FLandscapeSharedBuffers*> SharedBuffersMap;
 
 protected:
-	int8						MaxLOD;		// Maximum LOD level, user override possible
+	int8						MaxLOD;						// Maximum LOD level, user override possible
 	int8						NumWeightmapLayerAllocations;
 	uint8						StaticLightingLOD;
 	uint8						VirtualTexturePerPixelHeight;
 	float						WeightmapSubsectionOffset;
 	TArray<float>				LODScreenRatioSquared;		// Table of valid screen size -> LOD index
-	int32						FirstLOD;	// First LOD we have batch elements for
-	int32						LastLOD;	// Last LOD we have batch elements for
+	int32						FirstLOD;					// First LOD we have batch elements for
+	int32						LastLOD;					// Last LOD we have batch elements for
 	int32						FirstVirtualTextureLOD;
 	int32						LastVirtualTextureLOD;
 	float						ComponentMaxExtend; 		// The max extend value in any axis
@@ -723,7 +747,7 @@ protected:
 	/** Address of the component within the parent Landscape in unique height texels. */
 	FIntPoint					SectionBase;
 
-	const ULandscapeComponent* LandscapeComponent;
+	const ULandscapeComponent*	LandscapeComponent;
 
 	FMatrix						LocalToWorldNoScaling;
 
@@ -867,6 +891,9 @@ public:
 	// FLandscapeSectionInfo interface
 	virtual float ComputeLODForView(const FSceneView& InView) const override;
 	virtual float ComputeLODBias() const override;
+	virtual void OnRenderCoordsChanged() override;
+	virtual int32 GetComponentResolution() const override;
+
 	virtual double ComputeSectionResolution() const override;
 	virtual void GetSectionBoundsAndLocalToWorld(FBoxSphereBounds& LocalBounds, FMatrix& LocalToWorld) const override;
 };
