Make mobile weight-map allocations available in the runtime.

diff --git a/Engine/Source/Runtime/Landscape/Classes/LandscapeComponent.h b/Engine/Source/Runtime/Landscape/Classes/LandscapeComponent.h
--- a/Engine/Source/Runtime/Landscape/Classes/LandscapeComponent.h
+++ b/Engine/Source/Runtime/Landscape/Classes/LandscapeComponent.h
@@ -612,10 +612,11 @@ public:
 	UPROPERTY(NonPIEDuplicateTransient)
 	TArray<TObjectPtr<UTexture2D>> MobileWeightmapTextures;
 
-#if WITH_EDITORONLY_DATA
-	/** Layer allocations used by mobile. Cached value here used only in the editor for usage visualization. */
+	/** Layer allocations used by mobile.*/
+	UPROPERTY()
 	TArray<FWeightmapLayerAllocationInfo> MobileWeightmapLayerAllocations;
 
+#if WITH_EDITORONLY_DATA
 	/** The editor needs to save out the combination MIC we'll use for mobile, 
 	  because we cannot generate it at runtime for standalone PIE games */
 	UPROPERTY(NonPIEDuplicateTransient)
@@ -706,6 +707,9 @@ public:
 	LANDSCAPE_API TArray<FWeightmapLayerAllocationInfo>& GetWeightmapLayerAllocations(const FGuid& InLayerGuid);
 	LANDSCAPE_API const TArray<FWeightmapLayerAllocationInfo>& GetWeightmapLayerAllocations(const FGuid& InLayerGuid) const;
 
+	LANDSCAPE_API TArray<FWeightmapLayerAllocationInfo>& GetCurrentRuntimeWeightmapLayerAllocations();
+	LANDSCAPE_API const TArray<FWeightmapLayerAllocationInfo>& GetCurrentRuntimeWeightmapLayerAllocations() const;
+
 	const TArray<FLandscapePerLODMaterialOverride>& GetPerLODOverrideMaterials() const { return PerLODOverrideMaterials; }
 	void SetPerLODOverrideMaterials(const TArray<FLandscapePerLODMaterialOverride>& InValue) { PerLODOverrideMaterials = InValue; }
 
diff --git a/Engine/Source/Runtime/Landscape/Private/Landscape.cpp b/Engine/Source/Runtime/Landscape/Private/Landscape.cpp
--- a/Engine/Source/Runtime/Landscape/Private/Landscape.cpp
+++ b/Engine/Source/Runtime/Landscape/Private/Landscape.cpp
@@ -1282,6 +1282,13 @@ void ULandscapeComponent::PostLoad()
 	}
 #endif // UE_BUILD_SHIPPING
 
+
+	// May have been saved without mobile layer allocations, but those are serialized now
+	if (MobileWeightmapLayerAllocations.Num() == 0)
+	{
+		GenerateMobileWeightmapLayerAllocations();
+	}
+
 	if (!HasAnyFlags(RF_ClassDefaultObject))
 	{
 		FSceneInterface* SceneInterface = GetScene();
@@ -2045,6 +2052,19 @@ TArray<FWeightmapLayerAllocationInfo>& ULandscapeComponent::GetWeightmapLayerAll
 	return WeightmapLayerAllocations;
 }
 
+const TArray<FWeightmapLayerAllocationInfo>& ULandscapeComponent::GetCurrentRuntimeWeightmapLayerAllocations() const
+{
+
+	bool bIsMobile = GetWorld()->GetFeatureLevel() == ERHIFeatureLevel::ES3_1;
+	return bIsMobile ? MobileWeightmapLayerAllocations : WeightmapLayerAllocations;
+}
+
+TArray<FWeightmapLayerAllocationInfo>& ULandscapeComponent::GetCurrentRuntimeWeightmapLayerAllocations()
+{
+	bool bIsMobile = GetWorld()->GetFeatureLevel() == ERHIFeatureLevel::ES3_1;
+	return bIsMobile ? MobileWeightmapLayerAllocations : WeightmapLayerAllocations;
+}
+
 #if WITH_EDITOR
 
 FLandscapeLayerComponentData* ULandscapeComponent::GetEditingLayer()
