Test if the FLandscapeSharedBuffers::TileVertexFactory is null before we setup the MeshBatch. When conditions on the landscape culling are changing, the landscape proxy will be mark as dirty to be recreated, but sharedbuffer resource wasn't create. In this case, we just fallback to the non landscape culling path, until the proxy is recreated.

diff --git a/Engine/Source/Runtime/Landscape/Private/LandscapeCulling.cpp b/Engine/Source/Runtime/Landscape/Private/LandscapeCulling.cpp
--- a/Engine/Source/Runtime/Landscape/Private/LandscapeCulling.cpp
+++ b/Engine/Source/Runtime/Landscape/Private/LandscapeCulling.cpp
@@ -508,12 +508,13 @@ void SetupMeshBatch(const FLandscapeSharedBuffers& SharedBuffers, FMeshBatch& Me
 	{
 		return;
 	}
-		
-	if (MeshBatch.LODIndex == 0)
+
+	// TileVertexFactory can be temporarily null when we switch conditions on landscape culling. In this
+	// case, we just fallback to the non landscape culling path.
+	if (MeshBatch.LODIndex == 0 && SharedBuffers.TileVertexFactory != nullptr)
 	{
 		uint32 SubsectionSizeQuads = (SharedBuffers.SubsectionSizeVerts - 1);
 
-		check(SharedBuffers.TileVertexFactory);
 
 		MeshBatch.bViewDependentArguments = true;
 		MeshBatch.VertexFactory = SharedBuffers.TileVertexFactory;
