Fixed performance issues related to landscape heightfield in "show collision" mode :
* Use static draw commands
* Prevent landscape heightfield collision components to be invalidated because of editor throttling (time-only ticks, where the viewports are not refreshed, would lead to the landscape subsystem to wrongly detect that no view required the collision to be displayed, hence would delete all related scene proxies as a result, only to reallocate them during the next proper viewport refresh)
* Also used unlit vertex color material to better visualize physical materials

diff --git a/Engine/Source/Runtime/Landscape/Private/Landscape.cpp b/Engine/Source/Runtime/Landscape/Private/Landscape.cpp
--- a/Engine/Source/Runtime/Landscape/Private/Landscape.cpp
+++ b/Engine/Source/Runtime/Landscape/Private/Landscape.cpp
@@ -3453,11 +3453,6 @@ bool ULandscapeInfo::UpdateLayerInfoMapInternal(ALandscapeProxy* Proxy, bool bIn
 				});
 			}
 		}
-
-		//if (GCallbackEvent)
-		//{
-		//	GCallbackEvent->Send( CALLBACK_EditorPostModal );
-		//}
 	}
 	return bHasCollision;
 }
diff --git a/Engine/Source/Runtime/Landscape/Private/LandscapeCollision.cpp b/Engine/Source/Runtime/Landscape/Private/LandscapeCollision.cpp
--- a/Engine/Source/Runtime/Landscape/Private/LandscapeCollision.cpp
+++ b/Engine/Source/Runtime/Landscape/Private/LandscapeCollision.cpp
@@ -71,29 +71,35 @@ using namespace PhysicsInterfaceTypes;
 // It's a lot faster to compute than to request from DDC, so always skip.
 bool GLandscapeCollisionSkipDDC = true;
 
-static int32 CVarLandscapeCollisionMeshShowValue = static_cast<int>(EHeightfieldSource::Simple);
+
+// Callback to flag scene proxy as dirty when cvars changes
+static void OnCVarLandscapeShowCollisionMeshChanged(IConsoleVariable*)
+{
+	for (ULandscapeHeightfieldCollisionComponent* LandscapeHeightfieldCollisionComponent : TObjectRange<ULandscapeHeightfieldCollisionComponent>(RF_ClassDefaultObject | RF_ArchetypeObject, true, EInternalObjectFlags::Garbage))
+	{
+		LandscapeHeightfieldCollisionComponent->MarkRenderStateDirty();
+	}
+}
+
 static TAutoConsoleVariable<int32> CVarLandscapeCollisionMeshShow(
 	TEXT("landscape.CollisionMesh.Show"),
-	CVarLandscapeCollisionMeshShowValue,
+	static_cast<int>(EHeightfieldSource::Simple),
 	TEXT("Selects which heightfield to visualize when ShowFlags.Collision is used. 0 to disable, 1 for simple, 2 for complex, 3 for editor only."),
-	ECVF_RenderThreadSafe
-);
+	FConsoleVariableDelegate::CreateStatic(OnCVarLandscapeShowCollisionMeshChanged),
+	ECVF_RenderThreadSafe);
 
-static float CVarLandscapeCollisionMeshHeightOffsetValue = 0.f;
 static TAutoConsoleVariable<float> CVarLandscapeCollisionMeshHeightOffset(
 	TEXT("landscape.CollisionMesh.HeightOffset"),
-	CVarLandscapeCollisionMeshHeightOffsetValue,
+	0.0f,
 	TEXT("Offsets the collision mesh wireframe to assist in viewing from distances where the lower landscape lods might hide it."),
-	ECVF_RenderThreadSafe
-);
+	ECVF_RenderThreadSafe);
 
-static bool CVarLandscapeCollisionMeshShowPhysicalMaterialValue = false;
 static TAutoConsoleVariable<bool> CVarLandscapeCollisionMeshShowPhysicalMaterial(
 	TEXT("landscape.CollisionMesh.ShowPhysicalMaterial"),
-	CVarLandscapeCollisionMeshShowPhysicalMaterialValue,
+	false,
 	TEXT("When enabled, vertex colors of the collision mesh are chosen based on the physical material"),
-	ECVF_RenderThreadSafe
-);
+	ECVF_RenderThreadSafe);
+
 
 #if ENABLE_COOK_STATS
 namespace LandscapeCollisionCookStats
@@ -489,36 +495,6 @@ void ULandscapeHeightfieldCollisionComponent::ApplyWorldOffset(const FVector& In
 	}
 }
 
-template<typename Type>
-bool LandscapeCollisionCVarChanged(const TAutoConsoleVariable<Type>& CVar, Type& OldValue)
-{
-	const Type NewValue = CVar.GetValueOnAnyThread();
-
-	if (NewValue != OldValue)
-	{
-		OldValue = NewValue;
-		return true;
-	}
-	return false;
-}
-
-// Callbacks to flag scene proxy as dirty when cvars changes
-static void OnLandscapeShowCollisionMeshChanged()
-{
-	if (LandscapeCollisionCVarChanged(CVarLandscapeCollisionMeshShow, CVarLandscapeCollisionMeshShowValue)
-		|| LandscapeCollisionCVarChanged(CVarLandscapeCollisionMeshHeightOffset, CVarLandscapeCollisionMeshHeightOffsetValue)
-		|| LandscapeCollisionCVarChanged(CVarLandscapeCollisionMeshShowPhysicalMaterial, CVarLandscapeCollisionMeshShowPhysicalMaterialValue))
-	{
-		for (ULandscapeHeightfieldCollisionComponent* LandscapeHeightfieldCollisionComponent : TObjectRange<ULandscapeHeightfieldCollisionComponent>(RF_ClassDefaultObject | RF_ArchetypeObject, true, EInternalObjectFlags::Garbage))
-		{
-			LandscapeHeightfieldCollisionComponent->MarkRenderStateDirty();
-		}
-	}
-}
-
-// Registers callback with cvars
-static FAutoConsoleVariableSink OnLandscapeCollisionMeshChanged(FConsoleCommandDelegate::CreateStatic(&OnLandscapeShowCollisionMeshChanged));
-
 #if !(UE_BUILD_SHIPPING || UE_BUILD_TEST) && WITH_EDITORONLY_DATA
 FPrimitiveSceneProxy* ULandscapeHeightfieldCollisionComponent::CreateSceneProxy()
 {
@@ -535,7 +511,10 @@ FPrimitiveSceneProxy* ULandscapeHeightfieldCollisionComponent::CreateSceneProxy(
 		// used to construct the collision mesh inside GetDynamicMeshElements
 		FLandscapeHeightfieldCollisionComponentSceneProxy(const ULandscapeHeightfieldCollisionComponent* InComponent, const TArray<Chaos::FMaterialHandle>& InUsedChaosMaterials, const Chaos::FHeightField& InHeightfield, const FLinearColor& InWireframeColor)
 			: FPrimitiveSceneProxy(InComponent)
+			, VertexFactory(InComponent->GetWorld()->FeatureLevel, "FLandscapeHeightfieldCollisionComponentSceneProxy")
 		{
+			TArray<FDynamicMeshVertex> Vertices;
+
 			const Chaos::FHeightField::FData<uint16>& GeomData = InHeightfield.GeomData; 
 			const int32 NumRows = InHeightfield.GetNumRows();
 			const int32 NumCols = InHeightfield.GetNumCols();
@@ -554,13 +533,11 @@ FPrimitiveSceneProxy* ULandscapeHeightfieldCollisionComponent::CreateSceneProxy(
 				MaterialIndexColors.Emplace(HashElement[0], HashElement[1], HashElement[2]);
 			}
 
-			const float HeightOffset = CVarLandscapeCollisionMeshHeightOffset.GetValueOnGameThread();
-
 			for (int32 I = 0; I < NumVerts; I++)
 			{
 				const Chaos::FVec3 Point = GeomData.GetPointScaled(I);
 				uint8 MaterialIndex = InHeightfield.GetMaterialIndex(I % NumCols, I / NumCols);
-				Vertices[I].Position = FVector3f(static_cast<float>(Point.X), static_cast<float>(Point.Y), static_cast<float>(Point.Z + HeightOffset));
+				Vertices[I].Position = FVector3f(static_cast<float>(Point.X), static_cast<float>(Point.Y), static_cast<float>(Point.Z));
 
 				// Material indices are not defined for the last row/column in each component since they are per-triangle and not per-vertex.
 				// To show something intuitive for the user, we simply extend the previous vertex onto the final one.
@@ -577,7 +554,7 @@ FPrimitiveSceneProxy* ULandscapeHeightfieldCollisionComponent::CreateSceneProxy(
 					Vertices[I].Color = (MaterialIndex == 255) ? FColor::Black  : MaterialIndexColors[MaterialIndex];
 				}
 			}
-			Indices.SetNumUninitialized(NumTris * 3);
+			IndexBuffer.Indices.SetNumUninitialized(NumTris * 3);
 
 			// Editor heightfields don't have material indices (hence, no holes), in which case InHeightfield.GeomData.MaterialIndices.Num() == 1 : 
 			const int32 NumMaterialIndices = InHeightfield.GeomData.MaterialIndices.Num();
@@ -602,43 +579,69 @@ FPrimitiveSceneProxy* ULandscapeHeightfieldCollisionComponent::CreateSceneProxy(
 
 					if (bHole)
 					{
-						Indices[TriangleIdx + 0] = (X + 0) + (Y + 0) * NumCols;
-						Indices[TriangleIdx + 1] = Indices[TriangleIdx + 0];
-						Indices[TriangleIdx + 2] = Indices[TriangleIdx + 0];
+						IndexBuffer.Indices[TriangleIdx + 0] = (X + 0) + (Y + 0) * NumCols;
+						IndexBuffer.Indices[TriangleIdx + 1] = IndexBuffer.Indices[TriangleIdx + 0];
+						IndexBuffer.Indices[TriangleIdx + 2] = IndexBuffer.Indices[TriangleIdx + 0];
 					}
 					else
 					{
-						Indices[TriangleIdx + 0] = (X + 0) + (Y + 0) * NumCols;
-						Indices[TriangleIdx + 1] = (X + 1) + (Y + 1) * NumCols;
-						Indices[TriangleIdx + 2] = (X + 1) + (Y + 0) * NumCols;
+						IndexBuffer.Indices[TriangleIdx + 0] = (X + 0) + (Y + 0) * NumCols;
+						IndexBuffer.Indices[TriangleIdx + 1] = (X + 1) + (Y + 1) * NumCols;
+						IndexBuffer.Indices[TriangleIdx + 2] = (X + 1) + (Y + 0) * NumCols;
 					}
 
 					TriangleIdx += 3;
 
 					if (bHole)
 					{
-						Indices[TriangleIdx + 0] = (X + 0) + (Y + 0) * NumCols;
-						Indices[TriangleIdx + 1] = Indices[TriangleIdx + 0];
-						Indices[TriangleIdx + 2] = Indices[TriangleIdx + 0];
+						IndexBuffer.Indices[TriangleIdx + 0] = (X + 0) + (Y + 0) * NumCols;
+						IndexBuffer.Indices[TriangleIdx + 1] = IndexBuffer.Indices[TriangleIdx + 0];
+						IndexBuffer.Indices[TriangleIdx + 2] = IndexBuffer.Indices[TriangleIdx + 0];
 					}
 					else
 					{
-						Indices[TriangleIdx + 0] = (X + 0) + (Y + 0) * NumCols;
-						Indices[TriangleIdx + 1] = (X + 0) + (Y + 1) * NumCols;
-						Indices[TriangleIdx + 2] = (X + 1) + (Y + 1) * NumCols;
+						IndexBuffer.Indices[TriangleIdx + 0] = (X + 0) + (Y + 0) * NumCols;
+						IndexBuffer.Indices[TriangleIdx + 1] = (X + 0) + (Y + 1) * NumCols;
+						IndexBuffer.Indices[TriangleIdx + 2] = (X + 1) + (Y + 1) * NumCols;
 					}
 
 					TriangleIdx += 3;
 				}
 			}
 
-			WireframeMaterialInstance.Reset(new FColoredMaterialRenderProxy(
-				GEngine->VertexColorMaterial ? GEngine->WireframeMaterial->GetRenderProxy() : nullptr,
-				InWireframeColor));
+			// Allocate the static vertex resources now 
+			if (Vertices.Num() > 0)
+			{
+				FName Name = FName(TEXT("FLandscapeHeightfieldCollisionComponentSceneProxy ") + GetOwnerName().ToString());
+				// TODO [jonathan.bard] : VertexBuffers.SetOwnerName(Name);
+				IndexBuffer.SetOwnerName(Name);
+				VertexFactory.SetOwnerName(Name);
+
+				VertexBuffers.InitFromDynamicVertex(&VertexFactory, Vertices);
+				BeginInitResource(&VertexBuffers.PositionVertexBuffer);
+				BeginInitResource(&VertexBuffers.StaticMeshVertexBuffer);
+				BeginInitResource(&VertexBuffers.ColorVertexBuffer);
+				BeginInitResource(&IndexBuffer);
+				BeginInitResource(&VertexFactory);
+
+				WireframeMaterialInstance.Reset(new FColoredMaterialRenderProxy(
+					GEngine->WireframeMaterial ? GEngine->WireframeMaterial->GetRenderProxy() : nullptr,
+					InWireframeColor));
+
+				VertexColorMaterialInstance.Reset(new FColoredMaterialRenderProxy(
+					GEngine->VertexColorViewModeMaterial_ColorOnly ? GEngine->VertexColorViewModeMaterial_ColorOnly->GetRenderProxy() : nullptr,
+					FColor::White));
+			}
+		}
 
-			VertexColorMaterialInstance.Reset(new FColoredMaterialRenderProxy(
-				GEngine->VertexColorMaterial ? GEngine->VertexColorMaterial->GetRenderProxy() : nullptr,
-				FColor::White));
+		virtual ~FLandscapeHeightfieldCollisionComponentSceneProxy()
+		{
+			VertexBuffers.PositionVertexBuffer.ReleaseResource();
+			VertexBuffers.StaticMeshVertexBuffer.ReleaseResource();
+			VertexBuffers.ColorVertexBuffer.ReleaseResource();
+			IndexBuffer.ReleaseResource();
+			VertexFactory.ReleaseResource();
+			IndexBuffer.ReleaseResource();
 		}
 
 		virtual void GetDynamicMeshElements(const TArray<const FSceneView*>& Views, const FSceneViewFamily& ViewFamily, uint32 VisibilityMap, FMeshElementCollector& Collector) const override
@@ -646,9 +649,14 @@ FPrimitiveSceneProxy* ULandscapeHeightfieldCollisionComponent::CreateSceneProxy(
 			FMatrix LocalToWorldNoScale = GetLocalToWorld();
 			LocalToWorldNoScale.RemoveScaling();
 
-			const bool bDrawCollision = ViewFamily.EngineShowFlags.Collision && IsCollisionEnabled();
+			const bool bDrawCollision = ViewFamily.EngineShowFlags.Collision && ViewFamily.EngineShowFlags.Landscape && IsCollisionEnabled();
 
 			const bool bShowPhysicalMaterial = CVarLandscapeCollisionMeshShowPhysicalMaterial.GetValueOnRenderThread();
+			const float HeightOffset = CVarLandscapeCollisionMeshHeightOffset.GetValueOnRenderThread();
+			FVector ZAxis = LocalToWorldNoScale.GetUnitAxis(EAxis::Z);
+			LocalToWorldNoScale = LocalToWorldNoScale.ConcatTranslation(FVector(0.0, 0.0, HeightOffset));
+			FBoxSphereBounds Bounds = GetBounds();
+			Bounds.Origin += ZAxis * HeightOffset;
 
 			const TUniquePtr<FColoredMaterialRenderProxy>& MaterialToUse = bShowPhysicalMaterial ? VertexColorMaterialInstance : WireframeMaterialInstance;
 
@@ -659,18 +667,29 @@ FPrimitiveSceneProxy* ULandscapeHeightfieldCollisionComponent::CreateSceneProxy(
 					if (VisibilityMap & (1 << ViewIndex))
 					{
 						const FSceneView* View = Views[ViewIndex];
-						// Set up mesh builder
-						FDynamicMeshBuilder MeshBuilder(View->GetFeatureLevel());
-						MeshBuilder.AddVertices(Vertices);
-						MeshBuilder.AddTriangles(Indices);
-
-						FDynamicMeshBuilderSettings Settings;
-						Settings.bWireframe = true;
-						Settings.bDisableBackfaceCulling = false;
-						Settings.bReceivesDecals = false;
-						Settings.CastShadow = false;
-						Settings.bUseSelectionOutline = false;
-						MeshBuilder.GetMesh(LocalToWorldNoScale, MaterialToUse.Get(), SDPG_World, Settings, nullptr, ViewIndex, Collector);
+						
+						FMeshBatch& MeshBatch = Collector.AllocateMesh();
+						MeshBatch.MaterialRenderProxy = MaterialToUse.Get();
+						MeshBatch.bWireframe = true;
+						MeshBatch.VertexFactory = &VertexFactory;
+						MeshBatch.ReverseCulling = false;
+						MeshBatch.Type = PT_TriangleList;
+						MeshBatch.DepthPriorityGroup = SDPG_World;
+						MeshBatch.bCanApplyViewModeOverrides = true;
+
+						FMeshBatchElement& BatchElement = MeshBatch.Elements[0];
+						BatchElement.IndexBuffer = &IndexBuffer;
+						BatchElement.FirstIndex = 0;
+						BatchElement.NumPrimitives = IndexBuffer.Indices.Num() / 3;
+						check(BatchElement.NumPrimitives != 0);
+						BatchElement.MinVertexIndex = 0;
+						BatchElement.MaxVertexIndex = VertexBuffers.PositionVertexBuffer.GetNumVertices() - 1;
+
+						FDynamicPrimitiveUniformBuffer& DynamicPrimitiveUniformBuffer = Collector.AllocateOneFrameResource<FDynamicPrimitiveUniformBuffer>();
+						DynamicPrimitiveUniformBuffer.Set(LocalToWorldNoScale, LocalToWorldNoScale, Bounds, GetLocalBounds(), /*bReceivesDecals = */false, /*bHasPrecomputedVolumetricLightmap = */false, AlwaysHasVelocity());
+						BatchElement.PrimitiveUniformBufferResource = &DynamicPrimitiveUniformBuffer.UniformBuffer;
+
+						Collector.AddMesh(ViewIndex, MeshBatch);
 					}
 				}
 			}
@@ -694,10 +713,12 @@ FPrimitiveSceneProxy* ULandscapeHeightfieldCollisionComponent::CreateSceneProxy(
 		uint32 GetAllocatedSize(void) const { return static_cast<uint32>(FPrimitiveSceneProxy::GetAllocatedSize()); }
 
 	private:
-		TArray<FDynamicMeshVertex> Vertices;
-		TArray<uint32> Indices;
 		TUniquePtr<FColoredMaterialRenderProxy> WireframeMaterialInstance = nullptr;
 		TUniquePtr<FColoredMaterialRenderProxy> VertexColorMaterialInstance = nullptr;
+
+		FStaticMeshVertexBuffers VertexBuffers;
+		FDynamicMeshIndexBuffer32 IndexBuffer;
+		FLocalVertexFactory VertexFactory;
 	};
 
 	FLandscapeHeightfieldCollisionComponentSceneProxy* Proxy = nullptr;
diff --git a/Engine/Source/Runtime/Landscape/Private/LandscapeEdit.cpp b/Engine/Source/Runtime/Landscape/Private/LandscapeEdit.cpp
--- a/Engine/Source/Runtime/Landscape/Private/LandscapeEdit.cpp
+++ b/Engine/Source/Runtime/Landscape/Private/LandscapeEdit.cpp
@@ -5708,7 +5708,7 @@ void ALandscapeStreamingProxy::PostEditChangeProperty(FPropertyChangedEvent& Pro
 
 void ALandscapeStreamingProxy::PostRegisterAllComponents()
 {
-	ALandscapeProxy::PostRegisterAllComponents();
+	Super::PostRegisterAllComponents();
 
 #if WITH_EDITORONLY_DATA
 	// If the landscape actor is not set yet and we're transferring the property from the lazy object pointer it was previously stored as to the soft object ptr it is now stored as :
diff --git a/Engine/Source/Runtime/Landscape/Private/LandscapeSubsystem.cpp b/Engine/Source/Runtime/Landscape/Private/LandscapeSubsystem.cpp
--- a/Engine/Source/Runtime/Landscape/Private/LandscapeSubsystem.cpp
+++ b/Engine/Source/Runtime/Landscape/Private/LandscapeSubsystem.cpp
@@ -246,27 +246,33 @@ void ULandscapeSubsystem::Tick(float DeltaTime)
 
 	Super::Tick(DeltaTime);
 
+	UWorld* World = GetWorld();
+
 #if WITH_EDITOR
 	AppCurrentDateTime = FDateTime::Now();
+	uint32 FrameNumber = World->Scene->GetFrameNumber();
+	const bool bIsTimeOnlyTick = (FrameNumber == LastTickFrameNumber);
 
-	//Check if we need to start or stop creating Collision SceneProxies
 	ILandscapeModule& LandscapeModule = FModuleManager::GetModuleChecked<ILandscapeModule>("Landscape");
-	int32 NumViewsWithShowCollision = LandscapeModule.GetLandscapeSceneViewExtension()->GetNumViewsWithShowCollision();
-	bool bNewShowCollisions = NumViewsWithShowCollision > 0;
-    bool bCollisionChanged = bNewShowCollisions != bAnyViewShowCollisions;
-    bAnyViewShowCollisions = bNewShowCollisions;
+	// Check if we need to start or stop creating Collision SceneProxies. Don't do this on time-only ticks as the viewport (therefore the scenes) are not drawn in that case, which would lead to wrongly
+	//  assume that no view needed collision this frame
+	if (!bIsTimeOnlyTick)
+	{
+		int32 NumViewsWithShowCollision = LandscapeModule.GetLandscapeSceneViewExtension()->GetNumViewsWithShowCollision();
+		const bool bNewShowCollisions = NumViewsWithShowCollision > 0;
+		const bool bShowCollisionChanged = (bNewShowCollisions != bAnyViewShowCollisions);
+		bAnyViewShowCollisions = bNewShowCollisions;
         
-	if (bCollisionChanged)
-	{
-		for (ULandscapeHeightfieldCollisionComponent* LandscapeHeightfieldCollisionComponent : TObjectRange<ULandscapeHeightfieldCollisionComponent>(RF_ClassDefaultObject | RF_ArchetypeObject, true, EInternalObjectFlags::Garbage))
+		if (bShowCollisionChanged)
 		{
-			LandscapeHeightfieldCollisionComponent->MarkRenderStateDirty();
+			for (ULandscapeHeightfieldCollisionComponent* LandscapeHeightfieldCollisionComponent : TObjectRange<ULandscapeHeightfieldCollisionComponent>(RF_ClassDefaultObject | RF_ArchetypeObject, true, EInternalObjectFlags::Garbage))
+			{
+				LandscapeHeightfieldCollisionComponent->MarkRenderStateDirty();
+			}
 		}
 	}
-#endif
+#endif // WITH_EDITOR
 	
-	UWorld* World = GetWorld();
-
 	static TArray<FVector> OldCameras;
 	TArray<FVector>* Cameras = nullptr;
 	if (GUseStreamingManagerForCameras == 0)
@@ -308,7 +314,7 @@ void ULandscapeSubsystem::Tick(float DeltaTime)
 		NumMeshesToUpdate = NumNaniteMeshUpdatesAvailable;
 		NumNaniteMeshUpdatesAvailable -= NumMeshesToUpdate;
 	}
-#endif
+#endif // WITH_EDITOR
 	for (TWeakObjectPtr<ALandscapeProxy> ProxyPtr : Proxies)
 	{
 		if (ALandscapeProxy* Proxy = ProxyPtr.Get())
@@ -346,7 +352,7 @@ void ULandscapeSubsystem::Tick(float DeltaTime)
 
 #if !WITH_EDITOR
 			if (GUpdateProxyActorRenderMethodOnTickAtRuntime)
-#endif
+#endif // WITH_EDITOR
 			{
 				Proxy->UpdateRenderingMethod();
 			}
@@ -366,7 +372,8 @@ void ULandscapeSubsystem::Tick(float DeltaTime)
 
 	NaniteMeshBuildEvents.RemoveAllSwap([](const FGraphEventRef& Ref) -> bool { return Ref->IsComplete(); });
 
-#endif
+	LastTickFrameNumber = FrameNumber;
+#endif // WITH_EDITOR
 }
 
 void ULandscapeSubsystem::ForEachLandscapeInfo(TFunctionRef<bool(ULandscapeInfo*)> ForEachLandscapeInfoFunc) const
diff --git a/Engine/Source/Runtime/Landscape/Public/LandscapeSubsystem.h b/Engine/Source/Runtime/Landscape/Public/LandscapeSubsystem.h
--- a/Engine/Source/Runtime/Landscape/Public/LandscapeSubsystem.h
+++ b/Engine/Source/Runtime/Landscape/Public/LandscapeSubsystem.h
@@ -143,11 +143,11 @@ private:
 	FOnHeightmapStreamedDelegate OnHeightmapStreamed;
 	bool bAnyViewShowCollisions = false;
 	FDateTime AppCurrentDateTime; // Represents FDateTime::Now(), at the beginning of the frame (useful to get a human-readable date/time that is fixed during the frame)
+	int32 LastTickFrameNumber = -1;
 
 	TArray<FGraphEventRef> NaniteMeshBuildEvents;
 	float NumNaniteMeshUpdatesAvailable = 0.0f;
 
 	std::atomic<int32> NaniteBuildsInFlight;
-	
 #endif // WITH_EDITOR
 };
