Build fix

#rnx

diff --git a/Engine/Source/Runtime/Landscape/Private/LandscapeCollision.cpp b/Engine/Source/Runtime/Landscape/Private/LandscapeCollision.cpp
--- a/Engine/Source/Runtime/Landscape/Private/LandscapeCollision.cpp
+++ b/Engine/Source/Runtime/Landscape/Private/LandscapeCollision.cpp
@@ -511,7 +511,7 @@ FPrimitiveSceneProxy* ULandscapeHeightfieldCollisionComponent::CreateSceneProxy(
 		// used to construct the collision mesh inside GetDynamicMeshElements
 		FLandscapeHeightfieldCollisionComponentSceneProxy(const ULandscapeHeightfieldCollisionComponent* InComponent, const TArray<Chaos::FMaterialHandle>& InUsedChaosMaterials, const Chaos::FHeightField& InHeightfield, const FLinearColor& InWireframeColor)
 			: FPrimitiveSceneProxy(InComponent)
-			, VertexFactory(InComponent->GetWorld()->FeatureLevel, "FLandscapeHeightfieldCollisionComponentSceneProxy")
+			, VertexFactory(InComponent->GetWorld()->GetFeatureLevel(), "FLandscapeHeightfieldCollisionComponentSceneProxy")
 		{
 			TArray<FDynamicMeshVertex> Vertices;
 
@@ -613,7 +613,7 @@ FPrimitiveSceneProxy* ULandscapeHeightfieldCollisionComponent::CreateSceneProxy(
 			if (Vertices.Num() > 0)
 			{
 				FName Name = FName(TEXT("FLandscapeHeightfieldCollisionComponentSceneProxy ") + GetOwnerName().ToString());
-				// TODO [jonathan.bard] : VertexBuffers.SetOwnerName(Name);
+				VertexBuffers.SetOwnerName(Name);
 				IndexBuffer.SetOwnerName(Name);
 				VertexFactory.SetOwnerName(Name);
 
