Moved landscape buffer creation off the render thread.

diff --git a/Engine/Source/Runtime/Landscape/Private/LandscapeRender.cpp b/Engine/Source/Runtime/Landscape/Private/LandscapeRender.cpp
--- a/Engine/Source/Runtime/Landscape/Private/LandscapeRender.cpp
+++ b/Engine/Source/Runtime/Landscape/Private/LandscapeRender.cpp
@@ -53,6 +53,7 @@ LandscapeRender.cpp: New terrain rendering
 #include "DataDrivenShaderPlatformInfo.h"
 #include "Algo/Transform.h"
 #include "LandscapeCulling.h"
+#include "RenderGraphBuilder.h"
 
 using namespace UE::Landscape;
 
@@ -831,10 +832,8 @@ void FLandscapeRenderSystem::FetchHeightmapLODBiases()
 	}
 }
 
-void FLandscapeRenderSystem::UpdateBuffers()
+void FLandscapeRenderSystem::UpdateBuffers(FRHICommandListBase& RHICmdList)
 {
-	check(IsInRenderingThread());
-
 	TRACE_CPUPROFILER_EVENT_SCOPE(FLandscapeRenderSystem::UpdateBuffers);
 
 	bool bUpdateUB = false;
@@ -846,15 +845,15 @@ void FLandscapeRenderSystem::UpdateBuffers()
 			FRHIResourceCreateInfo CreateInfo(TEXT("SectionLODBiasBuffer"), &SectionLODBiases);
 			const static FLazyName ClassName(TEXT("FLandscapeRenderSystem"));
 			CreateInfo.ClassName = ClassName;
-			SectionLODBiasBuffer = RHICreateVertexBuffer(SectionLODBiases.GetResourceDataSize(), BUF_ShaderResource | BUF_Dynamic, CreateInfo);
-			SectionLODBiasSRV = RHICreateShaderResourceView(SectionLODBiasBuffer, sizeof(float), PF_R32_FLOAT);
+			SectionLODBiasBuffer = RHICmdList.CreateVertexBuffer(SectionLODBiases.GetResourceDataSize(), BUF_ShaderResource | BUF_Dynamic, CreateInfo);
+			SectionLODBiasSRV = RHICmdList.CreateShaderResourceView(SectionLODBiasBuffer, sizeof(float), PF_R32_FLOAT);
 			bUpdateUB = true;
 		}
 		else
 		{
-			float* Data = (float*)RHILockBuffer(SectionLODBiasBuffer, 0, SectionLODBiases.GetResourceDataSize(), RLM_WriteOnly);
+			float* Data = (float*)RHICmdList.LockBuffer(SectionLODBiasBuffer, 0, SectionLODBiases.GetResourceDataSize(), RLM_WriteOnly);
 			FMemory::Memcpy(Data, SectionLODBiases.GetData(), SectionLODBiases.GetResourceDataSize());
-			RHIUnlockBuffer(SectionLODBiasBuffer);
+			RHICmdList.UnlockBuffer(SectionLODBiasBuffer);
 		}
 
 		if (bUpdateUB)
@@ -865,7 +864,7 @@ void FLandscapeRenderSystem::UpdateBuffers()
 			Parameters.Size = Size;
 			Parameters.SectionLODBias = SectionLODBiasSRV;
 
-			RHIUpdateUniformBuffer(SectionLODUniformBuffer, &Parameters);
+			RHICmdList.UpdateUniformBuffer(SectionLODUniformBuffer, &Parameters);
 		}
 	}
 }
@@ -919,7 +918,7 @@ void FLandscapeSceneViewExtension::PreRenderView_RenderThread(FRDGBuilder& Graph
 	// Kick the job once all views have been collected.
 	if (!LandscapeRenderSystems.IsEmpty() && LandscapeViews.Num() == InView.Family->Views.Num())
 	{
-		auto ComputeLODs = [this]
+		LandscapeSetupTask = GraphBuilder.AddCommandListSetupTask([this] (FRHICommandListBase& RHICmdList)
 		{
 			TRACE_CPUPROFILER_EVENT_SCOPE(FLandscapeRenderSystem::ComputeLODs);
 			FOptionalTaskTagScope Scope(ETaskTag::EParallelRenderingThread);
@@ -946,16 +945,33 @@ void FLandscapeSceneViewExtension::PreRenderView_RenderThread(FRDGBuilder& Graph
 					LandscapeView.LandscapeLODData.Append(RenderSystem.ComputeSectionsLODForView(*LandscapeView.View));
 				}
 			}
-		};
 
-		if (GIsThreadedRendering && GLandscapeUseAsyncTasksForLODComputation)
-		{
-			LandscapeSetupTask = UE::Tasks::Launch(UE_SOURCE_LOCATION, MoveTemp(ComputeLODs), LowLevelTasks::ETaskPriority::Normal);
-		}
-		else
-		{
-			ComputeLODs();
-		}
+			for (auto& Pair : LandscapeRenderSystems)
+			{
+				FLandscapeRenderSystem& RenderSystem = *Pair.Value;
+				RenderSystem.UpdateBuffers(RHICmdList);
+			}
+
+			for (FLandscapeViewData& LandscapeView : LandscapeViews)
+			{
+				if (!LandscapeRenderSystems.IsEmpty())
+				{
+					FRHIResourceCreateInfo CreateInfoLODBuffer(TEXT("LandscapeLODDataBuffer"), &LandscapeView.LandscapeLODData);
+					LandscapeLODDataBuffer = RHICmdList.CreateVertexBuffer(LandscapeView.LandscapeLODData.GetResourceDataSize(), BUF_ShaderResource | BUF_Volatile, CreateInfoLODBuffer);
+					LandscapeView.View->LandscapePerComponentDataBuffer = RHICmdList.CreateShaderResourceView(LandscapeLODDataBuffer, sizeof(float), PF_R32_FLOAT);
+
+					FRHIResourceCreateInfo CreateInfoIndirection(TEXT("LandscapeIndirectionBuffer"), &LandscapeView.LandscapeIndirection);
+					LandscapeIndirectionBuffer = RHICmdList.CreateVertexBuffer(LandscapeView.LandscapeIndirection.GetResourceDataSize(), BUF_ShaderResource | BUF_Volatile, CreateInfoIndirection);
+					LandscapeView.View->LandscapeIndirectionBuffer = RHICmdList.CreateShaderResourceView(LandscapeIndirectionBuffer, sizeof(uint32), PF_R32_UINT);
+				}
+				else
+				{
+					LandscapeView.View->LandscapePerComponentDataBuffer = GWhiteVertexBufferWithSRV->ShaderResourceViewRHI;
+					LandscapeView.View->LandscapeIndirectionBuffer = GWhiteVertexBufferWithSRV->ShaderResourceViewRHI;
+				}
+			}
+
+		}, GIsThreadedRendering && GLandscapeUseAsyncTasksForLODComputation);
 	}
 }
 
@@ -963,32 +979,10 @@ void FLandscapeSceneViewExtension::PreInitViews_RenderThread(FRDGBuilder& GraphB
 {
 	LandscapeSetupTask.Wait();
 
-	for (auto& Pair : LandscapeRenderSystems)
-	{
-		FLandscapeRenderSystem& RenderSystem = *Pair.Value;
-		RenderSystem.UpdateBuffers();
-	}
-
 	TArray<const FSceneView*, TInlineAllocator<2>> SceneViews;
 
 	for (FLandscapeViewData& LandscapeView : LandscapeViews)
 	{
-		if (!LandscapeRenderSystems.IsEmpty())
-		{
-			FRHIResourceCreateInfo CreateInfoLODBuffer(TEXT("LandscapeLODDataBuffer"), &LandscapeView.LandscapeLODData);
-			LandscapeLODDataBuffer = RHICreateVertexBuffer(LandscapeView.LandscapeLODData.GetResourceDataSize(), BUF_ShaderResource | BUF_Volatile, CreateInfoLODBuffer);
-			LandscapeView.View->LandscapePerComponentDataBuffer = RHICreateShaderResourceView(LandscapeLODDataBuffer, sizeof(float), PF_R32_FLOAT);
-
-			FRHIResourceCreateInfo CreateInfoIndirection(TEXT("LandscapeIndirectionBuffer"), &LandscapeView.LandscapeIndirection);
-			LandscapeIndirectionBuffer = RHICreateVertexBuffer(LandscapeView.LandscapeIndirection.GetResourceDataSize(), BUF_ShaderResource | BUF_Volatile, CreateInfoIndirection);
-			LandscapeView.View->LandscapeIndirectionBuffer = RHICreateShaderResourceView(LandscapeIndirectionBuffer, sizeof(uint32), PF_R32_UINT);
-		}
-		else
-		{
-			LandscapeView.View->LandscapePerComponentDataBuffer = GWhiteVertexBufferWithSRV->ShaderResourceViewRHI;
-			LandscapeView.View->LandscapeIndirectionBuffer = GWhiteVertexBufferWithSRV->ShaderResourceViewRHI;
-		}
-
 		SceneViews.Add(LandscapeView.View);
 	}
 
@@ -996,7 +990,7 @@ void FLandscapeSceneViewExtension::PreInitViews_RenderThread(FRDGBuilder& GraphB
 	{
 		Culling::InitMainViews(GraphBuilder, SceneViews);
 	}
-		
+
 	LandscapeViews.Reset();
 }
 
diff --git a/Engine/Source/Runtime/Landscape/Public/LandscapeRender.h b/Engine/Source/Runtime/Landscape/Public/LandscapeRender.h
--- a/Engine/Source/Runtime/Landscape/Public/LandscapeRender.h
+++ b/Engine/Source/Runtime/Landscape/Public/LandscapeRender.h
@@ -534,7 +534,7 @@ struct FLandscapeRenderSystem
 
 	const TResourceArray<float>& ComputeSectionsLODForView(const FSceneView& InView);
 	void FetchHeightmapLODBiases();
-	void UpdateBuffers();
+	void UpdateBuffers(FRHICommandListBase& RHICmdList);
 
 private:
 	void CreateResources_Internal(FLandscapeSectionInfo* InSectionInfo);
