Fix for landscape culling crash with out-of-bounds render coords

diff --git a/Engine/Source/Runtime/Landscape/Classes/LandscapeProxy.h b/Engine/Source/Runtime/Landscape/Classes/LandscapeProxy.h
--- a/Engine/Source/Runtime/Landscape/Classes/LandscapeProxy.h
+++ b/Engine/Source/Runtime/Landscape/Classes/LandscapeProxy.h
@@ -464,7 +464,7 @@ public:
 	float LOD0ScreenSize;
 
 	UPROPERTY(EditAnywhere, Category = LOD, AdvancedDisplay, meta = (ToolTip = "Specifies the LOD Group (Zero is No Group). All landscapes in the same group calculate their LOD together, allowing matching border LODs to fix geometry seams."))
-	uint32 LODGroupKey;
+	uint32 LODGroupKey = 0;
 
 	/** The distribution setting used to change the LOD 0 generation, 1.75 is the normal distribution, numbers influence directly the LOD0 proportion on screen. */
 	UPROPERTY(EditAnywhere, Category = "LOD Distribution", meta = (DisplayName = "LOD 0", ClampMin = "1.0", ClampMax = "10.0", UIMin = "1.0", UIMax = "10.0"))
diff --git a/Engine/Source/Runtime/Landscape/Private/LandscapeCulling.cpp b/Engine/Source/Runtime/Landscape/Private/LandscapeCulling.cpp
--- a/Engine/Source/Runtime/Landscape/Private/LandscapeCulling.cpp
+++ b/Engine/Source/Runtime/Landscape/Private/LandscapeCulling.cpp
@@ -586,10 +586,13 @@ static float ComputeNeighborsMaxLOD(const FLandscapeRenderSystem& RenderSystem,
 	float NeighborsMaxLOD = 0.f;
 	for (int32 i = 0; i < UE_ARRAY_COUNT(NeightborBaseOffsets); ++i)
 	{
-		int32 SectionIdx = RenderSystem.GetSectionLinearIndex(RenderCoord + NeightborBaseOffsets[i]);
-		if (SectionLODValues.IsValidIndex(SectionIdx))
+		if (RenderSystem.IsValidCoord(RenderCoord + NeightborBaseOffsets[i]))
 		{
-			NeighborsMaxLOD = FMath::Max(SectionLODValues[SectionIdx], NeighborsMaxLOD);
+			int32 SectionIdx = RenderSystem.GetSectionLinearIndex(RenderCoord + NeightborBaseOffsets[i]);
+			if (SectionLODValues.IsValidIndex(SectionIdx))
+			{
+				NeighborsMaxLOD = FMath::Max(SectionLODValues[SectionIdx], NeighborsMaxLOD);
+			}
 		}
 	}
 
diff --git a/Engine/Source/Runtime/Landscape/Public/LandscapeRender.h b/Engine/Source/Runtime/Landscape/Public/LandscapeRender.h
--- a/Engine/Source/Runtime/Landscape/Public/LandscapeRender.h
+++ b/Engine/Source/Runtime/Landscape/Public/LandscapeRender.h
@@ -499,10 +499,15 @@ struct FLandscapeRenderSystem
 	static void RegisterSection(FLandscapeSectionInfo* SectionInfo);
 	static void UnregisterSection(FLandscapeSectionInfo* SectionInfo);
 
+	bool IsValidCoord(FIntPoint InRenderCoord) const
+	{
+		return	InRenderCoord.X >= Min.X && InRenderCoord.X < Min.X + Size.X &&
+				InRenderCoord.Y >= Min.Y && InRenderCoord.Y < Min.Y + Size.Y;
+	}
+
 	int32 GetSectionLinearIndex(FIntPoint InRenderCoord) const
 	{
-		check(InRenderCoord.X >= Min.X && InRenderCoord.X < Min.X + Size.X);
-		check(InRenderCoord.Y >= Min.Y && InRenderCoord.Y < Min.Y + Size.Y);
+		check(IsValidCoord(InRenderCoord));
 		int32 LinearIndex = (InRenderCoord.Y - Min.Y) * Size.X + InRenderCoord.X - Min.X;
 		return LinearIndex;
 	}
