Reworks landscape property sharing and override.
Adds LandscapeInherited and LandscapeOverridable meta tag to mark shared properties.
Enforce synchronicity between shared properties when editing and loading ALandscape and LandscapeStreamingProxy.
Modifies ALandscapeProxy detail panel to reflect how the properties are shared.

diff --git a/Engine/Source/Editor/DetailCustomizations/Private/LandscapeProxyUIDetails.cpp b/Engine/Source/Editor/DetailCustomizations/Private/LandscapeProxyUIDetails.cpp
--- a/Engine/Source/Editor/DetailCustomizations/Private/LandscapeProxyUIDetails.cpp
+++ b/Engine/Source/Editor/DetailCustomizations/Private/LandscapeProxyUIDetails.cpp
@@ -41,6 +41,9 @@
 #include "VT/RuntimeVirtualTextureVolume.h"
 #include "Widgets/DeclarativeSyntaxSupport.h"
 #include "Widgets/Input/SButton.h"
+#include "Widgets/Input/SCheckBox.h"
+#include "Widgets/IToolTip.h"
+#include "Widgets/Layout/SBox.h"
 #include "Widgets/Text/STextBlock.h"
 
 class IPropertyHandle;
@@ -324,6 +327,137 @@ void FLandscapeProxyUIDetails::CustomizeDetails( IDetailLayoutBuilder& DetailBui
 			]
 		];
 	}
+
+	TArray<TWeakObjectPtr<ALandscapeStreamingProxy>> EditingStreamingProxies;
+
+	for (const TWeakObjectPtr<UObject> EditingObject : EditingObjects)
+	{
+		TWeakObjectPtr<ALandscapeStreamingProxy> LandscapeStreamingProxyPtr = Cast<ALandscapeStreamingProxy>(EditingObject.Get());
+
+		if (LandscapeStreamingProxyPtr.IsValid())
+		{
+			EditingStreamingProxies.Add(MoveTemp(LandscapeStreamingProxyPtr));
+		}
+	}
+
+	ALandscapeStreamingProxy* LandscapeStreamingProxy = EditingStreamingProxies.IsEmpty() ? nullptr : EditingStreamingProxies[0].Get();
+
+	if (LandscapeStreamingProxy != nullptr)
+	{
+		for (TFieldIterator<FProperty> PropertyIterator(LandscapeStreamingProxy->GetClass()); PropertyIterator; ++PropertyIterator)
+		{
+			FProperty* Property = *PropertyIterator;
+
+			if (Property == nullptr)
+			{
+				continue;
+			}
+
+			if (LandscapeStreamingProxy->IsPropertyInherited(Property))
+			{
+				TSharedPtr<IPropertyHandle> PropertyHandle = DetailBuilder.GetProperty(Property->GetFName());
+				
+				if (PropertyHandle->IsValidHandle())
+				{
+					IDetailPropertyRow* DetailRow = DetailBuilder.EditDefaultProperty(PropertyHandle);
+
+					if (DetailRow != nullptr)
+					{
+						// Extend the tool tip to indicate this property is inherited
+						FText ToolTipText = PropertyHandle->GetToolTipText();
+						DetailRow->ToolTip(FText::Format(NSLOCTEXT("Landscape", "InheritedProperty", "{0} This property is inherited from the parent Landscape proxy."), ToolTipText));
+
+						// Disable the property editing
+						DetailRow->IsEnabled(false);
+					}
+				}
+			}
+			else if (LandscapeStreamingProxy->IsPropertyOverridable(Property))
+			{
+				TSharedPtr<IPropertyHandle> PropertyHandle = DetailBuilder.GetProperty(Property->GetFName());
+
+				if (PropertyHandle->IsValidHandle())
+				{
+					const FText TooltipText = NSLOCTEXT("Landscape", "OverriddenProperty", "Check this box to override the parent landscape's property.");
+					FString PropertyName = Property->GetName();
+					IDetailPropertyRow* DetailRow = DetailBuilder.EditDefaultProperty(PropertyHandle);
+					TSharedPtr<SWidget> NameWidget = nullptr;
+					TSharedPtr<SWidget> ValueWidget = nullptr;
+
+					DetailRow->GetDefaultWidgets(NameWidget, ValueWidget);
+
+					DetailRow->CustomWidget()
+					.NameContent()
+					[
+						SNew(SHorizontalBox)
+						+ SHorizontalBox::Slot()
+						.AutoWidth()
+						[
+							SNew(SCheckBox)
+							.ToolTipText(TooltipText)
+							.IsChecked_Lambda([EditingStreamingProxies, PropertyName]() -> ECheckBoxState
+							{
+								bool bContainsOverriddenProperty = false;
+								bool bContainsDefaultProperty = false;
+
+								for (const TWeakObjectPtr<ALandscapeStreamingProxy> StreamingProxy : EditingStreamingProxies)
+								{
+									if (!StreamingProxy.IsValid())
+									{
+										continue;
+									}
+
+									const bool bIsPropertyOverridden = StreamingProxy->IsSharedPropertyOverridden(PropertyName);
+
+									bContainsOverriddenProperty |= bIsPropertyOverridden;
+									bContainsDefaultProperty |= !bIsPropertyOverridden;
+								}
+
+								if (bContainsOverriddenProperty && bContainsDefaultProperty)
+								{
+									return ECheckBoxState::Undetermined;
+								}
+
+								return bContainsOverriddenProperty ? ECheckBoxState::Checked : ECheckBoxState::Unchecked;
+							})
+							.OnCheckStateChanged_Lambda([EditingStreamingProxies, PropertyName](ECheckBoxState NewState)
+							{
+								if (NewState == ECheckBoxState::Undetermined)
+								{
+									return;
+								}
+
+								const bool bChecked = NewState == ECheckBoxState::Checked;
+
+								for (const TWeakObjectPtr<ALandscapeStreamingProxy> StreamingProxy : EditingStreamingProxies)
+								{
+									if (!StreamingProxy.IsValid())
+									{
+										continue;
+									}
+
+									StreamingProxy->SetSharedPropertyOverride(PropertyName, bChecked);
+								}
+							})
+						]
+						+ SHorizontalBox::Slot()
+						.AutoWidth()
+						[
+							NameWidget->AsShared()
+						]
+					]
+					.ValueContent()
+					[
+						SNew(SBox)
+						.IsEnabled_Lambda([LandscapeStreamingProxy, PropertyName]() { return LandscapeStreamingProxy->IsSharedPropertyOverridden(PropertyName); })
+						[
+							ValueWidget->AsShared()
+						]
+					];
+				}
+			}
+		}
+	}
 }
 
 static void GetMissingRuntimeVirtualTextureVolumes(ALandscape* InLandscapeActor, TArray<URuntimeVirtualTexture*>& OutVirtualTextures)
diff --git a/Engine/Source/Editor/UnrealEd/Private/Commandlets/WorldPartitionConvertCommandlet.cpp b/Engine/Source/Editor/UnrealEd/Private/Commandlets/WorldPartitionConvertCommandlet.cpp
--- a/Engine/Source/Editor/UnrealEd/Private/Commandlets/WorldPartitionConvertCommandlet.cpp
+++ b/Engine/Source/Editor/UnrealEd/Private/Commandlets/WorldPartitionConvertCommandlet.cpp
@@ -990,7 +990,7 @@ int32 UWorldPartitionConvertCommandlet::Main(const FString& Params)
 				FTransform LandscapeTransform = FirstProxy->LandscapeActorToWorld();
 				ALandscape* NewLandscape = MainWorld->SpawnActor<ALandscape>(ALandscape::StaticClass(), LandscapeTransform, SpawnParams);
 
-				NewLandscape->GetSharedProperties(FirstProxy);
+				NewLandscape->CopySharedProperties(FirstProxy);
 
 				LandscapeInfo->RegisterActor(NewLandscape);
 			}
diff --git a/Engine/Source/Editor/WorldBrowser/Private/Tiles/WorldTileModel.cpp b/Engine/Source/Editor/WorldBrowser/Private/Tiles/WorldTileModel.cpp
--- a/Engine/Source/Editor/WorldBrowser/Private/Tiles/WorldTileModel.cpp
+++ b/Engine/Source/Editor/WorldBrowser/Private/Tiles/WorldTileModel.cpp
@@ -980,7 +980,7 @@ bool FWorldTileModel::CreateAdjacentLandscapeProxy(ALandscapeProxy* SourceLandsc
 	if (AdjacentLandscape)
 	{
 		// Copy source landscape properties 
-		AdjacentLandscape->GetSharedProperties(SourceLandscape);
+		AdjacentLandscape->CopySharedProperties(SourceLandscape);
 		
 		// Refresh level model bounding box
 		FBox AdjacentLandscapeBounds = AdjacentLandscape->GetComponentsBoundingBox(true);
diff --git a/Engine/Source/Runtime/Core/Public/UObject/FortniteReleaseBranchCustomObjectVersion.h b/Engine/Source/Runtime/Core/Public/UObject/FortniteReleaseBranchCustomObjectVersion.h
--- a/Engine/Source/Runtime/Core/Public/UObject/FortniteReleaseBranchCustomObjectVersion.h
+++ b/Engine/Source/Runtime/Core/Public/UObject/FortniteReleaseBranchCustomObjectVersion.h
@@ -36,6 +36,9 @@ struct CORE_API FFortniteReleaseBranchCustomObjectVersion
 		// Fix the Nanite landscape mesh being reused because of a bad name
 		FixNaniteLandscapeMeshNames,
 
+		// Fixup and synchronize shared properties modified before the synchronicity enforcement
+		LandscapeSharedPropertiesEnforcement,
+
 		// -----<new versions can be added above this line>-------------------------------------------------
 		VersionPlusOne,
 		LatestVersion = VersionPlusOne - 1
diff --git a/Engine/Source/Runtime/Landscape/Classes/LandscapeProxy.h b/Engine/Source/Runtime/Landscape/Classes/LandscapeProxy.h
--- a/Engine/Source/Runtime/Landscape/Classes/LandscapeProxy.h
+++ b/Engine/Source/Runtime/Landscape/Classes/LandscapeProxy.h
@@ -414,7 +414,7 @@ protected:
 	/** Guid for LandscapeEditorInfo, all proxies that belong to the same landscape should have the same LandscapeGuid, even if split across world partitions 
 	  * Note that this value may change when the landscape is instanced (or in PIE) in order to allow multiple instances of the same landscape to exist.
 	  **/
-	UPROPERTY()
+	UPROPERTY(meta = (LandscapeInherited))
 	FGuid LandscapeGuid;
 
 	/** 
@@ -423,10 +423,10 @@ protected:
 	FGuid OriginalLandscapeGuid;
 
 	/** Use Nanite to render landscape as a mesh on supported platforms. */
-	UPROPERTY(EditAnywhere, BlueprintReadOnly, Category = Nanite)
+	UPROPERTY(EditAnywhere, BlueprintReadOnly, Category = Nanite, meta = (LandscapeInherited))
 	bool bEnableNanite = false;
 
-	UPROPERTY(EditAnywhere, Category = Landscape)
+	UPROPERTY(EditAnywhere, Category = Landscape, meta = (LandscapeOverridable))
 	TArray<FLandscapePerLODMaterialOverride> PerLODOverrideMaterials;
 
 #if WITH_EDITORONLY_DATA
@@ -434,7 +434,7 @@ protected:
 	TArray<FLandscapePerLODMaterialOverride> PreEditPerLODOverrideMaterials;
 
 	/** LOD level of the landscape when generating the Nanite mesh. Mostly there for debug reasons, since Nanite is meant to allow high density meshes, we want to use 0 most of the times. */
-	UPROPERTY(EditAnywhere, Category = Nanite, AdvancedDisplay, meta = (EditCondition = "bEnableNanite"))
+	UPROPERTY(EditAnywhere, Category = Nanite, AdvancedDisplay, meta = (EditCondition = "bEnableNanite", LandscapeInherited))
 	int32 NaniteLODIndex = 0;
 #endif // WITH_EDITORONLY_DATA
 
@@ -444,7 +444,7 @@ public:
 	FIntPoint LandscapeSectionOffset;
 
 	/** Max LOD level to use when rendering, -1 means the max available */
-	UPROPERTY(EditAnywhere, Category=LOD)
+	UPROPERTY(EditAnywhere, Category=LOD, meta = (LandscapeInherited))
 	int32 MaxLODLevel;
 
 #if WITH_EDITORONLY_DATA
@@ -456,22 +456,22 @@ public:
 #endif // WITH_EDITORONLY_DATA
 
 	/** Component screen size (0.0 - 1.0) at which we should keep sub sections. This is mostly pertinent if you have large component of > 64 and component are close to the camera. The goal is to reduce draw call, so if a component is smaller than the value, we merge all subsections into 1 drawcall. */
-	UPROPERTY(EditAnywhere, Category = LOD, meta=(ClampMin = "0.01", ClampMax = "1.0", UIMin = "0.01", UIMax = "1.0", DisplayName= "SubSection Min Component ScreenSize"))
+	UPROPERTY(EditAnywhere, Category = LOD, meta=(ClampMin = "0.01", ClampMax = "1.0", UIMin = "0.01", UIMax = "1.0", DisplayName= "SubSection Min Component ScreenSize", LandscapeInherited))
 	float ComponentScreenSizeToUseSubSections;
 
 	/** This is the starting screen size used to calculate the distribution. You can increase the value if you want less LOD0 component, and you use very large landscape component. */
-	UPROPERTY(EditAnywhere, Category = "LOD Distribution", meta = (DisplayName = "LOD 0 Screen Size", ClampMin = "0.1", ClampMax = "10.0", UIMin = "1.0", UIMax = "10.0"))
+	UPROPERTY(EditAnywhere, Category = "LOD Distribution", meta = (DisplayName = "LOD 0 Screen Size", ClampMin = "0.1", ClampMax = "10.0", UIMin = "1.0", UIMax = "10.0", LandscapeInherited))
 	float LOD0ScreenSize;
 
-	UPROPERTY(EditAnywhere, Category = LOD, AdvancedDisplay, meta = (ToolTip = "Specifies the LOD Group (Zero is No Group). All landscapes in the same group calculate their LOD together, allowing matching border LODs to fix geometry seams."))
+	UPROPERTY(EditAnywhere, Category = LOD, AdvancedDisplay, meta = (ToolTip = "Specifies the LOD Group (Zero is No Group). All landscapes in the same group calculate their LOD together, allowing matching border LODs to fix geometry seams.", LandscapeInherited))
 	uint32 LODGroupKey = 0;
 
 	/** The distribution setting used to change the LOD 0 generation, 1.75 is the normal distribution, numbers influence directly the LOD0 proportion on screen. */
-	UPROPERTY(EditAnywhere, Category = "LOD Distribution", meta = (DisplayName = "LOD 0", ClampMin = "1.0", ClampMax = "10.0", UIMin = "1.0", UIMax = "10.0"))
+	UPROPERTY(EditAnywhere, Category = "LOD Distribution", meta = (DisplayName = "LOD 0", ClampMin = "1.0", ClampMax = "10.0", UIMin = "1.0", UIMax = "10.0", LandscapeInherited))
 	float LOD0DistributionSetting;
 
 	/** The distribution setting used to change the LOD generation, 2 is the normal distribution, small number mean you want your last LODs to take more screen space and big number mean you want your first LODs to take more screen space. */
-	UPROPERTY(EditAnywhere, Category = "LOD Distribution", meta = (DisplayName = "Other LODs", ClampMin = "1.0", ClampMax = "10.0", UIMin = "1.0", UIMax = "10.0"))
+	UPROPERTY(EditAnywhere, Category = "LOD Distribution", meta = (DisplayName = "Other LODs", ClampMin = "1.0", ClampMax = "10.0", UIMin = "1.0", UIMax = "10.0", LandscapeInherited))
 	float LODDistributionSetting;
 
 #if WITH_EDITORONLY_DATA
@@ -480,11 +480,11 @@ public:
 	int32 ExportLOD;
 
 	/** Display Order of the targets */
-	UPROPERTY(NonTransactional)
+	UPROPERTY(NonTransactional, meta = (LandscapeInherited))
 	TArray<FName> TargetDisplayOrderList;
 
 	/** Display Order mode for the targets */
-	UPROPERTY(NonTransactional)
+	UPROPERTY(NonTransactional, meta = (LandscapeInherited))
 	ELandscapeLayerDisplayMode TargetDisplayOrder;
 #endif
 
@@ -493,7 +493,7 @@ public:
 	int32 StaticLightingLOD;
 
 	/** Default physical material, used when no per-layer values physical materials */
-	UPROPERTY(EditAnywhere, Category=Landscape)
+	UPROPERTY(EditAnywhere, Category=Landscape, meta = (LandscapeOverridable))
 	TObjectPtr<UPhysicalMaterial> DefaultPhysMaterial;
 
 	/**
@@ -505,7 +505,7 @@ public:
 	float StreamingDistanceMultiplier;
 
 	/** Combined material used to render the landscape */
-	UPROPERTY(EditAnywhere, BlueprintSetter=EditorSetLandscapeMaterial, Category=Landscape)
+	UPROPERTY(EditAnywhere, BlueprintSetter=EditorSetLandscapeMaterial, Category=Landscape, meta=(LandscapeOverridable))
 	TObjectPtr<UMaterialInterface> LandscapeMaterial;
 
 #if !WITH_EDITORONLY_DATA
@@ -520,7 +520,7 @@ public:
 #endif
 
 	/** Material used to render landscape components with holes. If not set, LandscapeMaterial will be used (blend mode will be overridden to Masked if it is set to Opaque) */
-	UPROPERTY(EditAnywhere, Category=Landscape, AdvancedDisplay)
+	UPROPERTY(EditAnywhere, Category=Landscape, AdvancedDisplay, meta = (LandscapeOverridable))
 	TObjectPtr<UMaterialInterface> LandscapeHoleMaterial;
 
 
@@ -545,7 +545,7 @@ public:
 	 * Array of runtime virtual textures into which we draw this landscape.
 	 * The material also needs to be set up to output to a virtual texture.
 	 */
-	UPROPERTY(EditAnywhere, BlueprintReadWrite, Category = VirtualTexture, meta = (DisplayName = "Draw in Virtual Textures"))
+	UPROPERTY(EditAnywhere, BlueprintReadWrite, Category = VirtualTexture, meta = (DisplayName = "Draw in Virtual Textures", LandscapeOverridable))
 	TArray<TObjectPtr<URuntimeVirtualTexture>> RuntimeVirtualTextures;
 
 	/** Placeholder for details customization button. */
@@ -557,32 +557,32 @@ public:
 	 * This is the fastest path but it only gives correct results if the runtime virtual texture orientation is aligned with the landscape.
 	 * If the two are unaligned we need to render to the virtual texture using LODs with sufficient density.
 	 */
-	UPROPERTY(EditAnywhere, AdvancedDisplay, BlueprintReadOnly, Category = VirtualTexture)
+	UPROPERTY(EditAnywhere, AdvancedDisplay, BlueprintReadOnly, Category = VirtualTexture, meta = (LandscapeOverridable))
 	bool bVirtualTextureRenderWithQuad = false;
 
 	/** 
 	 * Use highest quality heightmap interpolation when using a single quad to render this landscape to runtime virtual texture pages.
 	 * This also requires the project setting: r.VT.RVT.HighQualityPerPixelHeight.
 	 */
-	UPROPERTY(EditAnywhere, AdvancedDisplay, BlueprintReadOnly, Category = VirtualTexture, meta = (DisplayName = "High Quality PerPixel Height", EditCondition = "bVirtualTextureRenderWithQuad"))
+	UPROPERTY(EditAnywhere, AdvancedDisplay, BlueprintReadOnly, Category = VirtualTexture, meta = (DisplayName = "High Quality PerPixel Height", EditCondition = "bVirtualTextureRenderWithQuad", LandscapeOverridable))
 	bool bVirtualTextureRenderWithQuadHQ = true;
 
 	/** 
 	 * Number of mesh levels to use when rendering landscape into runtime virtual texture.
 	 * Lower values reduce vertex count when rendering to the runtime virtual texture but decrease accuracy when using values that require vertex interpolation.
 	 */
-	UPROPERTY(EditAnywhere, AdvancedDisplay, BlueprintReadOnly, Category = VirtualTexture, meta = (DisplayName = "Virtual Texture Num LODs", EditCondition = "!bVirtualTextureRenderWithQuad", UIMin = "0", UIMax = "7"))
+	UPROPERTY(EditAnywhere, AdvancedDisplay, BlueprintReadOnly, Category = VirtualTexture, meta = (DisplayName = "Virtual Texture Num LODs", EditCondition = "!bVirtualTextureRenderWithQuad", UIMin = "0", UIMax = "7", LandscapeOverridable))
 	int32 VirtualTextureNumLods = 6;
 
 	/** 
 	 * Bias to the LOD selected for rendering to runtime virtual textures.
 	 * Higher values reduce vertex count when rendering to the runtime virtual texture.
 	 */
-	UPROPERTY(EditAnywhere, AdvancedDisplay, BlueprintReadOnly, Category = VirtualTexture, meta = (DisplayName = "Virtual Texture LOD Bias", EditCondition = "!bVirtualTextureRenderWithQuad", UIMin = "0", UIMax = "7"))
+	UPROPERTY(EditAnywhere, AdvancedDisplay, BlueprintReadOnly, Category = VirtualTexture, meta = (DisplayName = "Virtual Texture LOD Bias", EditCondition = "!bVirtualTextureRenderWithQuad", UIMin = "0", UIMax = "7", LandscapeOverridable))
 	int32 VirtualTextureLodBias = 0;
 
 	/** Controls if this component draws in the main pass as well as in the virtual texture. */
-	UPROPERTY(EditAnywhere, BlueprintSetter = SetVirtualTextureRenderPassType, Category = VirtualTexture, meta = (DisplayName = "Draw in Main Pass"))
+	UPROPERTY(EditAnywhere, BlueprintSetter = SetVirtualTextureRenderPassType, Category = VirtualTexture, meta = (DisplayName = "Draw in Main Pass", LandscapeOverridable))
 	ERuntimeVirtualTextureMainPassType VirtualTextureRenderPassType = ERuntimeVirtualTextureMainPassType::Always;
 
 	UFUNCTION(BlueprintSetter)
@@ -591,13 +591,13 @@ public:
 	/** Allows overriding the landscape bounds. This is useful if you distort the landscape with world-position-offset, for example
 	 *  Extension value in the negative Z axis, positive value increases bound size
 	 *  Note that this can also be overridden per-component when the component is selected with the component select tool */
-	UPROPERTY(EditAnywhere, Category=Landscape)
+	UPROPERTY(EditAnywhere, Category=Landscape, meta = (LandscapeOverridable))
 	float NegativeZBoundsExtension;
 
 	/** Allows overriding the landscape bounds. This is useful if you distort the landscape with world-position-offset, for example
 	 *  Extension value in the positive Z axis, positive value increases bound size
 	 *  Note that this can also be overridden per-component when the component is selected with the component select tool */
-	UPROPERTY(EditAnywhere, Category=Landscape)
+	UPROPERTY(EditAnywhere, Category=Landscape, meta = (LandscapeOverridable))
 	float PositiveZBoundsExtension;
 
 	/** The array of LandscapeComponent that are used by the landscape */
@@ -635,15 +635,15 @@ public:
 	float StaticLightingResolution;
 
 	/** Controls whether the primitive component should cast a shadow or not. */
-	UPROPERTY(EditAnywhere, BlueprintReadOnly, Category = Lighting)
+	UPROPERTY(EditAnywhere, BlueprintReadOnly, Category = Lighting, meta = (LandscapeOverridable))
 	uint8 CastShadow : 1;
 
 	/** Controls whether the primitive should cast shadows in the case of non precomputed shadowing.  This flag is only used if CastShadow is true. **/
-	UPROPERTY(EditAnywhere, BlueprintReadOnly, Category = Lighting, AdvancedDisplay, meta = (EditCondition = "CastShadow", DisplayName = "Dynamic Shadow"))
+	UPROPERTY(EditAnywhere, BlueprintReadOnly, Category = Lighting, AdvancedDisplay, meta = (EditCondition = "CastShadow", DisplayName = "Dynamic Shadow", LandscapeOverridable))
 	uint8 bCastDynamicShadow : 1;
 
 	/** Whether the object should cast a static shadow from shadow casting lights.  This flag is only used if CastShadow is true. */
-	UPROPERTY(EditAnywhere, BlueprintReadOnly, Category = Lighting, AdvancedDisplay, meta = (EditCondition = "CastShadow", DisplayName = "Static Shadow"))
+	UPROPERTY(EditAnywhere, BlueprintReadOnly, Category = Lighting, AdvancedDisplay, meta = (EditCondition = "CastShadow", DisplayName = "Static Shadow", LandscapeOverridable))
 	uint8 bCastStaticShadow : 1;
 
 	/** Control shadow invalidation behavior, in particular with respect to Virtual Shadow Maps and material effects like World Position Offset. */
@@ -651,30 +651,30 @@ public:
 	EShadowCacheInvalidationBehavior ShadowCacheInvalidationBehavior;
 
 	/** Whether the object should cast contact shadows. This flag is only used if CastShadow is true. */
-	UPROPERTY(EditAnywhere, BlueprintReadOnly, Category = Lighting, AdvancedDisplay, meta = (EditCondition = "CastShadow", DisplayName = "Contact Shadow"))
+	UPROPERTY(EditAnywhere, BlueprintReadOnly, Category = Lighting, AdvancedDisplay, meta = (EditCondition = "CastShadow", DisplayName = "Contact Shadow", LandscapeOverridable))
 	uint8 bCastContactShadow : 1;
 
 	/** When enabled, the component will be rendering into the far shadow cascades(only for directional lights).  This flag is only used if CastShadow is true. */
-	UPROPERTY(EditAnywhere, AdvancedDisplay, BlueprintReadOnly, Category = Lighting, meta = (EditCondition = "CastShadow", DisplayName = "Far Shadow"))
+	UPROPERTY(EditAnywhere, AdvancedDisplay, BlueprintReadOnly, Category = Lighting, meta = (EditCondition = "CastShadow", DisplayName = "Far Shadow", LandscapeOverridable))
 	uint32 bCastFarShadow : 1;
 
 	/** If true, the primitive will cast shadows even if bHidden is true.  Controls whether the primitive should cast shadows when hidden.  This flag is only used if CastShadow is true. */
-	UPROPERTY(EditAnywhere, AdvancedDisplay, BlueprintReadOnly, Category = Lighting, meta = (EditCondition = "CastShadow", DisplayName = "Hidden Shadow"))
+	UPROPERTY(EditAnywhere, AdvancedDisplay, BlueprintReadOnly, Category = Lighting, meta = (EditCondition = "CastShadow", DisplayName = "Hidden Shadow", LandscapeOverridable))
 	uint8 bCastHiddenShadow : 1;
 
 	/** Whether this primitive should cast dynamic shadows as if it were a two sided material.  This flag is only used if CastShadow is true. */
-	UPROPERTY(EditAnywhere, AdvancedDisplay, BlueprintReadOnly, Category = Lighting, meta = (EditCondition = "CastShadow", DisplayName = "Shadow Two Sided"))
+	UPROPERTY(EditAnywhere, AdvancedDisplay, BlueprintReadOnly, Category = Lighting, meta = (EditCondition = "CastShadow", DisplayName = "Shadow Two Sided", LandscapeOverridable))
 	uint32 bCastShadowAsTwoSided : 1;
 
 	/** Controls whether the primitive should affect dynamic distance field lighting methods.  This flag is only used if CastShadow is true. **/
-	UPROPERTY(EditAnywhere, BlueprintReadOnly, Category = Lighting, AdvancedDisplay, meta = (EditCondition = "CastShadow"))
+	UPROPERTY(EditAnywhere, BlueprintReadOnly, Category = Lighting, AdvancedDisplay, meta = (EditCondition = "CastShadow", LandscapeOverridable))
 	uint8 bAffectDistanceFieldLighting:1;
 
 	/**
 	* Channels that this Landscape should be in.  Lights with matching channels will affect the Landscape.
 	* These channels only apply to opaque materials, direct lighting, and dynamic lighting and shadowing.
 	*/
-	UPROPERTY(EditAnywhere, AdvancedDisplay, BlueprintReadOnly, Category = Lighting)
+	UPROPERTY(EditAnywhere, AdvancedDisplay, BlueprintReadOnly, Category = Lighting, meta = (LandscapeOverridable))
 	FLightingChannels LightingChannels;
 
 	/** Whether to use the landscape material's vertical world position offset when calculating static lighting.
@@ -684,32 +684,32 @@ public:
 	uint32 bUseMaterialPositionOffsetInStaticLighting:1;
 
 	/** If true, the Landscape will be rendered in the CustomDepth pass (usually used for outlines) */
-	UPROPERTY(EditAnywhere, AdvancedDisplay, BlueprintReadOnly, Category=Rendering, meta=(DisplayName = "Render CustomDepth Pass"))
+	UPROPERTY(EditAnywhere, AdvancedDisplay, BlueprintReadOnly, Category=Rendering, meta=(DisplayName = "Render CustomDepth Pass", LandscapeOverridable))
 	uint32 bRenderCustomDepth:1;
 
 	/** Mask used for stencil buffer writes. */
-	UPROPERTY(EditAnywhere, AdvancedDisplay, BlueprintReadOnly, Category = Rendering, meta = (editcondition = "bRenderCustomDepth"))
+	UPROPERTY(EditAnywhere, AdvancedDisplay, BlueprintReadOnly, Category = Rendering, meta = (editcondition = "bRenderCustomDepth", LandscapeOverridable))
 	ERendererStencilMask CustomDepthStencilWriteMask;
 
 	/** Optionally write this 0-255 value to the stencil buffer in CustomDepth pass (Requires project setting or r.CustomDepth == 3) */
-	UPROPERTY(EditAnywhere, AdvancedDisplay, BlueprintReadOnly, Category=Rendering,  meta=(UIMin = "0", UIMax = "255", editcondition = "bRenderCustomDepth", DisplayName = "CustomDepth Stencil Value"))
+	UPROPERTY(EditAnywhere, AdvancedDisplay, BlueprintReadOnly, Category=Rendering,  meta=(UIMin = "0", UIMax = "255", editcondition = "bRenderCustomDepth", DisplayName = "CustomDepth Stencil Value", LandscapeOverridable))
 	int32 CustomDepthStencilValue;
 
 	/**  Max draw distance exposed to LDs. The real max draw distance is the min (disregarding 0) of this and volumes affecting this object. */
-	UPROPERTY(EditAnywhere, AdvancedDisplay, BlueprintReadOnly, Category=Rendering, meta = (DisplayName = "Desired Max Draw Distance"))
+	UPROPERTY(EditAnywhere, AdvancedDisplay, BlueprintReadOnly, Category=Rendering, meta = (DisplayName = "Desired Max Draw Distance", LandscapeOverridable))
 	float LDMaxDrawDistance;
 
 	/** The Lightmass settings for this object. */
-	UPROPERTY(EditAnywhere, Category=Lightmass)
+	UPROPERTY(EditAnywhere, Category=Lightmass, meta = (LandscapeInherited))
 	FLightmassPrimitiveSettings LightmassSettings;
 
 	// Landscape LOD to use for collision tests. Higher numbers use less memory and process faster, but are much less accurate
-	UPROPERTY(EditAnywhere, Category=Collision)
+	UPROPERTY(EditAnywhere, Category=Collision, meta = (LandscapeOverridable))
 	int32 CollisionMipLevel;
 
 	// If set higher than the "Collision Mip Level", this specifies the Landscape LOD to use for "simple collision" tests, otherwise the "Collision Mip Level" is used for both simple and complex collision.
 	// Does not work with an XY offset map (mesh collision)
-	UPROPERTY(EditAnywhere, Category=Collision)
+	UPROPERTY(EditAnywhere, Category=Collision, meta = (LandscapeOverridable))
 	int32 SimpleCollisionMipLevel;
 
 	UE_DEPRECATED(5.2, "CollisionThickness is not supported by Chaos Physics and therefore deprecated. Please remove any usage of this property")
@@ -773,13 +773,13 @@ public:
 #endif // WITH_EDITORONLY_DATA
 
 	/** Data set at creation time */
-	UPROPERTY()
+	UPROPERTY(meta = (LandscapeInherited))
 	int32 ComponentSizeQuads;    // Total number of quads in each component
 
-	UPROPERTY()
+	UPROPERTY(meta = (LandscapeInherited))
 	int32 SubsectionSizeQuads;    // Number of quads for a subsection of a component. SubsectionSizeQuads+1 must be a power of two.
 
-	UPROPERTY()
+	UPROPERTY(meta = (LandscapeInherited))
 	int32 NumSubsections;    // Number of subsections in X and Y axis
 
 	/** Hints navigation system whether this landscape will ever be navigated on. true by default, but make sure to set it to false for faraway, background landscapes */
@@ -811,7 +811,7 @@ public:
 	bool bHasLayersContent;
 
 	/** Enable compressed heightmap texture storage. */
-	UPROPERTY(EditAnywhere, Category = Landscape, AdvancedDisplay)
+	UPROPERTY(EditAnywhere, Category = Landscape, AdvancedDisplay, meta = (LandscapeInherited))
 	bool bUseCompressedHeightmapStorage = false;
 
 #if WITH_EDITOR
@@ -1064,8 +1064,39 @@ public:
 	LANDSCAPE_API TArray<float> GetLODScreenSizeArray() const;
 
 	// Copy properties from parent Landscape actor
+	UE_DEPRECATED(5.3, "GetSharedProperties is being deprecated, please use CopySharedProperties or SynchronizeSharedProperties instead.")
 	LANDSCAPE_API void GetSharedProperties(ALandscapeProxy* Landscape);
 
+#if WITH_EDITOR
+	// Copy properties from parent Landscape actor
+	LANDSCAPE_API void CopySharedProperties(ALandscapeProxy* Landscape);
+
+	/**
+	* Enforce property sharing and requirements related to the parent Landscape actor
+	* @param InLandscape the parent landscape to use as a template.
+	* @return An array containing all synchronized properties names (properties that were different from the provided landscape).
+	*/
+	LANDSCAPE_API TArray<FName> SynchronizeSharedProperties(ALandscapeProxy* InLandscape);
+
+	// Returns true if the property is shared.
+	LANDSCAPE_API bool IsSharedProperty(const FName& InPropertyName) const;
+
+	// Returns true if the property is shared.
+	LANDSCAPE_API bool IsSharedProperty(const FProperty* InProperty) const;
+
+	// Returns true if the property is inherited.
+	LANDSCAPE_API bool IsPropertyInherited(const FProperty* InProperty) const;
+	
+	// Returns true if the property is overridable
+	LANDSCAPE_API bool IsPropertyOverridable(const FProperty* InProperty) const;
+	
+	// Returns true if the shared property is overridden by the object.
+	LANDSCAPE_API virtual bool IsSharedPropertyOverridden(const FString& InPropertyName) const { return false; }
+
+	// Modifies the override state of the property given as argument.
+	LANDSCAPE_API virtual void SetSharedPropertyOverride(const FString& InPropertyName, const bool bIsOverriden) { }
+#endif // WITH_EDITOR
+
 	// Get Landscape Material assigned to this Landscape
 	virtual UMaterialInterface* GetLandscapeMaterial(int8 InLODIndex = INDEX_NONE) const;
 
@@ -1128,6 +1159,9 @@ public:
 	// Changed Physical Material
 	LANDSCAPE_API void ChangedPhysMaterial();
 
+	// Converts shared properties overridden before explicit override, asking the user how to proceed.
+	LANDSCAPE_API void UpgradeSharedProperties(ALandscape* InParentLandscape);
+
 	// Assign only mismatching data and mark proxy package dirty
 	LANDSCAPE_API void FixupSharedData(ALandscape* Landscape, bool bMapCheck = false);
 
@@ -1397,9 +1431,14 @@ protected:
 	/** Initialize Layer with empty content if it hasn't been initialized yet. */
 	LANDSCAPE_API void InitializeLayerWithEmptyContent(const FGuid& InLayerGuid);
 
+	/** Fixup any internal representation for shared properties. Will be used when deprecating or renaming a property. */
+	LANDSCAPE_API virtual void FixupOverriddenSharedProperties() {}
+
 protected:
 	FLandscapeMaterialChangedDelegate LandscapeMaterialChangedDelegate;
 
+	// Used to know if the deprecation of shared properties modified before the enforcement system introduction has been performed on load.
+	bool bUpgradeSharedPropertiesPerformed = false;
 #endif // WITH_EDITOR
 private:
 
@@ -1438,7 +1477,12 @@ private:
 	static TArray<ALandscapeProxy*> LandscapeProxies;
 
 	FAsyncWorkMonitor AsyncWorkMonitor;
-#endif
+#endif // WITH_EDITORONLY_DATA
+
+#if WITH_EDITOR
+	static constexpr TCHAR const* LandscapeInheritedTag = TEXT("LandscapeInherited");
+	static constexpr TCHAR const* LandscapeOverridableTag = TEXT("LandscapeOverridable");
+#endif // WITH_EDITOR
 };
 
 
diff --git a/Engine/Source/Runtime/Landscape/Classes/LandscapeStreamingProxy.h b/Engine/Source/Runtime/Landscape/Classes/LandscapeStreamingProxy.h
--- a/Engine/Source/Runtime/Landscape/Classes/LandscapeStreamingProxy.h
+++ b/Engine/Source/Runtime/Landscape/Classes/LandscapeStreamingProxy.h
@@ -31,7 +31,10 @@ public:
 private:
 	UPROPERTY(EditAnywhere, Category = LandscapeProxy, Meta = (DisplayName = "Landscape Actor"))
 	TSoftObjectPtr<ALandscape> LandscapeActorRef;
-	
+
+	UPROPERTY()
+	TSet<FString> OverriddenSharedProperties;
+
 public:
 	//~ Begin UObject Interface
 #if WITH_EDITOR
@@ -57,4 +60,14 @@ public:
 
 	// Check input Landscape actor is match for this LandscapeProxy (by GUID)
 	bool IsValidLandscapeActor(ALandscape* Landscape);
+
+#if WITH_EDITOR
+	//~ Begin ALandscapeProxy Interface
+	virtual bool IsSharedPropertyOverridden(const FString& InPropertyName) const override;
+	virtual void SetSharedPropertyOverride(const FString& InPropertyName, const bool bIsOverridden) override;
+	//~ End ALandscapeProxy Interface
+
+protected:
+	virtual void FixupOverriddenSharedProperties() override;
+#endif // WITH_EDITOR
 };
diff --git a/Engine/Source/Runtime/Landscape/Private/Landscape.cpp b/Engine/Source/Runtime/Landscape/Private/Landscape.cpp
--- a/Engine/Source/Runtime/Landscape/Private/Landscape.cpp
+++ b/Engine/Source/Runtime/Landscape/Private/Landscape.cpp
@@ -3517,6 +3517,19 @@ void ALandscapeProxy::PostLoad()
 
 	PostLoadFixupLandscapeGuidsIfInstanced();
 
+#if WITH_EDITOR
+	FixupOverriddenSharedProperties();
+
+	ALandscape* LandscapeActor = GetLandscapeActor();
+
+	// Try to fixup shared properties if everything is ready for it as some postload process may depend on it.
+	if ((GetLandscapeInfo() != nullptr) && (LandscapeActor != nullptr) && (LandscapeActor != this))
+	{
+		const bool bMapCheck = true;
+		FixupSharedData(LandscapeActor, bMapCheck);
+	}
+#endif // WITH_EDITOR
+
 	// Temporary
 	if (ComponentSizeQuads == 0 && LandscapeComponents.Num() > 0)
 	{
@@ -3609,6 +3622,14 @@ void ALandscapeProxy::PostLoad()
 		// Only validate if uncooked and in the editor/commandlet mode (we cannot re-build material instance constants if this is not the case : see UMaterialInstance::CacheResourceShadersForRendering, which is only called if FApp::CanEverRender() returns true) 
 		if (!Comp->GetOutermost()->HasAnyPackageFlags(PKG_FilterEditorOnly) && (GIsEditor && FApp::CanEverRender()))
 		{
+			// MaterialInstance is different from the used LandscapeMaterial, we need to update the material as we cannot properly validate used combinations.
+			if (MaterialInstance->GetMaterial() != GetLandscapeMaterial())
+			{
+				Comp->UpdateMaterialInstances();
+				bFixedUpInvalidMaterialInstances = true;
+				continue;
+			}
+
 			if (Comp->ValidateCombinationMaterial(CombinationMaterialInstance))
 			{
 				MaterialInstanceConstantMap.Add(*ULandscapeComponent::GetLayerAllocationKey(Comp->GetWeightmapLayerAllocations(), CombinationMaterialInstance->Parent), CombinationMaterialInstance);
@@ -3799,6 +3820,139 @@ void ALandscapeProxy::GetSharedProperties(ALandscapeProxy* Landscape)
 #endif // WITH_EDITOR
 }
 
+#if WITH_EDITOR
+
+namespace UE::Landscape::Private
+{
+	static bool CopyProperty(FProperty* InProperty, UObject* InSourceObject, UObject* InDestinationObject)
+	{
+		void* SrcValuePtr = InProperty->ContainerPtrToValuePtr<void>(InSourceObject);
+		void* DestValuePtr = InProperty->ContainerPtrToValuePtr<void>(InDestinationObject);
+
+		if ((DestValuePtr == nullptr) || (SrcValuePtr == nullptr))
+		{
+			return false;
+		}
+
+		InProperty->CopyCompleteValue(DestValuePtr, SrcValuePtr);
+
+		return true;
+	}
+
+	static bool CopyPostEditPropertyByName(const TWeakObjectPtr<ALandscapeProxy>& InLandscapeProxy, const TWeakObjectPtr<ALandscape>& InParentLandscape, const FString& InPropertyName)
+	{
+		if (!InLandscapeProxy.IsValid() || !InParentLandscape.IsValid())
+		{
+			return false;
+		}
+
+		UClass* LandscapeProxyClass = InLandscapeProxy->GetClass();
+
+		if (LandscapeProxyClass == nullptr)
+		{
+			return false;
+		}
+
+		FProperty* PropertyToCopy = LandscapeProxyClass->FindPropertyByName(FName(InPropertyName));
+
+		if (PropertyToCopy == nullptr)
+		{
+			return false;
+		}
+
+		CopyProperty(PropertyToCopy, InParentLandscape.Get(), InLandscapeProxy.Get());
+
+		// Some properties may need additional processing (ex: LandscapeMaterial), notify the proxy of the change.
+		FPropertyChangedEvent PropertyChangedEvent(PropertyToCopy);
+		InLandscapeProxy->PostEditChangeProperty(PropertyChangedEvent);
+
+		InLandscapeProxy->Modify();
+
+		return true;
+	}
+} // namespace UE::Landscape::Private
+
+void ALandscapeProxy::CopySharedProperties(ALandscapeProxy* Landscape)
+{
+	for (TFieldIterator<FProperty> PropertyIterator(GetClass()); PropertyIterator; ++PropertyIterator)
+	{
+		FProperty* Property = *PropertyIterator;
+
+		if (Property == nullptr)
+		{
+			continue;
+		}
+
+		if (IsSharedProperty(Property))
+		{
+			UE::Landscape::Private::CopyProperty(Property, Landscape, this);
+		}
+	}
+}
+
+TArray<FName> ALandscapeProxy::SynchronizeSharedProperties(ALandscapeProxy* InLandscape)
+{
+	TArray<FName> SynchronizedProperties;
+
+	for (TFieldIterator<FProperty> PropertyIterator(GetClass()); PropertyIterator; ++PropertyIterator)
+	{
+		FProperty* Property = *PropertyIterator;
+
+		if (Property == nullptr)
+		{
+			continue;
+		}
+
+		if ((IsPropertyInherited(Property) ||
+			(IsPropertyOverridable(Property) && !IsSharedPropertyOverridden(Property->GetName()))) &&
+			!Property->Identical_InContainer(this, InLandscape))
+		{
+			SynchronizedProperties.Emplace(Property->GetFName());
+			UE::Landscape::Private::CopyProperty(Property, InLandscape, this);
+		}
+	}
+
+	if (!SynchronizedProperties.IsEmpty())
+	{
+		Modify();
+	}
+
+	return SynchronizedProperties;
+}
+
+bool ALandscapeProxy::IsSharedProperty(const FName& InPropertyName) const
+{
+	FProperty* Property = StaticClass()->FindPropertyByName(InPropertyName);
+
+	return IsSharedProperty(Property);
+}
+
+bool ALandscapeProxy::IsSharedProperty(const FProperty* InProperty) const
+{
+	return IsPropertyInherited(InProperty) || IsPropertyOverridable(InProperty);
+}
+
+bool ALandscapeProxy::IsPropertyInherited(const FProperty* InProperty) const
+{
+	if (InProperty == nullptr)
+	{
+		return false;
+	}
+
+	return InProperty->HasMetaData(LandscapeInheritedTag);
+}
+
+bool ALandscapeProxy::IsPropertyOverridable(const FProperty* InProperty) const
+{
+	if (InProperty == nullptr)
+	{
+		return false;
+	}
+
+	return InProperty->HasMetaData(LandscapeOverridableTag);
+}
+#endif // WITH_EDITOR
+
 UMaterialInterface* ALandscapeProxy::GetLandscapeMaterial(int8 InLODIndex) const
 {
 	if (InLODIndex != INDEX_NONE)
@@ -3867,115 +4021,173 @@ UMaterialInterface* ALandscapeStreamingProxy::GetLandscapeHoleMaterial() const
 
 #if WITH_EDITOR
 
-void ALandscapeProxy::FixupSharedData(ALandscape* Landscape, bool bMapCheck)
+bool ALandscapeStreamingProxy::IsSharedPropertyOverridden(const FString& InPropertyName) const
 {
-	if ((Landscape == nullptr) || (Landscape == this))
-	{
-		return;
-	}
+	return OverriddenSharedProperties.Contains(InPropertyName);
+}
+
+void ALandscapeStreamingProxy::SetSharedPropertyOverride(const FString& InPropertyName, const bool bIsOverridden)
+{
+	check(IsSharedProperty(FName(InPropertyName)));
 
-	ULandscapeInfo* LandscapeInfo = GetLandscapeInfo();
-	checkf(LandscapeInfo != nullptr, TEXT("FixupSharedData can only be called after the proxies are registered to ULandscapeInfo"));
-	ULandscapeSubsystem* LandscapeSubsystem = GetWorld() ? GetWorld()->GetSubsystem<ULandscapeSubsystem>() : nullptr;
-	checkf(LandscapeSubsystem != nullptr, TEXT("FixupSharedData can only be called when a subsystem is available"));
+	Modify();
 
-	bool bUpdated = false;
-	if (MaxLODLevel != Landscape->MaxLODLevel)
+	if (bIsOverridden)
 	{
-		MaxLODLevel = Landscape->MaxLODLevel;
-		bUpdated = true;
+		OverriddenSharedProperties.Add(InPropertyName);
 	}
-	
-	if (ComponentScreenSizeToUseSubSections != Landscape->ComponentScreenSizeToUseSubSections)
+	else
 	{
-		ComponentScreenSizeToUseSubSections = Landscape->ComponentScreenSizeToUseSubSections;
-		bUpdated = true;
-	}
+		TWeakObjectPtr<ALandscapeProxy> LandscapeProxy = this;
+		TWeakObjectPtr<ALandscape> ParentLandscape = GetLandscapeActor();
+		
+		if (!ParentLandscape.IsValid())
+		{
+			UE_LOG(LogLandscape, Warning, TEXT("Unable to retrieve the parent landscape's shared property value (ALandscapeStreamingProxy: %s, Property: %s). The proper value will be fixedup when reloading this proxy."),
+				   *GetFullName(), *InPropertyName);
+		}
+		else
+		{
+			UE::Landscape::Private::CopyPostEditPropertyByName(LandscapeProxy, ParentLandscape, InPropertyName);
+		}
 
-	if (LODDistributionSetting != Landscape->LODDistributionSetting)
-	{
-		LODDistributionSetting = Landscape->LODDistributionSetting;
-		bUpdated = true;
+		OverriddenSharedProperties.Remove(InPropertyName);
 	}
+}
 
-	if (LOD0DistributionSetting != Landscape->LOD0DistributionSetting)
-	{
-		LOD0DistributionSetting = Landscape->LOD0DistributionSetting;
-		bUpdated = true;
-	}
+void ALandscapeStreamingProxy::FixupOverriddenSharedProperties()
+{
+	const UClass* StreamingProxyClass = StaticClass();
 
-	if (LOD0ScreenSize != Landscape->LOD0ScreenSize)
+	for (const FString& PropertyName : OverriddenSharedProperties)
 	{
-		LOD0ScreenSize = Landscape->LOD0ScreenSize;
-		bUpdated = true;
+		const FProperty* Property = StreamingProxyClass->FindPropertyByName(FName(PropertyName));
+		checkf(Property != nullptr, TEXT("An overridden property is referenced but cannot be found. Please check this property hasn't been renamed or deprecated and/or provide the proper adapting mechanism."));
 	}
+}
 
-	if (LODGroupKey != Landscape->LODGroupKey)
-	{
-		LODGroupKey = Landscape->LODGroupKey;
-		bUpdated = true;
-	}
+void ALandscapeProxy::UpgradeSharedProperties(ALandscape* InParentLandscape)
+{
+	bool bOpenMapCheckWindow = false;
 
-	if (bUseCompressedHeightmapStorage)
+	for (TFieldIterator<FProperty> PropertyIterator(GetClass()); PropertyIterator; ++PropertyIterator)
 	{
-		bUseCompressedHeightmapStorage = Landscape->bUseCompressedHeightmapStorage;
-		bUpdated = true;
-	}
+		FProperty* Property = *PropertyIterator;
 
-	if (TargetDisplayOrder != Landscape->TargetDisplayOrder)
-	{
-		TargetDisplayOrder = Landscape->TargetDisplayOrder;
-		bUpdated = true;
+		if (Property == nullptr)
+		{
+			continue;
+		}
+
+		if ((IsPropertyInherited(Property) ||
+			(IsPropertyOverridable(Property) && !IsSharedPropertyOverridden(Property->GetName()))) &&
+			!Property->Identical_InContainer(this, InParentLandscape))
+		{
+			FFormatNamedArguments Arguments;
+			TWeakObjectPtr<ALandscapeProxy> LandscapeProxy = this;
+			TWeakObjectPtr<ALandscape> ParentLandscape = InParentLandscape;
+			const FString PropertyName = Property->GetName();
+
+			bOpenMapCheckWindow = true;
+
+			Arguments.Add(TEXT("Proxy"), FText::FromString(GetActorNameOrLabel()));
+			Arguments.Add(TEXT("Landscape"), FText::FromString(InParentLandscape->GetActorNameOrLabel()));
+			FMessageLog("MapCheck").Warning()
+				->AddToken(FUObjectToken::Create(this, FText::FromString(GetActorNameOrLabel())))
+				->AddToken(FTextToken::Create(FText::Format(LOCTEXT("MapCheck_Message_LandscapeProxy_UpgradeSharedProperties", "Contains a property ({0}) different from parent's landscape actor. Please select between "), FText::FromString(PropertyName))))
+				->AddToken(FActionToken::Create(LOCTEXT("MapCheck_OverrideProperty", "Override property"), LOCTEXT("MapCheck_OverrideProperty_Desc", "Keeping the current value and marking the property as overriding the parent landscape's value."),
+					FOnActionTokenExecuted::CreateLambda([LandscapeProxy, PropertyName]()
+					{
+						if (LandscapeProxy.IsValid())
+						{
+							LandscapeProxy->SetSharedPropertyOverride(PropertyName, true);
+						}
+					}),
+					/*bInSingleUse = */true))
+				->AddToken(FTextToken::Create(LOCTEXT("MapCheck_Message_LandscapeProxy_UpgradeSharedProperties_Or", " or ")))
+				->AddToken(FActionToken::Create(LOCTEXT("MapCheck_InheritProperty", "Inherit from parent landscape"), LOCTEXT("MapCheck_InheritProperty_Desc", "Copying the parent landscape's value for this property."),
+					FOnActionTokenExecuted::CreateLambda([LandscapeProxy, ParentLandscape, PropertyName]()
+					{
+						UE::Landscape::Private::CopyPostEditPropertyByName(LandscapeProxy, ParentLandscape, PropertyName);
+					}),
+					/*bInSingleUse = */true))
+				->AddToken(FMapErrorToken::Create(FMapErrors::LandscapeComponentPostLoad_Warning));
+		}
 	}
 
-	if (TargetDisplayOrderList != Landscape->TargetDisplayOrderList)
+	if (bOpenMapCheckWindow)
 	{
-		TargetDisplayOrderList = Landscape->TargetDisplayOrderList;
-		bUpdated = true;
+		// Show MapCheck window
+		FMessageLog("MapCheck").Open(EMessageSeverity::Warning);
 	}
+}
 
-	if (bEnableNanite != Landscape->IsNaniteEnabled())
+void ALandscapeProxy::FixupSharedData(ALandscape* Landscape, const bool bMapCheck)
+{
+	if ((Landscape == nullptr) || (Landscape == this))
 	{
-		bEnableNanite = Landscape->IsNaniteEnabled();
-		bUpdated = true;
+		return;
 	}
 
-	if (NaniteLODIndex != Landscape->GetNaniteLODIndex())
+	if (!bUpgradeSharedPropertiesPerformed && GetLinkerCustomVersion(FFortniteReleaseBranchCustomObjectVersion::GUID) < FFortniteReleaseBranchCustomObjectVersion::LandscapeSharedPropertiesEnforcement)
 	{
-		NaniteLODIndex = Landscape->GetNaniteLODIndex();
-		bUpdated = true;
+		UpgradeSharedProperties(Landscape);
+		bUpgradeSharedPropertiesPerformed = true;
 	}
+	else
+	{
+		ULandscapeInfo* LandscapeInfo = GetLandscapeInfo();
+		checkf(LandscapeInfo != nullptr, TEXT("FixupSharedData can only be called after the proxies are registered to ULandscapeInfo"));
+		ULandscapeSubsystem* LandscapeSubsystem = GetWorld() ? GetWorld()->GetSubsystem<ULandscapeSubsystem>() : nullptr;
+		checkf(LandscapeSubsystem != nullptr, TEXT("FixupSharedData can only be called when a subsystem is available"));
 
-	TSet<FGuid> LayerGuids;
-	Algo::Transform(Landscape->LandscapeLayers, LayerGuids, [](const FLandscapeLayer& Layer) { return Layer.Guid; });
-	bUpdated |= RemoveObsoleteLayers(LayerGuids);
+		TArray<FName> SynchronizedProperties = SynchronizeSharedProperties(Landscape);
+		bool bUpdated = !SynchronizedProperties.IsEmpty();
 
-	for (const FLandscapeLayer& Layer : Landscape->LandscapeLayers)
-	{
-		bUpdated |= AddLayer(Layer.Guid);
-	}
+		TSet<FGuid> LayerGuids;
+		Algo::Transform(Landscape->LandscapeLayers, LayerGuids, [](const FLandscapeLayer& Layer) { return Layer.Guid; });
+		bUpdated |= RemoveObsoleteLayers(LayerGuids);
 
-	if (bUpdated)
-	{
-		// Force resave the proxy through the modified landscape system, so that the user can then use the Save Modified Landscapes menu and therefore manually trigger the re-save of all modified proxies. * /
-		bool bNeedsManualResave = LandscapeInfo->MarkObjectDirty(/*InObject = */this, /*bInForceResave = */true);
+		for (const FLandscapeLayer& Layer : Landscape->LandscapeLayers)
+		{
+			bUpdated |= AddLayer(Layer.Guid);
+		}
 
-		if (bMapCheck && bNeedsManualResave)
+		if (bUpdated)
 		{
-			FFormatNamedArguments Arguments;
-			Arguments.Add(TEXT("Proxy"), FText::FromString(GetActorNameOrLabel()));
-			Arguments.Add(TEXT("Landscape"), FText::FromString(Landscape->GetActorNameOrLabel()));
-			FMessageLog("MapCheck").Warning()
-				->AddToken(FUObjectToken::Create(this, FText::FromString(GetActorNameOrLabel())))
-				->AddToken(FTextToken::Create(LOCTEXT("MapCheck_Message_LandscapeProxy_FixupSharedData", "had some shared properties not in sync with its parent landscape actor. This has been fixed but the proxy needs to be saved in order to ensure cooking behaves as expected. ")))
-				->AddToken(FActionToken::Create(LOCTEXT("MapCheck_SaveFixedUpData", "Save Modified Landscapes"), LOCTEXT("MapCheck_SaveFixedUpData_Desc", "Saves the modified landscape proxy actors"), 
-					FOnActionTokenExecuted::CreateUObject(LandscapeSubsystem, &ULandscapeSubsystem::SaveModifiedLandscapes), 
-					FCanExecuteActionToken::CreateUObject(LandscapeSubsystem, &ULandscapeSubsystem::HasModifiedLandscapes),
-					/*bInSingleUse = */false))
-				->AddToken(FMapErrorToken::Create(FMapErrors::LandscapeComponentPostLoad_Warning));
+			// Force resave the proxy through the modified landscape system, so that the user can then use the Save Modified Landscapes menu and therefore manually trigger the re-save of all modified proxies. * /
+			bool bNeedsManualResave = LandscapeInfo->MarkObjectDirty(/*InObject = */this, /*bInForceResave = */true);
 
-			// Show MapCheck window
-			FMessageLog("MapCheck").Open(EMessageSeverity::Warning);
+			if (bMapCheck && bNeedsManualResave)
+			{
+				TStringBuilder<1024> SynchronizedPropertiesStringBuilder;
+
+				for (const FName& SynchronizedProperty : SynchronizedProperties)
+				{
+					if (SynchronizedPropertiesStringBuilder.Len() > 0)
+					{
+						SynchronizedPropertiesStringBuilder.Append(TEXT(", "));
+					}
+
+					SynchronizedPropertiesStringBuilder.Append(SynchronizedProperty.ToString());
+				}
+
+				FFormatNamedArguments Arguments;
+				Arguments.Add(TEXT("Proxy"), FText::FromString(GetActorNameOrLabel()));
+				Arguments.Add(TEXT("Landscape"), FText::FromString(Landscape->GetActorNameOrLabel()));
+				FMessageLog("MapCheck").Warning()
+					->AddToken(FUObjectToken::Create(this, FText::FromString(GetActorNameOrLabel())))
+					->AddToken(FTextToken::Create(LOCTEXT("MapCheck_Message_LandscapeProxy_FixupSharedData", "had some shared properties not in sync with its parent landscape actor. This has been fixed but the proxy needs to be saved in order to ensure cooking behaves as expected. ")))
+					->AddToken(FActionToken::Create(LOCTEXT("MapCheck_SaveFixedUpData", "Save Modified Landscapes"), LOCTEXT("MapCheck_SaveFixedUpData_Desc", "Saves the modified landscape proxy actors"),
+						FOnActionTokenExecuted::CreateUObject(LandscapeSubsystem, &ULandscapeSubsystem::SaveModifiedLandscapes),
+						FCanExecuteActionToken::CreateUObject(LandscapeSubsystem, &ULandscapeSubsystem::HasModifiedLandscapes),
+						/*bInSingleUse = */false))
+					->AddToken(FTextToken::Create(FText::Format(LOCTEXT("MapCheck_Message_LandscapeProxy_FixupSharedData_SharedProperties", "The following properties were synchronized: {0}."), FText::FromString(SynchronizedPropertiesStringBuilder.ToString()))))
+					->AddToken(FMapErrorToken::Create(FMapErrors::LandscapeComponentPostLoad_Warning));
+
+				// Show MapCheck window
+				FMessageLog("MapCheck").Open(EMessageSeverity::Warning);
+			}
 		}
 	}
 }
@@ -4705,10 +4917,10 @@ void ULandscapeInfo::RegisterActor(ALandscapeProxy* Proxy, bool bMapCheck, bool
 		{
 			StreamingProxy->SetLandscapeActor(LandscapeActor.Get());
 			StreamingProxy->LODGroupKey = LandscapeActor.Get()->LODGroupKey;
-		}
 #if WITH_EDITOR
-		StreamingProxy->FixupSharedData(LandscapeActor.Get(), bMapCheck);
+			StreamingProxy->FixupSharedData(LandscapeActor.Get(), bMapCheck);
 #endif // WITH_EDITOR
+		}
 	}
 
 #if WITH_EDITOR
diff --git a/Engine/Source/Runtime/Landscape/Private/LandscapeConfigHelper.cpp b/Engine/Source/Runtime/Landscape/Private/LandscapeConfigHelper.cpp
--- a/Engine/Source/Runtime/Landscape/Private/LandscapeConfigHelper.cpp
+++ b/Engine/Source/Runtime/Landscape/Private/LandscapeConfigHelper.cpp
@@ -70,7 +70,7 @@ ALandscapeProxy* FLandscapeConfigHelper::FindOrAddLandscapeStreamingProxy(UActor
 
 		ALandscapeProxy* LandscapeProxy = CastChecked<ALandscapeProxy>(PartitionActor);
 		// copy shared properties to this new proxy
-		LandscapeProxy->GetSharedProperties(Landscape);
+		LandscapeProxy->SynchronizeSharedProperties(Landscape);
 		const FVector ProxyLocation = Landscape->GetActorLocation() + FVector(CellLocation.X * Landscape->GetActorRelativeScale3D().X, CellLocation.Y * Landscape->GetActorRelativeScale3D().Y, 0.0f);
 
 		LandscapeProxy->CreateLandscapeInfo();
@@ -486,7 +486,7 @@ ULandscapeInfo* FLandscapeConfigHelper::ChangeConfiguration(ULandscapeInfo* InLa
 	FActorSpawnParameters SpawnParams;
 	ALandscape* NewLandscape = World->SpawnActor<ALandscape>(OldLandscape->GetActorLocation() + ActorOffset, OldLandscape->GetActorRotation(), SpawnParams);
 	
-	NewLandscape->GetSharedProperties(OldLandscape);
+	NewLandscape->CopySharedProperties(OldLandscape);
 	NewLandscape->SetLandscapeGuid(FGuid::NewGuid());
 	NewLandscape->GridSize = InNewConfig.GetGridSizeQuads();
 	NewLandscape->ComponentSizeQuads = InNewConfig.GetComponentSizeQuads();
diff --git a/Engine/Source/Runtime/Landscape/Private/LandscapeEdit.cpp b/Engine/Source/Runtime/Landscape/Private/LandscapeEdit.cpp
--- a/Engine/Source/Runtime/Landscape/Private/LandscapeEdit.cpp
+++ b/Engine/Source/Runtime/Landscape/Private/LandscapeEdit.cpp
@@ -4963,7 +4963,7 @@ ALandscapeProxy* ULandscapeInfo::MoveComponentsToLevel(const TArray<ULandscapeCo
 		SpawnParams.OverrideLevel = TargetLevel;
 		LandscapeProxy = TargetLevel->GetWorld()->SpawnActor<ALandscapeStreamingProxy>(SpawnParams);
 
-		LandscapeProxy->GetSharedProperties(Landscape);
+		LandscapeProxy->SynchronizeSharedProperties(Landscape);
 		LandscapeProxy->CreateLandscapeInfo();
 		LandscapeProxy->SetActorLabel(LandscapeProxy->GetName());
 		bSetPositionAndOffset = true;
@@ -5956,7 +5956,7 @@ void ALandscape::PostEditChangeProperty(FPropertyChangedEvent& PropertyChangedEv
 	const FName PropertyName = PropertyChangedEvent.Property ? PropertyChangedEvent.Property->GetFName() : NAME_None;
 	const FName MemberPropertyName = PropertyChangedEvent.MemberProperty ? PropertyChangedEvent.MemberProperty->GetFName() : NAME_None;
 
-	bool ChangedMaterial = false;
+	bool bMaterialChanged = false;
 	bool bNeedsRecalcBoundingBox = false;
 	bool bChangedLighting = false;
 	bool bPropagateToProxies = false;
@@ -6005,7 +6005,7 @@ void ALandscape::PostEditChangeProperty(FPropertyChangedEvent& PropertyChangedEv
 			FMaterialUpdateContext MaterialUpdateContext;
 			Info->UpdateLayerInfoMap(/*this*/);
 
-			ChangedMaterial = true;
+			bMaterialChanged = true;
 
 			// Clear the parents out of combination material instances
 			for (const auto& MICPair : MaterialInstanceConstantMap)
@@ -6121,11 +6121,15 @@ void ALandscape::PostEditChangeProperty(FPropertyChangedEvent& PropertyChangedEv
 		bPropagateToProxies = true;
 		bMarkAllLandscapeRenderStateDirty = true;
 	}
+	else if (GIsEditor && IsSharedProperty(PropertyName))
+	{
+		bPropagateToProxies = true;
+	}
 
 	// Must do this *after* clamping values
 	Super::PostEditChangeProperty(PropertyChangedEvent);
 
-	bPropagateToProxies = bPropagateToProxies || bNeedsRecalcBoundingBox || bChangedLighting || bNaniteToggled;
+	bPropagateToProxies = bPropagateToProxies || bNeedsRecalcBoundingBox || bChangedLighting || bNaniteToggled || bMaterialChanged;
 
 	if (Info != nullptr)
 	{
@@ -6136,8 +6140,7 @@ void ALandscape::PostEditChangeProperty(FPropertyChangedEvent& PropertyChangedEv
 			{
 				if (ALandscapeProxy* Proxy = ProxyPtr.Get())
 				{
-					Proxy->GetSharedProperties(this);
-					
+					Proxy->SynchronizeSharedProperties(this);
 					Proxy->PostEditChangeProperty(PropertyChangedEvent);
 				}
 			}
@@ -6155,7 +6158,7 @@ void ALandscape::PostEditChangeProperty(FPropertyChangedEvent& PropertyChangedEv
 			LandscapeEdit.RecalculateNormals();
 		}
 
-		if (bNeedsRecalcBoundingBox || ChangedMaterial || bChangedLighting)
+		if (bNeedsRecalcBoundingBox || bMaterialChanged || bChangedLighting)
 		{
 			// We cannot iterate the XYtoComponentMap directly because reregistering components modifies the array.
 			TArray<ULandscapeComponent*> AllComponents;
@@ -6179,7 +6182,7 @@ void ALandscape::PostEditChangeProperty(FPropertyChangedEvent& PropertyChangedEv
 				}
 			}
 
-			if (ChangedMaterial)
+			if (bMaterialChanged)
 			{
 				UpdateAllComponentMaterialInstances();
 
@@ -6214,7 +6217,7 @@ void ALandscape::PostEditChangeProperty(FPropertyChangedEvent& PropertyChangedEv
 		}
 
 		// Must be done after the AActor::PostEditChange as we depend on the relinking of the landscapeInfo->LandscapeActor
-		if (ChangedMaterial)
+		if (bMaterialChanged)
 		{
 			LandscapeMaterialChangedDelegate.Broadcast();
 		}
