Landscape : Create textures with no mipchhain if used in as a layer height or weight

Height or weight maps used in layers are 2d composited and don't require mipmaps at all. By removing the mip chain we can save 30% on the majority of landscape textures used in Editor and hopefully making interaction with the DDC more efficient.

Tested all combinations of [WP |  Non-WP] x [ edit layers | no edit layers ]

diff --git a/Engine/Source/Runtime/Landscape/Classes/LandscapeProxy.h b/Engine/Source/Runtime/Landscape/Classes/LandscapeProxy.h
--- a/Engine/Source/Runtime/Landscape/Classes/LandscapeProxy.h
+++ b/Engine/Source/Runtime/Landscape/Classes/LandscapeProxy.h
@@ -1346,12 +1346,11 @@ public:
 	LANDSCAPE_API FIntRect GetBoundingRect() const;
 
 	/** Creates a Texture2D for use by this landscape proxy or one of it's components. If OptionalOverrideOuter is not specified, the proxy is used. */
-	LANDSCAPE_API UTexture2D* CreateLandscapeTexture(int32 InSizeX, int32 InSizeY, TextureGroup InLODGroup, ETextureSourceFormat InFormat, UObject* OptionalOverrideOuter = nullptr, bool bCompress = false) const;
+	LANDSCAPE_API UTexture2D* CreateLandscapeTexture(int32 InSizeX, int32 InSizeY, TextureGroup InLODGroup, ETextureSourceFormat InFormat, UObject* OptionalOverrideOuter = nullptr, bool bCompress = false, bool bMipChain = true) const;
 
 	/** Creates a Texture2D for use by this landscape proxy or one of it's components for tools .*/ 
 	LANDSCAPE_API UTexture2D* CreateLandscapeToolTexture(int32 InSizeX, int32 InSizeY, TextureGroup InLODGroup, ETextureSourceFormat InFormat) const;
 
-
 	/** Creates a LandscapeWeightMapUsage object outered to this proxy. */
 	LANDSCAPE_API ULandscapeWeightmapUsage* CreateWeightmapUsage();
 
diff --git a/Engine/Source/Runtime/Landscape/Private/Landscape.cpp b/Engine/Source/Runtime/Landscape/Private/Landscape.cpp
--- a/Engine/Source/Runtime/Landscape/Private/Landscape.cpp
+++ b/Engine/Source/Runtime/Landscape/Private/Landscape.cpp
@@ -1142,6 +1142,36 @@ void ULandscapeComponent::PostLoad()
 			UpdateMaterialInstances();
 		}
 	}
+
+	auto DropMipChain = [](UTexture2D* InTexture) 
+	{
+		if (InTexture->Source.GetNumMips() <= 1)
+		{
+			return;
+		}
+
+		uint64 NumBytes = InTexture->Source.GetSizeX() * InTexture->Source.GetSizeY() * FTextureSource::GetBytesPerPixel(InTexture->Source.GetFormat());
+		uint8* MipDataCopy = static_cast<uint8*>(FMemory::Malloc(NumBytes));
+		uint8 const* MipData = InTexture->Source.LockMipReadOnly(0);
+		FMemory::Memcpy(MipDataCopy, MipData, NumBytes);
+
+		InTexture->Source.Init(InTexture->Source.GetSizeX(), InTexture->Source.GetSizeY(), 1, 1, InTexture->Source.GetFormat(), MipDataCopy);
+		InTexture->UpdateResource();
+
+		FMemory::Free(MipDataCopy);
+
+	};
+
+	// Remove Non zero mip levels found in layer textures
+	for (auto& LayerIt : LayersData)
+	{
+		DropMipChain(LayerIt.Value.HeightmapData.Texture);
+		for (int32 i = 0; i < LayerIt.Value.WeightmapData.Textures.Num(); ++i)
+		{
+			DropMipChain(LayerIt.Value.WeightmapData.Textures[i]);
+		}
+	}
+	
 #endif
 
 	auto ReparentObject = [this](UObject* Object)
@@ -1203,7 +1233,7 @@ void ULandscapeComponent::PostLoad()
 		}
 	}
 
-	for (UMaterialInstance* MobileCombinationMaterialInstance : MobileCombinationMaterialInstances)
+ 	for (UMaterialInstance* MobileCombinationMaterialInstance : MobileCombinationMaterialInstances)
 	{
 		while (ReparentObject(MobileCombinationMaterialInstance))
 		{
@@ -2172,11 +2202,10 @@ void ULandscapeComponent::AddDefaultLayerData(const FGuid& InLayerGuid, const TA
 
 		if (LayerHeightmap == nullptr)
 		{
-			UTexture2D* NewLayerHeightmap = GetLandscapeProxy()->CreateLandscapeTexture(ComponentHeightmap->Source.GetSizeX(), ComponentHeightmap->Source.GetSizeY(), TEXTUREGROUP_Terrain_Heightmap, ComponentHeightmap->Source.GetFormat());
+			// No mipchain required as these layer weight maps are used in layer compositing to generate a final set of weight maps to be used for rendering
+			UTexture2D* NewLayerHeightmap = GetLandscapeProxy()->CreateLandscapeTexture(ComponentHeightmap->Source.GetSizeX(), ComponentHeightmap->Source.GetSizeY(), TEXTUREGROUP_Terrain_Heightmap, ComponentHeightmap->Source.GetFormat(), /* OptionalOverrideOuter = */ nullptr, /* bCompress = */ false, /* bMipChain = */ false);
 			LayerHeightmap = &InOutCreatedHeightmapTextures.Add(ComponentHeightmap, NewLayerHeightmap);
 
-			ULandscapeComponent::CreateEmptyTextureMips(NewLayerHeightmap, true);
-
 			// Init Mip0 to be at 32768 which is equal to "0"
 			TArrayView<FColor> Mip0Data((FColor*)NewLayerHeightmap->Source.LockMip(0), NewLayerHeightmap->Source.GetSizeX() * NewLayerHeightmap->Source.GetSizeY());
 
diff --git a/Engine/Source/Runtime/Landscape/Private/LandscapeEdit.cpp b/Engine/Source/Runtime/Landscape/Private/LandscapeEdit.cpp
--- a/Engine/Source/Runtime/Landscape/Private/LandscapeEdit.cpp
+++ b/Engine/Source/Runtime/Landscape/Private/LandscapeEdit.cpp
@@ -5319,7 +5319,8 @@ void ALandscape::SplitHeightmap(ULandscapeComponent* Comp, ALandscapeProxy* Targ
 				// Restore new heightmap scale/bias
 				Comp->HeightmapScaleBias = NewHeightmapScaleBias;
 				{
-					UTexture2D* LayerHeightmapTexture = DstProxy->CreateLandscapeTexture(HeightmapSizeU, HeightmapSizeV, TEXTUREGROUP_Terrain_Heightmap, TSF_BGRA8);
+					// no mipchain required as these layer weight maps are used in layer compositing to generate a final set of weight maps to be used for rendering
+					UTexture2D* LayerHeightmapTexture = DstProxy->CreateLandscapeTexture(HeightmapSizeU, HeightmapSizeV, TEXTUREGROUP_Terrain_Heightmap, TSF_BGRA8, /* OptionalOverrideOuter = */ nullptr, /* bCompress = */ false, /* bMipChain = */ false);
 					ULandscapeComponent::CreateEmptyTextureMips(LayerHeightmapTexture, true);
 					LayerHeightmapTexture->PostEditChange();
 					// Set Layer heightmap texture
@@ -6431,6 +6432,8 @@ void ULandscapeComponent::ReallocateWeightmapsInternal(FLandscapeEditDataInterfa
 	int32 NeededNewChannels = 0;
 	ALandscapeProxy* TargetProxy = InTargetProxy ? InTargetProxy : GetLandscapeProxy();
 	
+	const bool bIsFinalWeightmap = !InEditLayerGuid.IsValid();
+
 	if (InEditLayerGuid.IsValid())
 	{
 		check(TargetProxy->HasLayersContent());
@@ -6604,10 +6607,13 @@ void ULandscapeComponent::ReallocateWeightmapsInternal(FLandscapeEditDataInterfa
 			int32 WeightmapSize = (SubsectionSizeQuads + 1) * NumSubsections;
 
 			// We need a new weightmap texture
-			CurrentWeightmapTexture = TargetProxy->CreateLandscapeTexture(WeightmapSize, WeightmapSize, TEXTUREGROUP_Terrain_Weightmap, TSF_BGRA8);
+			CurrentWeightmapTexture = TargetProxy->CreateLandscapeTexture(WeightmapSize, WeightmapSize, TEXTUREGROUP_Terrain_Weightmap, TSF_BGRA8, nullptr, false, bIsFinalWeightmap); //dmb-nomips
 
 			// Alloc dummy mips
-			CreateEmptyTextureMips(CurrentWeightmapTexture, true);
+			if (bIsFinalWeightmap)
+			{
+				CreateEmptyTextureMips(CurrentWeightmapTexture, true);
+			}
 
 			CurrentWeightmapTexture->PostEditChange();
 
@@ -6669,7 +6675,7 @@ void ULandscapeComponent::ReallocateWeightmapsInternal(FLandscapeEditDataInterfa
 		}
 	}
 
-	if (DataInterface)
+	if (DataInterface && bIsFinalWeightmap)
 	{
 		// Update the mipmaps for the textures we edited
 		for (int32 Idx = 0; Idx < NewWeightmapTextures.Num(); Idx++)
@@ -7236,15 +7242,23 @@ FName ALandscapeProxy::GenerateUniqueLandscapeTextureName(UObject* InOuter, Text
 	return MakeUniqueObjectName(InOuter, UTexture2D::StaticClass(), BaseName);
 }
 
-UTexture2D* ALandscapeProxy::CreateLandscapeTexture(int32 InSizeX, int32 InSizeY, TextureGroup InLODGroup, ETextureSourceFormat InFormat, UObject* OptionalOverrideOuter, bool bCompress) const
+UTexture2D* ALandscapeProxy::CreateLandscapeTexture(int32 InSizeX, int32 InSizeY, TextureGroup InLODGroup, ETextureSourceFormat InFormat, UObject* OptionalOverrideOuter, bool bCompress, bool bMipChain) const
 {
 	UObject* TexOuter = OptionalOverrideOuter ? OptionalOverrideOuter : const_cast<ALandscapeProxy*>(this);
 	UTexture2D* NewTexture = NewObject<UTexture2D>(TexOuter, GenerateUniqueLandscapeTextureName(TexOuter, InLODGroup));
-	NewTexture->Source.Init2DWithMipChain(InSizeX, InSizeY, InFormat);
+	if (bMipChain)
+	{
+		NewTexture->Source.Init2DWithMipChain(InSizeX, InSizeY, InFormat);
+	}
+	else
+	{
+		NewTexture->Source.Init(InSizeX, InSizeY, 1, 1, InFormat);
+	}
+	
 	NewTexture->SRGB = false;
 	NewTexture->CompressionNone = !bCompress;
 	NewTexture->CompressionQuality = TCQ_Highest;
-	NewTexture->MipGenSettings = TMGS_LeaveExistingMips;
+	NewTexture->MipGenSettings = bMipChain ? TMGS_LeaveExistingMips : TMGS_NoMipmaps;
 	NewTexture->AddressX = TA_Clamp;
 	NewTexture->AddressY = TA_Clamp;
 	NewTexture->LODGroup = InLODGroup;
diff --git a/Engine/Source/Runtime/Landscape/Private/LandscapeEditInterface.cpp b/Engine/Source/Runtime/Landscape/Private/LandscapeEditInterface.cpp
--- a/Engine/Source/Runtime/Landscape/Private/LandscapeEditInterface.cpp
+++ b/Engine/Source/Runtime/Landscape/Private/LandscapeEditInterface.cpp
@@ -1831,6 +1831,8 @@ void ULandscapeComponent::FillLayer(ULandscapeLayerInfoObject* LayerInfo, FLands
 
 	FGuid EditLayerGuid = LandscapeEdit.GetEditLayer();
 
+	const bool bIsFinalWeightmap = !EditLayerGuid.IsValid();
+
 	ALandscapeProxy* Proxy = Component->GetLandscapeProxy();
 	Component->Modify(LandscapeEdit.GetShouldDirtyPackage());
 	Proxy->Modify(LandscapeEdit.GetShouldDirtyPackage());
@@ -1954,23 +1956,26 @@ void ULandscapeComponent::FillLayer(ULandscapeLayerInfoObject* LayerInfo, FLands
 	// e.g. if you start with 4 blended layers (in one texture) and a non-blended layer (in a 2nd), and fill one weight-blended layer (deleting the other three)
 	// this can also happen with normal painting I believe
 
-	// update mips
-	for (int32 TextureIdx = 0; TextureIdx < ComponentWeightmapTextures.Num(); ++TextureIdx)
+	// update mips if this texture isn't used in the layer compositing 
+	if (bIsFinalWeightmap)
 	{
-		UTexture2D* WeightmapTexture = ComponentWeightmapTextures[TextureIdx];
-		FLandscapeTextureDataInfo* WeightmapDataInfo = LandscapeEdit.GetTextureDataInfo(WeightmapTexture);
-
-		const int32 NumMips = WeightmapTexture->Source.GetNumMips();
-		TArray<FColor*> WeightmapTextureMipData;
-		WeightmapTextureMipData.AddUninitialized(NumMips);
-		for (int32 MipIdx = 0; MipIdx < NumMips; MipIdx++)
+		for (int32 TextureIdx = 0; TextureIdx < ComponentWeightmapTextures.Num(); ++TextureIdx)
 		{
-			WeightmapTextureMipData[MipIdx] = (FColor*)WeightmapDataInfo->GetMipData(MipIdx);
-		}
+			UTexture2D* WeightmapTexture = ComponentWeightmapTextures[TextureIdx];
+			FLandscapeTextureDataInfo* WeightmapDataInfo = LandscapeEdit.GetTextureDataInfo(WeightmapTexture);
 
-		ULandscapeComponent::UpdateWeightmapMips(NumSubsections, SubsectionSizeQuads, WeightmapTexture, WeightmapTextureMipData, 0, 0, MAX_int32, MAX_int32, WeightmapDataInfo);
+			const int32 NumMips = WeightmapTexture->Source.GetNumMips();
+			TArray<FColor*> WeightmapTextureMipData;
+			WeightmapTextureMipData.AddUninitialized(NumMips);
+			for (int32 MipIdx = 0; MipIdx < NumMips; MipIdx++)
+			{
+				WeightmapTextureMipData[MipIdx] = (FColor*)WeightmapDataInfo->GetMipData(MipIdx);
+			}
 
-		WeightmapDataInfo->AddMipUpdateRegion(0, 0, 0, WeightmapTexture->Source.GetSizeX() - 1, WeightmapTexture->Source.GetSizeY() - 1);
+			ULandscapeComponent::UpdateWeightmapMips(NumSubsections, SubsectionSizeQuads, WeightmapTexture, WeightmapTextureMipData, 0, 0, MAX_int32, MAX_int32, WeightmapDataInfo);
+
+			WeightmapDataInfo->AddMipUpdateRegion(0, 0, 0, WeightmapTexture->Source.GetSizeX() - 1, WeightmapTexture->Source.GetSizeY() - 1);
+		}
 	}
 
 	// Update the shaders for this component
