#UE Assume all pre-existing landscape proxy properties that differ from their parent are intended to be overridden, so mark them as overridden

diff --git a/Engine/Source/Runtime/Landscape/Private/Landscape.cpp b/Engine/Source/Runtime/Landscape/Private/Landscape.cpp
--- a/Engine/Source/Runtime/Landscape/Private/Landscape.cpp
+++ b/Engine/Source/Runtime/Landscape/Private/Landscape.cpp
@@ -4097,8 +4097,6 @@ void ALandscapeStreamingProxy::FixupOverriddenSharedProperties()
 
 void ALandscapeProxy::UpgradeSharedProperties(ALandscape* InParentLandscape)
 {
-	bool bOpenMapCheckWindow = false;
-
 	for (TFieldIterator<FProperty> PropertyIterator(GetClass()); PropertyIterator; ++PropertyIterator)
 	{
 		FProperty* Property = *PropertyIterator;
@@ -4112,43 +4110,9 @@ void ALandscapeProxy::UpgradeSharedProperties(ALandscape* InParentLandscape)
 			(IsPropertyOverridable(Property) && !IsSharedPropertyOverridden(Property->GetName()))) &&
 			!Property->Identical_InContainer(this, InParentLandscape))
 		{
-			FFormatNamedArguments Arguments;
-			TWeakObjectPtr<ALandscapeProxy> LandscapeProxy = this;
-			TWeakObjectPtr<ALandscape> ParentLandscape = InParentLandscape;
-			const FString PropertyName = Property->GetName();
-
-			bOpenMapCheckWindow = true;
-
-			Arguments.Add(TEXT("Proxy"), FText::FromString(GetActorNameOrLabel()));
-			Arguments.Add(TEXT("Landscape"), FText::FromString(InParentLandscape->GetActorNameOrLabel()));
-			FMessageLog("MapCheck").Warning()
-				->AddToken(FUObjectToken::Create(this, FText::FromString(GetActorNameOrLabel())))
-				->AddToken(FTextToken::Create(FText::Format(LOCTEXT("MapCheck_Message_LandscapeProxy_UpgradeSharedProperties", "Contains a property ({0}) different from parent's landscape actor. Please select between "), FText::FromString(PropertyName))))
-				->AddToken(FActionToken::Create(LOCTEXT("MapCheck_OverrideProperty", "Override property"), LOCTEXT("MapCheck_OverrideProperty_Desc", "Keeping the current value and marking the property as overriding the parent landscape's value."),
-					FOnActionTokenExecuted::CreateLambda([LandscapeProxy, PropertyName]()
-					{
-						if (LandscapeProxy.IsValid())
-						{
-							LandscapeProxy->SetSharedPropertyOverride(PropertyName, true);
-						}
-					}),
-					/*bInSingleUse = */true))
-				->AddToken(FTextToken::Create(LOCTEXT("MapCheck_Message_LandscapeProxy_UpgradeSharedProperties_Or", " or ")))
-				->AddToken(FActionToken::Create(LOCTEXT("MapCheck_InheritProperty", "Inherit from parent landscape"), LOCTEXT("MapCheck_InheritProperty_Desc", "Copying the parent landscape's value for this property."),
-					FOnActionTokenExecuted::CreateLambda([LandscapeProxy, ParentLandscape, PropertyName]()
-					{
-						UE::Landscape::Private::CopyPostEditPropertyByName(LandscapeProxy, ParentLandscape, PropertyName);
-					}),
-					/*bInSingleUse = */true))
-				->AddToken(FMapErrorToken::Create(FMapErrors::LandscapeComponentPostLoad_Warning));
+			SetSharedPropertyOverride(Property->GetName(), true);
 		}
 	}
-
-	if (bOpenMapCheckWindow)
-	{
-		// Show MapCheck window
-		FMessageLog("MapCheck").Open(EMessageSeverity::Warning);
-	}
 }
 
 void ALandscapeProxy::FixupSharedData(ALandscape* Landscape, const bool bMapCheck)
