Chaos - fix landscape collision lifetime bug

If a landscape object was destroyed and GC'd right after being unregistered, the collision geometry could be destroyed while the body was still registered with physics.

Temporary fix, pending proper reference counting for implicit objects.

diff --git a/Engine/Source/Runtime/Landscape/Classes/LandscapeHeightfieldCollisionComponent.h b/Engine/Source/Runtime/Landscape/Classes/LandscapeHeightfieldCollisionComponent.h
--- a/Engine/Source/Runtime/Landscape/Classes/LandscapeHeightfieldCollisionComponent.h
+++ b/Engine/Source/Runtime/Landscape/Classes/LandscapeHeightfieldCollisionComponent.h
@@ -298,6 +298,11 @@ public:
 	/** Recreate heightfield and restart physics */
 	LANDSCAPE_API virtual bool RecreateCollision();
 
+private:
+	// @todo(chaos): remove when implicit objects are ref counted
+	void DeferredDestroyCollision(const TRefCountPtr<FHeightfieldGeometryRef>& HeightfieldRefLifetimeExtender);
+public:
+
 #if WITH_EDITORONLY_DATA
 	// Called from editor code to manage foliage instances on landscape.
 	LANDSCAPE_API void SnapFoliageInstances(const FBox& InInstanceBox);
diff --git a/Engine/Source/Runtime/Landscape/Private/LandscapeCollision.cpp b/Engine/Source/Runtime/Landscape/Private/LandscapeCollision.cpp
--- a/Engine/Source/Runtime/Landscape/Private/LandscapeCollision.cpp
+++ b/Engine/Source/Runtime/Landscape/Private/LandscapeCollision.cpp
@@ -14,6 +14,7 @@
 #include "EngineDefines.h"
 #include "Engine/EngineTypes.h"
 #include "Components/SceneComponent.h"
+#include "GameFramework/Actor.h"
 #include "AI/Navigation/NavigationTypes.h"
 #include "Misc/SecureHash.h"
 #include "CollisionQueryParams.h"
@@ -253,6 +254,18 @@ void ULandscapeHeightfieldCollisionComponent::OnUnregister()
 {
 	Super::OnUnregister();
 
+	// The physics object was destroyed in Super::OnUnregister. However we must
+	// extend the lifetime of the collision until the enqueued destroy command
+	// if processed on the physics thread, otherwise we may get Destroyed before
+	// that happens and the collision geometry will be destroyed with us, leaving 
+	// a dangling pointer in physics.
+	// NOTE: we don't destroy collision in DestroyPhysicsState because we may
+	// change the physics state without generating new collision geometry.
+	DeferredDestroyCollision(HeightfieldRef);
+	HeightfieldRef = nullptr;
+	HeightfieldGuid = FGuid();
+	CachedHeightFieldSamples.Empty();
+
 	if (GetLandscapeProxy())
 	{
 		// AActor::GetWorld checks for Unreachable and BeginDestroyed
@@ -341,6 +354,13 @@ void ULandscapeHeightfieldCollisionComponent::OnCreatePhysicsState()
 				Params.bQueryOnly = false;
 				Params.bStatic = true;
 				Params.Scene = GetWorld()->GetPhysicsScene();
+
+#if USE_BODYINSTANCE_DEBUG_NAMES
+				const FString DebugName = (GetOwner() != nullptr) ? FString::Printf(TEXT("%s:%s"), *GetOwner()->GetFullName(), *GetName()) : *GetName();
+				BodyInstance.CharDebugName = MakeShareable(new TArray<ANSICHAR>(StringToArray<ANSICHAR>(*DebugName, DebugName.Len() + 1)));
+				Params.DebugName = BodyInstance.CharDebugName.IsValid() ? BodyInstance.CharDebugName->GetData() : nullptr;
+#endif
+
 				FPhysicsActorHandle PhysHandle;
 				FPhysicsInterface::CreateActor(Params, PhysHandle);
 				Chaos::FRigidBodyHandle_External& Body_External = PhysHandle->GetGameThreadAPI();
@@ -1914,9 +1934,15 @@ FBoxSphereBounds ULandscapeHeightfieldCollisionComponent::CalcBounds(const FTran
 
 void ULandscapeHeightfieldCollisionComponent::BeginDestroy()
 {
-	HeightfieldRef = nullptr;
-	HeightfieldGuid = FGuid();
 	Super::BeginDestroy();
+
+	// Should have been reset in OnUnregister which is called from Super::BeginDestroy
+	if (!ensure(HeightfieldRef == nullptr))
+	{
+		HeightfieldRef = nullptr;
+		HeightfieldGuid = FGuid();
+		CachedHeightFieldSamples.Empty();
+	}
 }
 
 void ULandscapeMeshCollisionComponent::BeginDestroy()
@@ -1942,45 +1968,69 @@ bool ULandscapeHeightfieldCollisionComponent::RecreateCollision()
 		}
 		CollisionHash = NewHash;
 #endif // WITH_EDITOR
-		TRefCountPtr<FHeightfieldGeometryRef> HeightfieldRefLifetimeExtender = HeightfieldRef; // Ensure heightfield data is alive until removed from physics world
+
+		// Collision geometry must be kept alive as long as we have a particle on the physics
+		// that references it. See ExtendCollisionLifetime
+		TRefCountPtr<FHeightfieldGeometryRef> HeightfieldRefLifetimeExtender = HeightfieldRef;
+
 		HeightfieldRef = nullptr; // Ensure data will be recreated
 		HeightfieldGuid = FGuid();
 		CachedHeightFieldSamples.Empty();
 		RecreatePhysicsState();
 
+		// Make sure our collision isn't destroyed while we still have a physics particle active
+		// NOTE: Must be after the call to DestroyPhysicsState
+		DeferredDestroyCollision(HeightfieldRefLifetimeExtender);
+
 		MarkRenderStateDirty();
 
-		// If we have a world, then our physics state will be queued for processing on the physics thread
-		if(UWorld* World = GetWorld())
-		{
-			if(FPhysScene* PhysScene = World->GetPhysicsScene(); PhysScene && PhysScene->GetSolver() && HeightfieldRefLifetimeExtender)
-			{
-				// We could potentially call RecreateCollision multiple times before a physics update happens, especially
-				// if we're using the async tick mode for physics. In this case we would have a pending actor in the
-				// dirty proxy list on the physics thread with a geometry that has been destructed by the lifetime
-				// extender falling out of scope.
-				// To avoid this we dispatch an empty callable with the unique geometries which runs after the
-				// proxy queue will have been cleared, avoiding a use-after-free.
-				// #TODO auto ref counted user objects for Chaos.
-				PhysScene->GetSolver()->EnqueueCommandImmediate(
-					[ComplexHeightfield = MoveTemp(HeightfieldRefLifetimeExtender->Heightfield)
-					, SimpleHeightfield = MoveTemp(HeightfieldRefLifetimeExtender->HeightfieldSimple)
+	}
+	return true;
+}
+
+// @todo(chaos): get rid of this when collision shapes are properly ref counted
+void ULandscapeHeightfieldCollisionComponent::DeferredDestroyCollision(const TRefCountPtr<FHeightfieldGeometryRef>& HeightfieldRefLifetimeExtender)
+{
+	// The editor may have a reference to the geometry as well, so we don't destroy it unless we're the last reference
+	if (!HeightfieldRefLifetimeExtender.IsValid() || (HeightfieldRefLifetimeExtender->GetRefCount() > 1))
+	{
+		return;
+	}
+
+	UWorld* World = GetWorld();
+	if (World == nullptr)
+	{
+		return;
+	}
+
+	FPhysScene* PhysScene = World->GetPhysicsScene();
+	if ((PhysScene == nullptr) || (PhysScene->GetSolver() == nullptr))
+	{
+		return;
+	}
+
+	// We could potentially call RecreateCollision multiple times before a physics update happens, especially
+	// if we're using the async tick mode for physics. In this case we would have a pending actor in the
+	// dirty proxy list on the physics thread with a geometry that has been destructed by the lifetime
+	// extender falling out of scope.
+	// To avoid this we dispatch an empty callable with the unique geometries which runs after the
+	// proxy queue will have been cleared, avoiding a use-after-free.
+	// #TODO auto ref counted user objects for Chaos.
+	PhysScene->GetSolver()->EnqueueCommandImmediate(
+		[ComplexHeightfield = MoveTemp(HeightfieldRefLifetimeExtender->Heightfield)
+		, SimpleHeightfield = MoveTemp(HeightfieldRefLifetimeExtender->HeightfieldSimple)
 #if WITH_EDITORONLY_DATA
-					, EditorHeightfield = MoveTemp(HeightfieldRefLifetimeExtender->EditorHeightfield)
+		, EditorHeightfield = MoveTemp(HeightfieldRefLifetimeExtender->EditorHeightfield)
 #endif
-					]
-				() mutable
-				{
-					ComplexHeightfield = nullptr;
-					SimpleHeightfield = nullptr;
+		]
+	() mutable
+		{
+			ComplexHeightfield = nullptr;
+			SimpleHeightfield = nullptr;
 #if WITH_EDITORONLY_DATA
-					EditorHeightfield = nullptr;
+			EditorHeightfield = nullptr;
 #endif
-				});
-			}
-		}
-	}
-	return true;
+		});
 }
 
 #if WITH_EDITORONLY_DATA
