Modifies landscape shared property deprecation for inherited properties.
Adds CVar to silence and handle deprecation of overridable properties.

diff --git a/Engine/Source/Runtime/Landscape/Private/Landscape.cpp b/Engine/Source/Runtime/Landscape/Private/Landscape.cpp
--- a/Engine/Source/Runtime/Landscape/Private/Landscape.cpp
+++ b/Engine/Source/Runtime/Landscape/Private/Landscape.cpp
@@ -174,6 +174,11 @@ static FAutoConsoleVariable CVarForceInvalidateNaniteOnLoad(
 	TEXT("landscape.ForceInvalidateNaniteOnLoad"),
 	false,
 	TEXT("Trigger a rebuild of Nanite representation on load (for debugging purposes)"));
+
+static FAutoConsoleVariable CVarSilenceSharedPropertyDeprecationFixup(
+	TEXT("landscape.SilenceSharedPropertyDeprecationFixup"),
+	false,
+	TEXT("Silently performs the fixup of discrepancies in shared properties when handling data modified before the enforcement introduction."));
 #endif // WITH_EDITOR
 
 int32 GRenderNaniteLandscape = 1;
@@ -3927,6 +3932,45 @@ namespace UE::Landscape::Private
 
 		return true;
 	}
+
+	static void DisplaySynchronizedPropertiesMapcheckWarning(const TArray<FName>& InSynchronizedProperties, const ALandscapeProxy& InSynchronizedProxy, const ALandscapeProxy& InParentLandscape, const bool bAddSilencingMessage = false)
+	{
+		TStringBuilder<1024> SynchronizedPropertiesStringBuilder;
+		ULandscapeSubsystem* LandscapeSubsystem = InSynchronizedProxy.GetWorld() ? InSynchronizedProxy.GetWorld()->GetSubsystem<ULandscapeSubsystem>() : nullptr;
+		checkf(LandscapeSubsystem != nullptr, TEXT("DisplaySynchronizedPropertiesMapcheckWarning can only be called when a subsystem is available"));
+
+		for (const FName& SynchronizedProperty : InSynchronizedProperties)
+		{
+			if (SynchronizedPropertiesStringBuilder.Len() > 0)
+			{
+				SynchronizedPropertiesStringBuilder.Append(TEXT(", "));
+			}
+
+			SynchronizedPropertiesStringBuilder.Append(SynchronizedProperty.ToString());
+		}
+
+		FFormatNamedArguments Arguments;
+		Arguments.Add(TEXT("Proxy"), FText::FromString(InSynchronizedProxy.GetActorNameOrLabel()));
+		Arguments.Add(TEXT("Landscape"), FText::FromString(InParentLandscape.GetActorNameOrLabel()));
+		TSharedRef<FTokenizedMessage> Message = FMessageLog("MapCheck").Warning()
+			->AddToken(FUObjectToken::Create(&InSynchronizedProxy, FText::FromString(InSynchronizedProxy.GetActorNameOrLabel())))
+			->AddToken(FTextToken::Create(LOCTEXT("MapCheck_Message_LandscapeProxy_FixupSharedData", "had some shared properties not in sync with its parent landscape actor. This has been fixed but the proxy needs to be saved in order to ensure cooking behaves as expected. ")))
+			->AddToken(FActionToken::Create(LOCTEXT("MapCheck_SaveFixedUpData", "Save Modified Landscapes"), LOCTEXT("MapCheck_SaveFixedUpData_Desc", "Saves the modified landscape proxy actors"),
+				FOnActionTokenExecuted::CreateUObject(LandscapeSubsystem, &ULandscapeSubsystem::SaveModifiedLandscapes),
+				FCanExecuteActionToken::CreateUObject(LandscapeSubsystem, &ULandscapeSubsystem::HasModifiedLandscapes),
+				/*bInSingleUse = */false))
+			->AddToken(FTextToken::Create(FText::Format(LOCTEXT("MapCheck_Message_LandscapeProxy_FixupSharedData_SharedProperties", "The following properties were synchronized: {0}."), FText::FromString(SynchronizedPropertiesStringBuilder.ToString()))));
+			
+		if (bAddSilencingMessage)
+		{
+			Message->AddToken(FTextToken::Create(LOCTEXT("MapCheck_Message_LandscapeProxy_SilenceWarning", "You can silence this warning and perform the deprecation silently using the landscape.SilenceSharedPropertyDeprecationFixup CVar. ")));
+		}
+
+		Message->AddToken(FMapErrorToken::Create(FMapErrors::LandscapeComponentPostLoad_Warning));
+
+		// Show MapCheck window
+		FMessageLog("MapCheck").Open(EMessageSeverity::Warning);
+	}
 } // namespace UE::Landscape::Private
 
 void ALandscapeProxy::CopySharedProperties(ALandscapeProxy* Landscape)
@@ -4125,6 +4169,11 @@ void ALandscapeStreamingProxy::FixupOverriddenSharedProperties()
 
 void ALandscapeProxy::UpgradeSharedProperties(ALandscape* InParentLandscape)
 {
+	TArray<FName> SynchronizedProperties;
+	bool bOpenMapCheckWindow = false;
+	ULandscapeInfo* LandscapeInfo = GetLandscapeInfo();
+	checkf(LandscapeInfo != nullptr, TEXT("UpgradeSharedProperties can only be called after the proxies are registered to ULandscapeInfo"));
+
 	for (TFieldIterator<FProperty> PropertyIterator(GetClass()); PropertyIterator; ++PropertyIterator)
 	{
 		FProperty* Property = *PropertyIterator;
@@ -4134,13 +4183,67 @@ void ALandscapeProxy::UpgradeSharedProperties(ALandscape* InParentLandscape)
 			continue;
 		}
 
-		if ((IsPropertyInherited(Property) ||
-			(IsPropertyOverridable(Property) && !IsSharedPropertyOverridden(Property->GetName()))) &&
-			!Property->Identical_InContainer(this, InParentLandscape))
+		if (IsPropertyInherited(Property) && !Property->Identical_InContainer(this, InParentLandscape))
+		{
+			SynchronizedProperties.Emplace(Property->GetFName());
+			UE::Landscape::Private::CopyProperty(Property, InParentLandscape, this);
+		}
+		else if (IsPropertyOverridable(Property) && !IsSharedPropertyOverridden(Property->GetName()) && !Property->Identical_InContainer(this, InParentLandscape))
+		{
+			if (CVarSilenceSharedPropertyDeprecationFixup->GetBool())
+			{
+				SetSharedPropertyOverride(Property->GetName(), true);
+			}
+			else
+			{
+				FFormatNamedArguments Arguments;
+				TWeakObjectPtr<ALandscapeProxy> LandscapeProxy = this;
+				TWeakObjectPtr<ALandscape> ParentLandscape = InParentLandscape;
+				const FString PropertyName = Property->GetName();
+
+				bOpenMapCheckWindow = true;
+
+				Arguments.Add(TEXT("Proxy"), FText::FromString(GetActorNameOrLabel()));
+				Arguments.Add(TEXT("Landscape"), FText::FromString(InParentLandscape->GetActorNameOrLabel()));
+				FMessageLog("MapCheck").Warning()
+					->AddToken(FUObjectToken::Create(this, FText::FromString(GetActorNameOrLabel())))
+					->AddToken(FTextToken::Create(FText::Format(LOCTEXT("MapCheck_Message_LandscapeProxy_UpgradeSharedProperties", "Contains a property ({0}) different from parent's landscape actor. Please select between "), FText::FromString(PropertyName))))
+					->AddToken(FActionToken::Create(LOCTEXT("MapCheck_OverrideProperty", "Override property"), LOCTEXT("MapCheck_OverrideProperty_Desc", "Keeping the current value and marking the property as overriding the parent landscape's value."),
+						FOnActionTokenExecuted::CreateLambda([LandscapeProxy, PropertyName]()
+							{
+								if (LandscapeProxy.IsValid())
+								{
+									LandscapeProxy->SetSharedPropertyOverride(PropertyName, true);
+								}
+							}),
+						/*bInSingleUse = */true))
+					->AddToken(FTextToken::Create(LOCTEXT("MapCheck_Message_LandscapeProxy_UpgradeSharedProperties_Or", " or ")))
+					->AddToken(FActionToken::Create(LOCTEXT("MapCheck_InheritProperty", "Inherit from parent landscape"), LOCTEXT("MapCheck_InheritProperty_Desc", "Copying the parent landscape's value for this property."),
+							FOnActionTokenExecuted::CreateLambda([LandscapeProxy, ParentLandscape, PropertyName]()
+								{
+									UE::Landscape::Private::CopyPostEditPropertyByName(LandscapeProxy, ParentLandscape, PropertyName);
+								}),
+							/*bInSingleUse = */true))
+					->AddToken(FMapErrorToken::Create(FMapErrors::LandscapeComponentPostLoad_Warning));
+			}
+		}
+	}
+
+	if (!SynchronizedProperties.IsEmpty())
+	{
+		LandscapeInfo->MarkObjectDirty(/*InObject = */this, /*bInForceResave = */true);
+
+		if (!CVarSilenceSharedPropertyDeprecationFixup->GetBool())
 		{
-			SetSharedPropertyOverride(Property->GetName(), true);
+			UE::Landscape::Private::DisplaySynchronizedPropertiesMapcheckWarning(SynchronizedProperties, *this, *InParentLandscape, true);
 		}
 	}
+
+	if (bOpenMapCheckWindow)
+	{
+		// Show MapCheck window
+		FMessageLog("MapCheck").Open(EMessageSeverity::Warning);
+	}
 }
 
 void ALandscapeProxy::FixupSharedData(ALandscape* Landscape, const bool bMapCheck)
@@ -4159,8 +4262,6 @@ void ALandscapeProxy::FixupSharedData(ALandscape* Landscape, const bool bMapChec
 	{
 		ULandscapeInfo* LandscapeInfo = GetLandscapeInfo();
 		checkf(LandscapeInfo != nullptr, TEXT("FixupSharedData can only be called after the proxies are registered to ULandscapeInfo"));
-		ULandscapeSubsystem* LandscapeSubsystem = GetWorld() ? GetWorld()->GetSubsystem<ULandscapeSubsystem>() : nullptr;
-		checkf(LandscapeSubsystem != nullptr, TEXT("FixupSharedData can only be called when a subsystem is available"));
 
 		TArray<FName> SynchronizedProperties = SynchronizeSharedProperties(Landscape);
 		bool bUpdated = !SynchronizedProperties.IsEmpty();
@@ -4181,33 +4282,7 @@ void ALandscapeProxy::FixupSharedData(ALandscape* Landscape, const bool bMapChec
 
 			if (bMapCheck && bNeedsManualResave)
 			{
-				TStringBuilder<1024> SynchronizedPropertiesStringBuilder;
-
-				for (const FName& SynchronizedProperty : SynchronizedProperties)
-				{
-					if (SynchronizedPropertiesStringBuilder.Len() > 0)
-					{
-						SynchronizedPropertiesStringBuilder.Append(TEXT(", "));
-					}
-
-					SynchronizedPropertiesStringBuilder.Append(SynchronizedProperty.ToString());
-				}
-
-				FFormatNamedArguments Arguments;
-				Arguments.Add(TEXT("Proxy"), FText::FromString(GetActorNameOrLabel()));
-				Arguments.Add(TEXT("Landscape"), FText::FromString(Landscape->GetActorNameOrLabel()));
-				FMessageLog("MapCheck").Warning()
-					->AddToken(FUObjectToken::Create(this, FText::FromString(GetActorNameOrLabel())))
-					->AddToken(FTextToken::Create(LOCTEXT("MapCheck_Message_LandscapeProxy_FixupSharedData", "had some shared properties not in sync with its parent landscape actor. This has been fixed but the proxy needs to be saved in order to ensure cooking behaves as expected. ")))
-					->AddToken(FActionToken::Create(LOCTEXT("MapCheck_SaveFixedUpData", "Save Modified Landscapes"), LOCTEXT("MapCheck_SaveFixedUpData_Desc", "Saves the modified landscape proxy actors"),
-						FOnActionTokenExecuted::CreateUObject(LandscapeSubsystem, &ULandscapeSubsystem::SaveModifiedLandscapes),
-						FCanExecuteActionToken::CreateUObject(LandscapeSubsystem, &ULandscapeSubsystem::HasModifiedLandscapes),
-						/*bInSingleUse = */false))
-					->AddToken(FTextToken::Create(FText::Format(LOCTEXT("MapCheck_Message_LandscapeProxy_FixupSharedData_SharedProperties", "The following properties were synchronized: {0}."), FText::FromString(SynchronizedPropertiesStringBuilder.ToString()))))
-					->AddToken(FMapErrorToken::Create(FMapErrors::LandscapeComponentPostLoad_Warning));
-
-				// Show MapCheck window
-				FMessageLog("MapCheck").Open(EMessageSeverity::Warning);
+				UE::Landscape::Private::DisplaySynchronizedPropertiesMapcheckWarning(SynchronizedProperties, *this, *Landscape);
 			}
 		}
 	}
