Only dirty render state if necessary when assigning LOD Group Key

diff --git a/Engine/Source/Runtime/Landscape/Classes/LandscapeProxy.h b/Engine/Source/Runtime/Landscape/Classes/LandscapeProxy.h
--- a/Engine/Source/Runtime/Landscape/Classes/LandscapeProxy.h
+++ b/Engine/Source/Runtime/Landscape/Classes/LandscapeProxy.h
@@ -929,6 +929,8 @@ public:
 	 */
 	const FGuid& GetOriginalLandscapeGuid() const { return OriginalLandscapeGuid; }
 
+	void SetLODGroupKeyInternal(uint32 InLODGroupKey);
+
 	UFUNCTION(BlueprintCallable, Category = "Landscape|Runtime")
 	virtual ALandscape* GetLandscapeActor() PURE_VIRTUAL(GetLandscapeActor, return nullptr;)
 	virtual const ALandscape* GetLandscapeActor() const PURE_VIRTUAL(GetLandscapeActor, return nullptr;)
diff --git a/Engine/Source/Runtime/Landscape/Private/Landscape.cpp b/Engine/Source/Runtime/Landscape/Private/Landscape.cpp
--- a/Engine/Source/Runtime/Landscape/Private/Landscape.cpp
+++ b/Engine/Source/Runtime/Landscape/Private/Landscape.cpp
@@ -1558,8 +1558,27 @@ void ALandscapeStreamingProxy::SetLandscapeActor(ALandscape* InLandscape)
 
 void ALandscape::SetLODGroupKey(uint32 InLODGroupKey)
 {
-	LODGroupKey = InLODGroupKey;
-	MarkAllLandscapeRenderStateDirty();
+	SetLODGroupKeyInternal(InLODGroupKey);
+
+	// change LODGroupKey on any proxies that are currently registered
+	// (any proxies that get registered later will copy the value on registration)
+	if (ULandscapeInfo* Info = GetLandscapeInfo())
+	{
+		Info->ForEachLandscapeProxy([InLODGroupKey](ALandscapeProxy* Proxy)
+			{
+				Proxy->SetLODGroupKeyInternal(InLODGroupKey);
+				return true;
+			});
+	}
+}
+
+void ALandscapeProxy::SetLODGroupKeyInternal(uint32 InLODGroupKey)
+{
+	if (LODGroupKey != InLODGroupKey)
+	{
+		LODGroupKey = InLODGroupKey;
+		MarkComponentsRenderStateDirty();
+	}
 }
 
 uint32 ALandscape::GetLODGroupKey()
@@ -4965,7 +4984,7 @@ void ULandscapeInfo::RegisterActor(ALandscapeProxy* Proxy, bool bMapCheck, bool
 #if WITH_EDITOR
 					StreamingProxy->FixupSharedData(Landscape, bMapCheck);
 #endif // WITH_EDITOR
-					StreamingProxy->LODGroupKey = Landscape->LODGroupKey;
+					StreamingProxy->SetLODGroupKeyInternal(Landscape->LODGroupKey);
 				}
 			}
 		}
@@ -5012,7 +5031,7 @@ void ULandscapeInfo::RegisterActor(ALandscapeProxy* Proxy, bool bMapCheck, bool
 		if (LandscapeActor.IsValid())	// don't overwrite the proxy's landscape actor if we don't have one registered yet
 		{
 			StreamingProxy->SetLandscapeActor(LandscapeActor.Get());
-			StreamingProxy->LODGroupKey = LandscapeActor.Get()->LODGroupKey;
+			StreamingProxy->SetLODGroupKeyInternal(LandscapeActor.Get()->LODGroupKey);
 #if WITH_EDITOR
 			StreamingProxy->FixupSharedData(LandscapeActor.Get(), bMapCheck);
 #endif // WITH_EDITOR
