Added getter/setter for APartitionedActor::GridSize
* Setter will log an error if GridSize is 0
* Deprecated GridSize, will ultimately be made private

diff --git a/Engine/Plugins/Experimental/PCG/Source/PCG/Private/Grid/PCGPartitionActor.cpp b/Engine/Plugins/Experimental/PCG/Source/PCG/Private/Grid/PCGPartitionActor.cpp
--- a/Engine/Plugins/Experimental/PCG/Source/PCG/Private/Grid/PCGPartitionActor.cpp
+++ b/Engine/Plugins/Experimental/PCG/Source/PCG/Private/Grid/PCGPartitionActor.cpp
@@ -421,7 +421,7 @@ AActor* APCGPartitionActor::GetSceneOutlinerParent() const
 
 void APCGPartitionActor::PostCreation(const FGuid& InGridGUID)
 {
-	PCGGridSize = GridSize;
+	PCGGridSize = GetGridSize();
 	PCGGuid = InGridGUID;
 
 	// Put in cache if we use the 2D grid or not.
diff --git a/Engine/Source/Editor/LandscapeEditor/Private/LandscapeEditorDetailCustomization_ImportExport.cpp b/Engine/Source/Editor/LandscapeEditor/Private/LandscapeEditorDetailCustomization_ImportExport.cpp
--- a/Engine/Source/Editor/LandscapeEditor/Private/LandscapeEditorDetailCustomization_ImportExport.cpp
+++ b/Engine/Source/Editor/LandscapeEditor/Private/LandscapeEditorDetailCustomization_ImportExport.cpp
@@ -723,8 +723,8 @@ FReply FLandscapeEditorDetailCustomization_ImportExport::OnImportExportButtonCli
 						if (LandscapeInfo->GetLandscapeExtent(LandscapeProxy, ExportRegion))
 						{
 							const int32 YCoord = LandscapeEdMode->UISettings->bFlipYAxis ? LandscapeExtent.Max.Y - ExportRegion.Max.Y : ExportRegion.Min.Y - LandscapeExtent.Min.Y;
-							FIntPoint FileOffset = FIntPoint((ExportRegion.Min.X - LandscapeExtent.Min.X) / Landscape->GridSize,
-								YCoord / Landscape->GridSize);
+							FIntPoint FileOffset = FIntPoint((ExportRegion.Min.X - LandscapeExtent.Min.X) / Landscape->GetGridSize(),
+								YCoord / Landscape->GetGridSize());
 
 							// Remove the shared line/column that this proxy has with its neighbors because it 
 							// will be included by the neighbor or lost if there is none (that could become an option to avoid that loss)
@@ -748,8 +748,8 @@ FReply FLandscapeEditorDetailCustomization_ImportExport::OnImportExportButtonCli
 					if (LandscapeInfo->GetLandscapeExtent(LandscapeProxy, ExportRegion))
 					{
 						const int32 YCoord = LandscapeEdMode->UISettings->bFlipYAxis ? LandscapeExtent.Max.Y - ExportRegion.Max.Y : ExportRegion.Min.Y - LandscapeExtent.Min.Y;
-						FIntPoint FileOffset = FIntPoint((ExportRegion.Min.X - LandscapeExtent.Min.X) / Landscape->GridSize,
-							YCoord / Landscape->GridSize);
+						FIntPoint FileOffset = FIntPoint((ExportRegion.Min.X - LandscapeExtent.Min.X) / Landscape->GetGridSize(),
+							YCoord / Landscape->GetGridSize());
 
 						// Remove the shared line/column that this proxy has with its neighbors because it 
 						// will be included by the neighbor or lost if there is none (that could become an option to avoid that loss)
diff --git a/Engine/Source/Editor/UnrealEd/Private/WorldPartition/WorldPartitionFoliageBuilder.cpp b/Engine/Source/Editor/UnrealEd/Private/WorldPartition/WorldPartitionFoliageBuilder.cpp
--- a/Engine/Source/Editor/UnrealEd/Private/WorldPartition/WorldPartitionFoliageBuilder.cpp
+++ b/Engine/Source/Editor/UnrealEd/Private/WorldPartition/WorldPartitionFoliageBuilder.cpp
@@ -222,7 +222,7 @@ bool UWorldPartitionFoliageBuilder::RunInternal(UWorld* World, const FCellInfo&
 		}
 		else
 		{	
-			check(IFA->GridSize != NewGridSize);
+			check(IFA->GetGridSize() != NewGridSize);
 			// Harvest Instances from existing IFA and build up instance bounds
 			UE_LOG(LogWorldPartitionFoliageBuilder, Display, TEXT("Processing existing foliage actor: %s (%s)"), *ActorDesc->GetActorName().ToString(), *ActorDesc->GetActorPackage().ToString());
 		
@@ -280,7 +280,7 @@ bool UWorldPartitionFoliageBuilder::RunInternal(UWorld* World, const FCellInfo&
 					AInstancedFoliageActor* ModifiedIFA = AInstancedFoliageActor::Get(World, /*bCreateIfNone=*/false, World->PersistentLevel, Instance.Location);
 					check(ModifiedIFA);
 					ModifiedIFAs.Add(ModifiedIFA);
-					check(ModifiedIFA->GridSize == NewGridSize);
+					check(ModifiedIFA->GetGridSize() == NewGridSize);
 					FFoliageInfo* NewFoliageInfo = nullptr;
 					UFoliageType* NewFoliageType = ModifiedIFA->AddFoliageType(InstancesPerFoliageType.Key, &NewFoliageInfo);
 					NewFoliageInfo->AddInstance(NewFoliageType, Instance);
diff --git a/Engine/Source/Runtime/Engine/Private/ActorPartition/ActorPartitionSubsystem.cpp b/Engine/Source/Runtime/Engine/Private/ActorPartition/ActorPartitionSubsystem.cpp
--- a/Engine/Source/Runtime/Engine/Private/ActorPartition/ActorPartitionSubsystem.cpp
+++ b/Engine/Source/Runtime/Engine/Private/ActorPartition/ActorPartitionSubsystem.cpp
@@ -265,7 +265,7 @@ public:
 						
 			FVector CellCenter(CellBounds.GetCenter());
 			FoundActor = CastChecked<APartitionActor>(World->SpawnActor(InActorPartitionId.GetClass(), &CellCenter, nullptr, SpawnParams));
-			FoundActor->GridSize = InGridSize;
+			FoundActor->SetGridSize(InGridSize);
 			FoundActor->SetLockLocation(true);
 			
 			InActorCreated(FoundActor);
diff --git a/Engine/Source/Runtime/Engine/Private/ActorPartition/PartitionActor.cpp b/Engine/Source/Runtime/Engine/Private/ActorPartition/PartitionActor.cpp
--- a/Engine/Source/Runtime/Engine/Private/ActorPartition/PartitionActor.cpp
+++ b/Engine/Source/Runtime/Engine/Private/ActorPartition/PartitionActor.cpp
@@ -14,6 +14,8 @@
 
 #define LOCTEXT_NAMESPACE "PartitionActor"
 
+DEFINE_LOG_CATEGORY_STATIC(LogPartitionActor, Log, All);
+
 APartitionActor::APartitionActor(const FObjectInitializer& ObjectInitializer)
 	: Super(ObjectInitializer)
 #if WITH_EDITORONLY_DATA
@@ -52,6 +54,23 @@ bool APartitionActor::ShouldIncludeGridSizeInName(UWorld* InWorld, const FActorP
 	return InWorld->GetWorldSettings()->bIncludeGridSizeInNameForPartitionedActors;
 }
 
+PRAGMA_DISABLE_DEPRECATION_WARNINGS
+uint32 APartitionActor::GetGridSize() const
+{
+	return GridSize;
+} 
+
+void APartitionActor::SetGridSize(uint32 InGridSize)
+{
+	if (InGridSize == 0)
+	{
+		UE_LOG(LogPartitionActor, Error, TEXT("APartitionActor::SetGridSize() called for actor %s with grid size == 0. Grid size must be greater than zero."), *GetName());
+		InGridSize = 1;
+	}
+	GridSize = InGridSize;
+}
+PRAGMA_ENABLE_DEPRECATION_WARNINGS
+
 FString APartitionActor::GetActorName(UWorld* World, const UClass* Class, const FGuid& GridGuid, const FActorPartitionIdentifier& ActorPartitionId, uint32 GridSize, int32 CellCoordsX, int32 CellCoordsY, int32 CellCoordsZ, uint32 DataLayerEditorContext)
 {
 	TStringBuilderWithBuffer<TCHAR, NAME_SIZE> ActorNameBuilder;
diff --git a/Engine/Source/Runtime/Engine/Private/WorldPartition/ActorPartition/PartitionActorDesc.cpp b/Engine/Source/Runtime/Engine/Private/WorldPartition/ActorPartition/PartitionActorDesc.cpp
--- a/Engine/Source/Runtime/Engine/Private/WorldPartition/ActorPartition/PartitionActorDesc.cpp
+++ b/Engine/Source/Runtime/Engine/Private/WorldPartition/ActorPartition/PartitionActorDesc.cpp
@@ -20,7 +20,7 @@ void FPartitionActorDesc::Init(const AActor* InActor)
 
 	const APartitionActor* PartitionActor = CastChecked<APartitionActor>(InActor);
 
-	GridSize = PartitionActor->GridSize;
+	GridSize = PartitionActor->GetGridSize();
 
 	if (!bIsDefaultActorDesc)
 	{
diff --git a/Engine/Source/Runtime/Engine/Public/ActorPartition/PartitionActor.h b/Engine/Source/Runtime/Engine/Public/ActorPartition/PartitionActor.h
--- a/Engine/Source/Runtime/Engine/Public/ActorPartition/PartitionActor.h
+++ b/Engine/Source/Runtime/Engine/Public/ActorPartition/PartitionActor.h
@@ -21,7 +21,7 @@ public:
 	ENGINE_API APartitionActor(const FObjectInitializer& ObjectInitializer = FObjectInitializer::Get());
 
 	//~ Begin AActor Interface
-#if WITH_EDITOR	
+#if WITH_EDITOR
 	virtual bool CanChangeIsSpatiallyLoadedFlag() const override { return false; }
 	ENGINE_API virtual TUniquePtr<class FWorldPartitionActorDesc> CreateClassActorDesc() const override;
 	ENGINE_API virtual uint32 GetDefaultGridSize(UWorld* InWorld) const PURE_VIRTUAL(APartitionActor, return 0;)
@@ -30,10 +30,14 @@ public:
 	virtual FGuid GetGridGuid() const { return FGuid(); }
 	ENGINE_API virtual bool IsUserManaged() const override;
 #endif
-	//~ End AActor Interface	
+	//~ End AActor Interface
 
 #if WITH_EDITORONLY_DATA
+	ENGINE_API uint32 GetGridSize() const;
+	ENGINE_API void SetGridSize(uint32 InGridSize);
+
 	/** The grid size this actors was generated for */
+	UE_DEPRECATED(5.3, "Use GetGridSize()/SetGridSize() instead.")
 	UPROPERTY()
 	uint32 GridSize;
 
diff --git a/Engine/Source/Runtime/Landscape/Private/LandscapeConfigHelper.cpp b/Engine/Source/Runtime/Landscape/Private/LandscapeConfigHelper.cpp
--- a/Engine/Source/Runtime/Landscape/Private/LandscapeConfigHelper.cpp
+++ b/Engine/Source/Runtime/Landscape/Private/LandscapeConfigHelper.cpp
@@ -21,7 +21,7 @@ int32 FLandscapeConfig::NumSectionValues[2] = { 1, 2 };
 int32 FLandscapeConfig::SubsectionSizeQuadsValues[6] = { 7, 15, 31, 63, 127, 255 };
 
 FLandscapeConfig::FLandscapeConfig(ULandscapeInfo* InLandscapeInfo)
-	: FLandscapeConfig(InLandscapeInfo->ComponentNumSubsections, InLandscapeInfo->SubsectionSizeQuads, InLandscapeInfo->LandscapeActor->GridSize / InLandscapeInfo->ComponentSizeQuads)
+	: FLandscapeConfig(InLandscapeInfo->ComponentNumSubsections, InLandscapeInfo->SubsectionSizeQuads, InLandscapeInfo->LandscapeActor->GetGridSize() / InLandscapeInfo->ComponentSizeQuads)
 {
 
 }
@@ -53,7 +53,7 @@ ALandscapeProxy* FLandscapeConfigHelper::FindOrAddLandscapeStreamingProxy(ULands
 {
 	UWorld* World = InLandscapeInfo->LandscapeActor->GetWorld();
 	UActorPartitionSubsystem* ActorPartitionSubsystem = World->GetSubsystem<UActorPartitionSubsystem>();
-	const uint32 GridSize = InLandscapeInfo->LandscapeActor->GridSize;
+	const uint32 GridSize = InLandscapeInfo->LandscapeActor->GetGridSize();
 
 	UActorPartitionSubsystem::FCellCoord CellCoord = UActorPartitionSubsystem::FCellCoord::GetCellCoord(InSectionBase, World->PersistentLevel, GridSize);
 	return FLandscapeConfigHelper::FindOrAddLandscapeStreamingProxy(ActorPartitionSubsystem, InLandscapeInfo, CellCoord);
@@ -66,7 +66,7 @@ ALandscapeProxy* FLandscapeConfigHelper::FindOrAddLandscapeStreamingProxy(UActor
 
 	auto LandscapeProxyCreated = [InCellCoord, Landscape](APartitionActor* PartitionActor)
 	{
-		const FIntPoint CellLocation(static_cast<int32>(InCellCoord.X) * Landscape->GridSize, static_cast<int32>(InCellCoord.Y) * Landscape->GridSize);
+		const FIntPoint CellLocation(static_cast<int32>(InCellCoord.X) * Landscape->GetGridSize(), static_cast<int32>(InCellCoord.Y) * Landscape->GetGridSize());
 
 		ALandscapeProxy* LandscapeProxy = CastChecked<ALandscapeProxy>(PartitionActor);
 		// copy shared properties to this new proxy
@@ -82,8 +82,8 @@ ALandscapeProxy* FLandscapeConfigHelper::FindOrAddLandscapeStreamingProxy(UActor
 	const bool bCreate = true;
 	const bool bBoundsSearch = false;
 
-	ALandscapeProxy* LandscapeProxy = Cast<ALandscapeProxy>(InActorPartitionSubsystem->GetActor(ALandscapeStreamingProxy::StaticClass(), InCellCoord, bCreate, InLandscapeInfo->LandscapeGuid, Landscape->GridSize, bBoundsSearch, LandscapeProxyCreated));
-	check(!LandscapeProxy || LandscapeProxy->GridSize == Landscape->GridSize);
+	ALandscapeProxy* LandscapeProxy = Cast<ALandscapeProxy>(InActorPartitionSubsystem->GetActor(ALandscapeStreamingProxy::StaticClass(), InCellCoord, bCreate, InLandscapeInfo->LandscapeGuid, Landscape->GetGridSize(), bBoundsSearch, LandscapeProxyCreated));
+	check(!LandscapeProxy || LandscapeProxy->GetGridSize() == Landscape->GetGridSize());
 	return LandscapeProxy;
 }
 
@@ -94,7 +94,7 @@ bool FLandscapeConfigHelper::ChangeGridSize(ULandscapeInfo* InLandscapeInfo, uin
 	const uint32 GridSize = InLandscapeInfo->GetGridSize(InNewGridSizeInComponents);
 	
 	InLandscapeInfo->LandscapeActor->Modify();
-	InLandscapeInfo->LandscapeActor->GridSize = GridSize;
+	InLandscapeInfo->LandscapeActor->SetGridSize(GridSize);
 
 	// This needs to be done before moving components
 	InLandscapeInfo->LandscapeActor->InitializeLandscapeLayersWeightmapUsage();
@@ -224,7 +224,7 @@ bool FLandscapeConfigHelper::ChangeGridSize(ULandscapeInfo* InLandscapeInfo, uin
 	{
 		if (ProxyToDelete->LandscapeComponents.Num() > 0 || ProxyToDelete->IsA<ALandscape>())
 		{
-			check(ProxyToDelete->GridSize == GridSize);
+			check(ProxyToDelete->GetGridSize() == GridSize);
 			continue;
 		}
 
@@ -488,7 +488,7 @@ ULandscapeInfo* FLandscapeConfigHelper::ChangeConfiguration(ULandscapeInfo* InLa
 	
 	NewLandscape->CopySharedProperties(OldLandscape);
 	NewLandscape->SetLandscapeGuid(FGuid::NewGuid());
-	NewLandscape->GridSize = InNewConfig.GetGridSizeQuads();
+	NewLandscape->SetGridSize(InNewConfig.GetGridSizeQuads());
 	NewLandscape->ComponentSizeQuads = InNewConfig.GetComponentSizeQuads();
 	NewLandscape->NumSubsections = InNewConfig.ComponentNumSubsections;
 	NewLandscape->SubsectionSizeQuads = InNewConfig.SubsectionSizeQuads;	
@@ -520,7 +520,7 @@ ULandscapeInfo* FLandscapeConfigHelper::ChangeConfiguration(ULandscapeInfo* InLa
 				continue;
 			}
 
-			UActorPartitionSubsystem::FCellCoord CellCoord = UActorPartitionSubsystem::FCellCoord::GetCellCoord(NewSectionBase, World->PersistentLevel, NewLandscape->GridSize);
+			UActorPartitionSubsystem::FCellCoord CellCoord = UActorPartitionSubsystem::FCellCoord::GetCellCoord(NewSectionBase, World->PersistentLevel, NewLandscape->GetGridSize());
 			ALandscapeProxy* LandscapeProxy = FLandscapeConfigHelper::FindOrAddLandscapeStreamingProxy(ActorPartitionSubsystem, NewLandscapeInfo, CellCoord);
 			check(LandscapeProxy);
 
diff --git a/Engine/Source/Runtime/Landscape/Private/LandscapeEdit.cpp b/Engine/Source/Runtime/Landscape/Private/LandscapeEdit.cpp
--- a/Engine/Source/Runtime/Landscape/Private/LandscapeEdit.cpp
+++ b/Engine/Source/Runtime/Landscape/Private/LandscapeEdit.cpp
@@ -2608,8 +2608,8 @@ bool ULandscapeInfo::HasUnloadedComponentsInRegion(int32 X1, int32 Y1, int32 X2,
 		int32 ComponentIndexX1, ComponentIndexY1, ComponentIndexX2, ComponentIndexY2;
 		ALandscape::CalcComponentIndicesOverlap(X1, Y1, X2, Y2, ComponentSizeQuads, ComponentIndexX1, ComponentIndexY1, ComponentIndexX2, ComponentIndexY2);
 
-		const UActorPartitionSubsystem::FCellCoord MinCoord = UActorPartitionSubsystem::FCellCoord::GetCellCoord(FIntPoint(ComponentIndexX1 * ComponentSizeQuads, ComponentIndexY1 * ComponentSizeQuads), World->PersistentLevel, LandscapeActor->GridSize);
-		const UActorPartitionSubsystem::FCellCoord MaxCoord = UActorPartitionSubsystem::FCellCoord::GetCellCoord(FIntPoint(ComponentIndexX2 * ComponentSizeQuads, ComponentIndexY2 * ComponentSizeQuads), World->PersistentLevel, LandscapeActor->GridSize);
+		const UActorPartitionSubsystem::FCellCoord MinCoord = UActorPartitionSubsystem::FCellCoord::GetCellCoord(FIntPoint(ComponentIndexX1 * ComponentSizeQuads, ComponentIndexY1 * ComponentSizeQuads), World->PersistentLevel, LandscapeActor->GetGridSize());
+		const UActorPartitionSubsystem::FCellCoord MaxCoord = UActorPartitionSubsystem::FCellCoord::GetCellCoord(FIntPoint(ComponentIndexX2 * ComponentSizeQuads, ComponentIndexY2 * ComponentSizeQuads), World->PersistentLevel, LandscapeActor->GetGridSize());
 
 		if (UWorldPartition* WorldPartition = World->GetWorldPartition())
 		{
@@ -6410,7 +6410,7 @@ void ALandscapeStreamingProxy::PostRegisterAllComponents()
 			if (UWorldPartition* WorldPartition = GetWorld()->GetWorldPartition(); WorldPartition && WorldPartition->IsInitialized())
 			{
 				const FVector ActorLocation = GetActorLocation();
-				FBox Bounds(ActorLocation, ActorLocation + (GridSize * LandscapeInfo->DrawScale));
+				FBox Bounds(ActorLocation, ActorLocation + (GetGridSize() * LandscapeInfo->DrawScale));
 
 				// all actors that intersect Landscape in 2D need to be considered
 				Bounds.Min.Z = -HALF_WORLD_MAX;
