[Engine]

diff --git a/Engine/Source/Runtime/Landscape/Classes/Landscape.h b/Engine/Source/Runtime/Landscape/Classes/Landscape.h
--- a/Engine/Source/Runtime/Landscape/Classes/Landscape.h
+++ b/Engine/Source/Runtime/Landscape/Classes/Landscape.h
@@ -271,9 +271,9 @@ public:
 #if WITH_EDITOR
 	/** Computes & returns bounds containing all landscape proxies (if any) or this landscape's bounds otherwise. Note that in non-WP worlds this will call GetLoadedBounds(). */
 	LANDSCAPE_API FBox GetCompleteBounds() const;
-	LANDSCAPE_API void RegisterLandscapeEdMode(ILandscapeEdModeInterface* InLandscapeEdMode) { LandscapeEdMode = InLandscapeEdMode; }
-	LANDSCAPE_API void UnregisterLandscapeEdMode() { LandscapeEdMode = nullptr; }
-	LANDSCAPE_API bool HasLandscapeEdMode() const { return LandscapeEdMode != nullptr; }
+	void RegisterLandscapeEdMode(ILandscapeEdModeInterface* InLandscapeEdMode) { LandscapeEdMode = InLandscapeEdMode; }
+	void UnregisterLandscapeEdMode() { LandscapeEdMode = nullptr; }
+	bool HasLandscapeEdMode() const { return LandscapeEdMode != nullptr; }
 	LANDSCAPE_API virtual bool HasLayersContent() const override;
 	LANDSCAPE_API virtual void UpdateCachedHasLayersContent(bool bInCheckComponentDataIntegrity) override;
 	LANDSCAPE_API void RequestSplineLayerUpdate();
@@ -355,7 +355,7 @@ public:
 	LANDSCAPE_API bool GetUseGeneratedLandscapeSplineMeshesActors() const;
 	LANDSCAPE_API bool PrepareTextureResources(bool bInWaitForStreaming);
 
-	LANDSCAPE_API bool GetVisibilityLayerAllocationIndex() const { return 0; }
+	bool GetVisibilityLayerAllocationIndex() const { return 0; }
 
 protected:
 	FName GenerateUniqueLayerName(FName InName = NAME_None) const;
@@ -565,11 +565,11 @@ private:
 };
 
 #if WITH_EDITOR
-class LANDSCAPE_API FScopedSetLandscapeEditingLayer
+class FScopedSetLandscapeEditingLayer
 {
 public:
-	FScopedSetLandscapeEditingLayer(ALandscape* InLandscape, const FGuid& InLayerGUID, TFunction<void()> InCompletionCallback = TFunction<void()>());
-	~FScopedSetLandscapeEditingLayer();
+	LANDSCAPE_API FScopedSetLandscapeEditingLayer(ALandscape* InLandscape, const FGuid& InLayerGUID, TFunction<void()> InCompletionCallback = TFunction<void()>());
+	LANDSCAPE_API ~FScopedSetLandscapeEditingLayer();
 
 private:
 	TWeakObjectPtr<ALandscape> Landscape;
diff --git a/Engine/Source/Runtime/Landscape/Classes/LandscapeComponent.h b/Engine/Source/Runtime/Landscape/Classes/LandscapeComponent.h
--- a/Engine/Source/Runtime/Landscape/Classes/LandscapeComponent.h
+++ b/Engine/Source/Runtime/Landscape/Classes/LandscapeComponent.h
@@ -739,7 +739,7 @@ public:
 	LANDSCAPE_API void ForEachLayer(TFunctionRef<void(const FGuid&, struct FLandscapeLayerComponentData&)> Fn);
 
 	UE_DEPRECATED(5.1, "SetEditingLayer has been deprecated, use GetLandscapeActor()->SetEditingLayer() instead.")
-	LANDSCAPE_API void SetEditingLayer(const FGuid& InEditingLayer) {}
+	void SetEditingLayer(const FGuid& InEditingLayer) {}
 
 	/** Get the Landscape Actor's editing layer data */
 	FLandscapeLayerComponentData* GetEditingLayer();
@@ -869,13 +869,13 @@ public:
 	LANDSCAPE_API ALandscapeProxy* GetLandscapeProxy() const;
 
 	/** @return Component section base as FIntPoint */
-	LANDSCAPE_API FIntPoint GetSectionBase() const
+	FIntPoint GetSectionBase() const
 	{
 		return FIntPoint(SectionBaseX, SectionBaseY);
 	}
 
 	/** @param InSectionBase new section base for a component */
-	LANDSCAPE_API void SetSectionBase(FIntPoint InSectionBase)
+	void SetSectionBase(FIntPoint InSectionBase)
 	{
 		SectionBaseX = InSectionBase.X;
 		SectionBaseY = InSectionBase.Y;
@@ -1076,7 +1076,7 @@ public:
 	LANDSCAPE_API void RequestHeightmapUpdate(bool bUpdateAll = false, bool bUpdateCollision = true, bool bInUserTriggered = false);
 	LANDSCAPE_API void RequestEditingClientUpdate(bool bInUserTriggered = false);
 	LANDSCAPE_API void RequestDeferredClientUpdate();
-	LANDSCAPE_API uint32 GetLayerUpdateFlagPerMode() const { return LayerUpdateFlagPerMode; }
+	uint32 GetLayerUpdateFlagPerMode() const { return LayerUpdateFlagPerMode; }
 	LANDSCAPE_API uint32 ComputeWeightmapsHash();
 
 	void GetUsedPaintLayers(const FGuid& InLayerGuid, TArray<ULandscapeLayerInfoObject*>& OutUsedLayerInfos) const;
@@ -1113,8 +1113,8 @@ public:
 		return bNaniteActive;
 	}
 
-	LANDSCAPE_API ULandscapeHeightfieldCollisionComponent* GetCollisionComponent() const { return CollisionComponentRef.Get(); }
-	LANDSCAPE_API void SetCollisionComponent(ULandscapeHeightfieldCollisionComponent* InCollisionComponent) { CollisionComponentRef = InCollisionComponent; }
+	ULandscapeHeightfieldCollisionComponent* GetCollisionComponent() const { return CollisionComponentRef.Get(); }
+	void SetCollisionComponent(ULandscapeHeightfieldCollisionComponent* InCollisionComponent) { CollisionComponentRef = InCollisionComponent; }
 
 	void SetUserTriggeredChangeRequested(bool bInUserTriggeredChangeRequested)
 	{
diff --git a/Engine/Source/Runtime/Landscape/Classes/LandscapeHeightfieldCollisionComponent.h b/Engine/Source/Runtime/Landscape/Classes/LandscapeHeightfieldCollisionComponent.h
--- a/Engine/Source/Runtime/Landscape/Classes/LandscapeHeightfieldCollisionComponent.h
+++ b/Engine/Source/Runtime/Landscape/Classes/LandscapeHeightfieldCollisionComponent.h
@@ -310,7 +310,7 @@ public:
 	LANDSCAPE_API void SnapFoliageInstances();
 #endif
 
-	LANDSCAPE_API void SetRenderComponent(ULandscapeComponent* InRenderComponent) { RenderComponentRef = InRenderComponent; }
+	void SetRenderComponent(ULandscapeComponent* InRenderComponent) { RenderComponentRef = InRenderComponent; }
 
 public:
 	TOptional<float> GetHeight(float X, float Y, EHeightfieldSource HeightFieldSource);
diff --git a/Engine/Source/Runtime/Landscape/Classes/LandscapeInfo.h b/Engine/Source/Runtime/Landscape/Classes/LandscapeInfo.h
--- a/Engine/Source/Runtime/Landscape/Classes/LandscapeInfo.h
+++ b/Engine/Source/Runtime/Landscape/Classes/LandscapeInfo.h
@@ -258,7 +258,7 @@ public:
 	LANDSCAPE_API ALandscapeProxy* GetLandscapeProxyForLevel(ULevel* Level) const;
 
 	UE_DEPRECATED(5.3, "Use ULandscapeSubsystem::GetDirtyOnlyInMode() instead")
-	LANDSCAPE_API static bool IsDirtyOnlyInModeEnabled() { return false; }
+	static bool IsDirtyOnlyInModeEnabled() { return false; }
 
 	LANDSCAPE_API bool GetDirtyOnlyInMode() const;
 
@@ -328,7 +328,7 @@ public:
 	LANDSCAPE_API void RemoveXYOffsets();
 
 	UE_DEPRECATED(5.3, "Texture Baking is officially deprecated now and nothing updates it anymore")
-	LANDSCAPE_API void PostponeTextureBaking() {}
+	void PostponeTextureBaking() {}
 
 	/** Will tell if the landscape actor can have some content related to the layer system */
 	LANDSCAPE_API bool CanHaveLayersContent() const;
diff --git a/Engine/Source/Runtime/Landscape/Classes/LandscapeSplineMeshesActor.h b/Engine/Source/Runtime/Landscape/Classes/LandscapeSplineMeshesActor.h
--- a/Engine/Source/Runtime/Landscape/Classes/LandscapeSplineMeshesActor.h
+++ b/Engine/Source/Runtime/Landscape/Classes/LandscapeSplineMeshesActor.h
@@ -11,8 +11,8 @@
 //
 // This serves as an optimization as these actors will be streamed at runtime
 //
-UCLASS(NotPlaceable, NotBlueprintable, hideCategories = (Input), showCategories = ("Input|MouseInput", "Input|TouchInput"), ConversionRoot, ComponentWrapperClass, meta = (ChildCanTick))
-class LANDSCAPE_API ALandscapeSplineMeshesActor : public APartitionActor
+UCLASS(NotPlaceable, NotBlueprintable, hideCategories = (Input), showCategories = ("Input|MouseInput", "Input|TouchInput"), ConversionRoot, ComponentWrapperClass, meta = (ChildCanTick), MinimalAPI)
+class ALandscapeSplineMeshesActor : public APartitionActor
 {
 	GENERATED_UCLASS_BODY()
 
@@ -23,28 +23,28 @@ private:
 public:
 
 #if WITH_EDITOR
-	class UStaticMeshComponent* CreateStaticMeshComponent(const TSubclassOf<class UStaticMeshComponent>& InComponentClass);
-	void ClearStaticMeshComponents();
+	LANDSCAPE_API class UStaticMeshComponent* CreateStaticMeshComponent(const TSubclassOf<class UStaticMeshComponent>& InComponentClass);
+	LANDSCAPE_API void ClearStaticMeshComponents();
 
 	//~ Begin AActor Interface
-	virtual void CheckForErrors() override;
-	virtual bool GetReferencedContentObjects(TArray<UObject*>& Objects) const override;
+	LANDSCAPE_API virtual void CheckForErrors() override;
+	LANDSCAPE_API virtual bool GetReferencedContentObjects(TArray<UObject*>& Objects) const override;
 	virtual bool IsRuntimeOnly() const override { return true; }
 	//~ End AActor Interface
 
 	//~ Begin APartitionActor Interface
-	virtual uint32 GetDefaultGridSize(UWorld* InWorld) const override;
+	LANDSCAPE_API virtual uint32 GetDefaultGridSize(UWorld* InWorld) const override;
 	virtual FGuid GetGridGuid() const override { return GridGuid; }
 	//~ End APartitionActor Interface
 
 	void SetGridGuid(const FGuid& InGuid) { GridGuid = InGuid; }
 #endif // WITH_EDITOR
 
-	const TArray<class UStaticMeshComponent*>& GetStaticMeshComponents() const;
+	LANDSCAPE_API const TArray<class UStaticMeshComponent*>& GetStaticMeshComponents() const;
 
 protected:
 	//~ Begin UObject Interface.
-	virtual FString GetDetailedInfoInternal() const override;
+	LANDSCAPE_API virtual FString GetDetailedInfoInternal() const override;
 	//~ End UObject Interface.
 
 private:
@@ -52,4 +52,4 @@ private:
 	UPROPERTY()
 	FGuid GridGuid;
 #endif
-};
\ No newline at end of file
+};
diff --git a/Engine/Source/Runtime/Landscape/Classes/Materials/MaterialExpressionLandscapeGrassOutput.h b/Engine/Source/Runtime/Landscape/Classes/Materials/MaterialExpressionLandscapeGrassOutput.h
--- a/Engine/Source/Runtime/Landscape/Classes/Materials/MaterialExpressionLandscapeGrassOutput.h
+++ b/Engine/Source/Runtime/Landscape/Classes/Materials/MaterialExpressionLandscapeGrassOutput.h
@@ -37,8 +37,8 @@ struct FGrassInput
 	{}
 };
 
-UCLASS(collapsecategories, hidecategories=Object)
-class LANDSCAPE_API UMaterialExpressionLandscapeGrassOutput : public UMaterialExpressionCustomOutput
+UCLASS(collapsecategories, hidecategories=Object, MinimalAPI)
+class UMaterialExpressionLandscapeGrassOutput : public UMaterialExpressionCustomOutput
 {
 	GENERATED_UCLASS_BODY()
 
@@ -46,17 +46,17 @@ class LANDSCAPE_API UMaterialExpressionLandscapeGrassOutput : public UMaterialEx
 	static constexpr int32 MaxGrassTypes = 32;
 
 #if WITH_EDITOR
-	virtual int32 Compile(class FMaterialCompiler* Compiler, int32 OutputIndex) override;
-	virtual void GetCaption(TArray<FString>& OutCaptions) const override;
-	virtual TArrayView<FExpressionInput*> GetInputsView() override;
-	virtual FExpressionInput* GetInput(int32 InputIndex) override;
-	virtual FName GetInputName(int32 InputIndex) const override;
-	void ValidateInputName(FGrassInput& Input) const;
+	LANDSCAPE_API virtual int32 Compile(class FMaterialCompiler* Compiler, int32 OutputIndex) override;
+	LANDSCAPE_API virtual void GetCaption(TArray<FString>& OutCaptions) const override;
+	LANDSCAPE_API virtual TArrayView<FExpressionInput*> GetInputsView() override;
+	LANDSCAPE_API virtual FExpressionInput* GetInput(int32 InputIndex) override;
+	LANDSCAPE_API virtual FName GetInputName(int32 InputIndex) const override;
+	LANDSCAPE_API void ValidateInputName(FGrassInput& Input) const;
 #endif
 
 	//~ Begin UObject Interface
 #if WITH_EDITOR
-	virtual void PostEditChangeProperty(FPropertyChangedEvent& PropertyChangedEvent) override;
+	LANDSCAPE_API virtual void PostEditChangeProperty(FPropertyChangedEvent& PropertyChangedEvent) override;
 	//~ End UObject Interface
 
 	virtual uint32 GetInputType(int32 InputIndex) override { return MCT_Float; }
@@ -70,7 +70,7 @@ class LANDSCAPE_API UMaterialExpressionLandscapeGrassOutput : public UMaterialEx
 	TArray<FGrassInput> GrassTypes;
 
 private:
-	static FName PinDefaultName;
+	static LANDSCAPE_API FName PinDefaultName;
 };
 
 
diff --git a/Engine/Source/Runtime/Landscape/Classes/Materials/MaterialExpressionLandscapeLayerBlend.h b/Engine/Source/Runtime/Landscape/Classes/Materials/MaterialExpressionLandscapeLayerBlend.h
--- a/Engine/Source/Runtime/Landscape/Classes/Materials/MaterialExpressionLandscapeLayerBlend.h
+++ b/Engine/Source/Runtime/Landscape/Classes/Materials/MaterialExpressionLandscapeLayerBlend.h
@@ -60,8 +60,8 @@ struct FLayerBlendInput
 	}
 };
 
-UCLASS(collapsecategories, hidecategories=Object)
-class LANDSCAPE_API UMaterialExpressionLandscapeLayerBlend : public UMaterialExpression
+UCLASS(collapsecategories, hidecategories=Object, MinimalAPI)
+class UMaterialExpressionLandscapeLayerBlend : public UMaterialExpression
 {
 	GENERATED_UCLASS_BODY()
 
@@ -69,32 +69,32 @@ class LANDSCAPE_API UMaterialExpressionLandscapeLayerBlend : public UMaterialExp
 	TArray<FLayerBlendInput> Layers;
 
 	//~ Begin UObject Interface
-	virtual void Serialize(FStructuredArchive::FRecord Record) override;
+	LANDSCAPE_API virtual void Serialize(FStructuredArchive::FRecord Record) override;
 #if WITH_EDITOR
-	virtual void PostEditChangeProperty(FPropertyChangedEvent& PropertyChangedEvent) override;
+	LANDSCAPE_API virtual void PostEditChangeProperty(FPropertyChangedEvent& PropertyChangedEvent) override;
 	//~ End UObject Interface
 
 	//~ Begin UMaterialExpression Interface
-	virtual uint32 GetInputType(int32 InputIndex) override;
+	LANDSCAPE_API virtual uint32 GetInputType(int32 InputIndex) override;
 	virtual uint32 GetOutputType(int32 InputIndex) override { return MCT_Unknown; }
-	virtual bool IsResultMaterialAttributes(int32 OutputIndex) override;
-	virtual int32 Compile(class FMaterialCompiler* Compiler, int32 OutputIndex) override;
-	virtual void GetCaption(TArray<FString>& OutCaptions) const override;
-	virtual TArrayView<FExpressionInput*> GetInputsView() override;
-	virtual FExpressionInput* GetInput(int32 InputIndex) override;
-	virtual FName GetInputName(int32 InputIndex) const override;
+	LANDSCAPE_API virtual bool IsResultMaterialAttributes(int32 OutputIndex) override;
+	LANDSCAPE_API virtual int32 Compile(class FMaterialCompiler* Compiler, int32 OutputIndex) override;
+	LANDSCAPE_API virtual void GetCaption(TArray<FString>& OutCaptions) const override;
+	LANDSCAPE_API virtual TArrayView<FExpressionInput*> GetInputsView() override;
+	LANDSCAPE_API virtual FExpressionInput* GetInput(int32 InputIndex) override;
+	LANDSCAPE_API virtual FName GetInputName(int32 InputIndex) const override;
 
 	/**
 	 * Gets the landscape layer names
 	 */
-	virtual void GetLandscapeLayerNames(TArray<FName>& OutLayers) const override;
+	LANDSCAPE_API virtual void GetLandscapeLayerNames(TArray<FName>& OutLayers) const override;
 #endif
-	virtual UObject* GetReferencedTexture() const override;
+	LANDSCAPE_API virtual UObject* GetReferencedTexture() const override;
 	virtual bool CanReferenceTexture() const override { return true; }
 	//~ End UMaterialExpression Interface
 
 private:
-	int32 Compile(class FMaterialCompiler* Compiler, int32 OutputIndex, bool bForPreview);
+	LANDSCAPE_API int32 Compile(class FMaterialCompiler* Compiler, int32 OutputIndex, bool bForPreview);
 };
 
 
diff --git a/Engine/Source/Runtime/Landscape/Classes/Materials/MaterialExpressionLandscapeLayerCoords.h b/Engine/Source/Runtime/Landscape/Classes/Materials/MaterialExpressionLandscapeLayerCoords.h
--- a/Engine/Source/Runtime/Landscape/Classes/Materials/MaterialExpressionLandscapeLayerCoords.h
+++ b/Engine/Source/Runtime/Landscape/Classes/Materials/MaterialExpressionLandscapeLayerCoords.h
@@ -31,8 +31,8 @@ enum ELandscapeCustomizedCoordType : int
 	LCCT_MAX,
 };
 
-UCLASS(collapsecategories, hidecategories=Object)
-class LANDSCAPE_API UMaterialExpressionLandscapeLayerCoords : public UMaterialExpression
+UCLASS(collapsecategories, hidecategories=Object, MinimalAPI)
+class UMaterialExpressionLandscapeLayerCoords : public UMaterialExpression
 {
 	GENERATED_UCLASS_BODY()
 
@@ -62,8 +62,8 @@ class LANDSCAPE_API UMaterialExpressionLandscapeLayerCoords : public UMaterialEx
 
 	//~ Begin UMaterialExpression Interface
 #if WITH_EDITOR
-	virtual int32 Compile(class FMaterialCompiler* Compiler, int32 OutputIndex) override;
-	virtual void GetCaption(TArray<FString>& OutCaptions) const override;
+	LANDSCAPE_API virtual int32 Compile(class FMaterialCompiler* Compiler, int32 OutputIndex) override;
+	LANDSCAPE_API virtual void GetCaption(TArray<FString>& OutCaptions) const override;
 #endif
 	//~ End UMaterialExpression Interface
 };
diff --git a/Engine/Source/Runtime/Landscape/Classes/Materials/MaterialExpressionLandscapeLayerSample.h b/Engine/Source/Runtime/Landscape/Classes/Materials/MaterialExpressionLandscapeLayerSample.h
--- a/Engine/Source/Runtime/Landscape/Classes/Materials/MaterialExpressionLandscapeLayerSample.h
+++ b/Engine/Source/Runtime/Landscape/Classes/Materials/MaterialExpressionLandscapeLayerSample.h
@@ -12,8 +12,8 @@
 class UTexture;
 struct FMaterialParameterInfo;
 
-UCLASS(collapsecategories, hidecategories=Object)
-class LANDSCAPE_API UMaterialExpressionLandscapeLayerSample : public UMaterialExpression
+UCLASS(collapsecategories, hidecategories=Object, MinimalAPI)
+class UMaterialExpressionLandscapeLayerSample : public UMaterialExpression
 {
 	GENERATED_UCLASS_BODY()
 
@@ -28,18 +28,18 @@ public:
 	//~ Begin UMaterialExpression Interface
 #if WITH_EDITOR
 	virtual bool CanRenameNode() const override { return true; }
-	virtual FString GetEditableName() const override;
-	virtual void SetEditableName(const FString& NewName) override;
-	virtual int32 Compile(class FMaterialCompiler* Compiler, int32 OutputIndex) override;
-	virtual void GetCaption(TArray<FString>& OutCaptions) const override;
-	virtual bool MatchesSearchQuery(const TCHAR* SearchQuery) override;
+	LANDSCAPE_API virtual FString GetEditableName() const override;
+	LANDSCAPE_API virtual void SetEditableName(const FString& NewName) override;
+	LANDSCAPE_API virtual int32 Compile(class FMaterialCompiler* Compiler, int32 OutputIndex) override;
+	LANDSCAPE_API virtual void GetCaption(TArray<FString>& OutCaptions) const override;
+	LANDSCAPE_API virtual bool MatchesSearchQuery(const TCHAR* SearchQuery) override;
 
 	/**
 	 * Gets the landscape layer names
 	 */
-	virtual void GetLandscapeLayerNames(TArray<FName>& OutLayers) const override;
+	LANDSCAPE_API virtual void GetLandscapeLayerNames(TArray<FName>& OutLayers) const override;
 #endif
-	virtual UObject* GetReferencedTexture() const override;
+	LANDSCAPE_API virtual UObject* GetReferencedTexture() const override;
 	virtual bool CanReferenceTexture() const override { return true; }
 	//~ End UMaterialExpression Interface
 };
diff --git a/Engine/Source/Runtime/Landscape/Classes/Materials/MaterialExpressionLandscapeLayerSwitch.h b/Engine/Source/Runtime/Landscape/Classes/Materials/MaterialExpressionLandscapeLayerSwitch.h
--- a/Engine/Source/Runtime/Landscape/Classes/Materials/MaterialExpressionLandscapeLayerSwitch.h
+++ b/Engine/Source/Runtime/Landscape/Classes/Materials/MaterialExpressionLandscapeLayerSwitch.h
@@ -14,8 +14,8 @@
 class UTexture;
 struct FMaterialParameterInfo;
 
-UCLASS(collapsecategories, hidecategories=Object)
-class LANDSCAPE_API UMaterialExpressionLandscapeLayerSwitch : public UMaterialExpression
+UCLASS(collapsecategories, hidecategories=Object, MinimalAPI)
+class UMaterialExpressionLandscapeLayerSwitch : public UMaterialExpression
 {
 	GENERATED_UCLASS_BODY()
 
@@ -34,20 +34,20 @@ class LANDSCAPE_API UMaterialExpressionLandscapeLayerSwitch : public UMaterialEx
 public:
 
 	//~ Begin UObject Interface
-	virtual void Serialize(FStructuredArchive::FRecord Record) override;
+	LANDSCAPE_API virtual void Serialize(FStructuredArchive::FRecord Record) override;
 	//~ End UObject Interface
 
 	//~ Begin UMaterialExpression Interface
 #if WITH_EDITOR
-	virtual bool IsResultMaterialAttributes(int32 OutputIndex) override;
-	virtual int32 Compile(class FMaterialCompiler* Compiler, int32 OutputIndex) override;
+	LANDSCAPE_API virtual bool IsResultMaterialAttributes(int32 OutputIndex) override;
+	LANDSCAPE_API virtual int32 Compile(class FMaterialCompiler* Compiler, int32 OutputIndex) override;
 	virtual bool CanRenameNode() const override { return true; }
-	virtual FString GetEditableName() const override;
-	virtual void SetEditableName(const FString& NewName) override;
-	virtual void GetCaption(TArray<FString>& OutCaptions) const override;
-	virtual bool MatchesSearchQuery(const TCHAR* SearchQuery) override;
+	LANDSCAPE_API virtual FString GetEditableName() const override;
+	LANDSCAPE_API virtual void SetEditableName(const FString& NewName) override;
+	LANDSCAPE_API virtual void GetCaption(TArray<FString>& OutCaptions) const override;
+	LANDSCAPE_API virtual bool MatchesSearchQuery(const TCHAR* SearchQuery) override;
 #endif
-	virtual UObject* GetReferencedTexture() const override;
+	LANDSCAPE_API virtual UObject* GetReferencedTexture() const override;
 	virtual bool CanReferenceTexture() const override { return true; }
 #if WITH_EDITOR
 	virtual uint32 GetInputType(int32 InputIndex) override {return MCT_Unknown;}
@@ -56,7 +56,7 @@ public:
 	/**
 	 * Gets the landscape layer names
 	 */
-	virtual void GetLandscapeLayerNames(TArray<FName>& OutLayers) const override;
+	LANDSCAPE_API virtual void GetLandscapeLayerNames(TArray<FName>& OutLayers) const override;
 #endif
 	//~ End UMaterialExpression Interface
 
@@ -65,7 +65,7 @@ public:
 	 * Do any object-specific cleanup required immediately after loading an object,
 	 * and immediately after any undo/redo.
 	 */
-	virtual void PostLoad() override;
+	LANDSCAPE_API virtual void PostLoad() override;
 	//~ End UObject Interface
 };
 
diff --git a/Engine/Source/Runtime/Landscape/Classes/Materials/MaterialExpressionLandscapeLayerWeight.h b/Engine/Source/Runtime/Landscape/Classes/Materials/MaterialExpressionLandscapeLayerWeight.h
--- a/Engine/Source/Runtime/Landscape/Classes/Materials/MaterialExpressionLandscapeLayerWeight.h
+++ b/Engine/Source/Runtime/Landscape/Classes/Materials/MaterialExpressionLandscapeLayerWeight.h
@@ -14,8 +14,8 @@
 class UTexture;
 struct FMaterialParameterInfo;
 
-UCLASS(collapsecategories, hidecategories=Object)
-class LANDSCAPE_API UMaterialExpressionLandscapeLayerWeight : public UMaterialExpression
+UCLASS(collapsecategories, hidecategories=Object, MinimalAPI)
+class UMaterialExpressionLandscapeLayerWeight : public UMaterialExpression
 {
 	GENERATED_UCLASS_BODY()
 
@@ -39,15 +39,15 @@ public:
 
 	//~ Begin UMaterialExpression Interface
 #if WITH_EDITOR
-	virtual bool IsResultMaterialAttributes(int32 OutputIndex) override;
-	virtual int32 Compile(class FMaterialCompiler* Compiler, int32 OutputIndex) override;
+	LANDSCAPE_API virtual bool IsResultMaterialAttributes(int32 OutputIndex) override;
+	LANDSCAPE_API virtual int32 Compile(class FMaterialCompiler* Compiler, int32 OutputIndex) override;
 	virtual bool CanRenameNode() const override { return true; }
-	virtual FString GetEditableName() const override;
-	virtual void SetEditableName(const FString& NewName) override;
-	virtual void GetCaption(TArray<FString>& OutCaptions) const override;
-	virtual bool MatchesSearchQuery(const TCHAR* SearchQuery) override;
+	LANDSCAPE_API virtual FString GetEditableName() const override;
+	LANDSCAPE_API virtual void SetEditableName(const FString& NewName) override;
+	LANDSCAPE_API virtual void GetCaption(TArray<FString>& OutCaptions) const override;
+	LANDSCAPE_API virtual bool MatchesSearchQuery(const TCHAR* SearchQuery) override;
 #endif
-	virtual UObject* GetReferencedTexture() const override;
+	LANDSCAPE_API virtual UObject* GetReferencedTexture() const override;
 	virtual bool CanReferenceTexture() const override { return true; }
 #if WITH_EDITOR
 	virtual uint32 GetInputType(int32 InputIndex) override {return MCT_Float | MCT_MaterialAttributes;}
@@ -55,7 +55,7 @@ public:
 	/**
 	 * Gets the landscape layer names
 	 */
-	virtual void GetLandscapeLayerNames(TArray<FName>& OutLayers) const override;
+	LANDSCAPE_API virtual void GetLandscapeLayerNames(TArray<FName>& OutLayers) const override;
 #endif //WITH_EDITOR
 	//~ End UMaterialExpression Interface
 
@@ -64,7 +64,7 @@ public:
 	 * Do any object-specific cleanup required immediately after loading an object,
 	 * and immediately after any undo/redo.
 	 */
-	virtual void PostLoad() override;
+	LANDSCAPE_API virtual void PostLoad() override;
 	//~ End UObject Interface
 };
 
diff --git a/Engine/Source/Runtime/Landscape/Classes/Materials/MaterialExpressionLandscapePhysicalMaterialOutput.h b/Engine/Source/Runtime/Landscape/Classes/Materials/MaterialExpressionLandscapePhysicalMaterialOutput.h
--- a/Engine/Source/Runtime/Landscape/Classes/Materials/MaterialExpressionLandscapePhysicalMaterialOutput.h
+++ b/Engine/Source/Runtime/Landscape/Classes/Materials/MaterialExpressionLandscapePhysicalMaterialOutput.h
@@ -29,8 +29,8 @@ struct FPhysicalMaterialInput
  * This can be used to generate the dominant physical material for each point on a landscape.
  * Note that the use of a material output node to generate this information is optional and when a node of this type is not present we fall back on a CPU path which analyzes landscape layer data.
  */
-UCLASS(collapsecategories, hidecategories=Object)
-class LANDSCAPE_API UMaterialExpressionLandscapePhysicalMaterialOutput : public UMaterialExpressionCustomOutput
+UCLASS(collapsecategories, hidecategories=Object, MinimalAPI)
+class UMaterialExpressionLandscapePhysicalMaterialOutput : public UMaterialExpressionCustomOutput
 {
 	GENERATED_UCLASS_BODY()
 
@@ -40,16 +40,16 @@ class LANDSCAPE_API UMaterialExpressionLandscapePhysicalMaterialOutput : public
 
 #if WITH_EDITOR
 	//~ Begin UObject Interface
-	virtual void PostEditChangeProperty(FPropertyChangedEvent& PropertyChangedEvent) override;
+	LANDSCAPE_API virtual void PostEditChangeProperty(FPropertyChangedEvent& PropertyChangedEvent) override;
 	//~ End UObject Interface
 
 	//~ Begin UMaterialExpression Interface
-	virtual void GetCaption(TArray<FString>& OutCaptions) const override;
-	virtual TArrayView<FExpressionInput*> GetInputsView() override;
-	virtual FExpressionInput* GetInput(int32 InputIndex) override;
-	virtual FName GetInputName(int32 InputIndex) const override;
+	LANDSCAPE_API virtual void GetCaption(TArray<FString>& OutCaptions) const override;
+	LANDSCAPE_API virtual TArrayView<FExpressionInput*> GetInputsView() override;
+	LANDSCAPE_API virtual FExpressionInput* GetInput(int32 InputIndex) override;
+	LANDSCAPE_API virtual FName GetInputName(int32 InputIndex) const override;
 	virtual uint32 GetInputType(int32 InputIndex) override { return MCT_Float; }
-	virtual int32 Compile(class FMaterialCompiler* Compiler, int32 OutputIndex) override;
+	LANDSCAPE_API virtual int32 Compile(class FMaterialCompiler* Compiler, int32 OutputIndex) override;
 	//~ End UMaterialExpression Interface
 #endif
 
diff --git a/Engine/Source/Runtime/Landscape/Classes/Materials/MaterialExpressionLandscapeVisibilityMask.h b/Engine/Source/Runtime/Landscape/Classes/Materials/MaterialExpressionLandscapeVisibilityMask.h
--- a/Engine/Source/Runtime/Landscape/Classes/Materials/MaterialExpressionLandscapeVisibilityMask.h
+++ b/Engine/Source/Runtime/Landscape/Classes/Materials/MaterialExpressionLandscapeVisibilityMask.h
@@ -12,26 +12,26 @@
 class UTexture;
 struct FMaterialParameterInfo;
 
-UCLASS(collapseCategories, hideCategories=Object)
-class LANDSCAPE_API UMaterialExpressionLandscapeVisibilityMask : public UMaterialExpression
+UCLASS(collapseCategories, hideCategories=Object, MinimalAPI)
+class UMaterialExpressionLandscapeVisibilityMask : public UMaterialExpression
 {
 	GENERATED_UCLASS_BODY()
 
 public:
 
-	static FName ParameterName;
+	static LANDSCAPE_API FName ParameterName;
 
 	//~ Begin UMaterialExpression Interface
 #if WITH_EDITOR
-	virtual int32 Compile(class FMaterialCompiler* Compiler, int32 OutputIndex) override;
-	virtual void GetCaption(TArray<FString>& OutCaptions) const override;
+	LANDSCAPE_API virtual int32 Compile(class FMaterialCompiler* Compiler, int32 OutputIndex) override;
+	LANDSCAPE_API virtual void GetCaption(TArray<FString>& OutCaptions) const override;
 
 	/**
 	 * Gets the landscape layer names
 	 */
-	virtual void GetLandscapeLayerNames(TArray<FName>& OutLayers) const override;
+	LANDSCAPE_API virtual void GetLandscapeLayerNames(TArray<FName>& OutLayers) const override;
 #endif
-	virtual UObject* GetReferencedTexture() const override;
+	LANDSCAPE_API virtual UObject* GetReferencedTexture() const override;
 	virtual bool CanReferenceTexture() const override { return true; }
 	//~ End UMaterialExpression Interface
 };
diff --git a/Engine/Source/Runtime/Landscape/Public/LandscapeBlueprintBrushBase.h b/Engine/Source/Runtime/Landscape/Public/LandscapeBlueprintBrushBase.h
--- a/Engine/Source/Runtime/Landscape/Public/LandscapeBlueprintBrushBase.h
+++ b/Engine/Source/Runtime/Landscape/Public/LandscapeBlueprintBrushBase.h
@@ -12,7 +12,7 @@
 class UTextureRenderTarget2D;
 
 USTRUCT(BlueprintType)
-struct LANDSCAPE_API FLandscapeBrushParameters
+struct FLandscapeBrushParameters
 {
 	GENERATED_BODY()
 
@@ -39,8 +39,8 @@ struct LANDSCAPE_API FLandscapeBrushParameters
 };
 
 
-UCLASS(Abstract, NotBlueprintable)
-class LANDSCAPE_API ALandscapeBlueprintBrushBase : public AActor
+UCLASS(Abstract, NotBlueprintable, MinimalAPI)
+class ALandscapeBlueprintBrushBase : public AActor
 {
 	GENERATED_UCLASS_BODY()
 
@@ -77,40 +77,40 @@ public:
 	virtual void Initialize_Native(const FTransform& InLandscapeTransform, const FIntPoint& InLandscapeSize, const FIntPoint& InLandscapeRenderTargetSize) {}
 
 	UFUNCTION(BlueprintNativeEvent, meta = (DeprecatedFunction, DeprecationMessage = "Please use RenderLayer instead."))
-	UTextureRenderTarget2D* Render(bool InIsHeightmap, UTextureRenderTarget2D* InCombinedResult, const FName& InWeightmapLayerName);
+	LANDSCAPE_API UTextureRenderTarget2D* Render(bool InIsHeightmap, UTextureRenderTarget2D* InCombinedResult, const FName& InWeightmapLayerName);
 
 	UFUNCTION(BlueprintNativeEvent)
-	UTextureRenderTarget2D* RenderLayer(const FLandscapeBrushParameters& InParameters);
+	LANDSCAPE_API UTextureRenderTarget2D* RenderLayer(const FLandscapeBrushParameters& InParameters);
 
-	virtual UTextureRenderTarget2D* RenderLayer_Native(const FLandscapeBrushParameters& InParameters);
+	LANDSCAPE_API virtual UTextureRenderTarget2D* RenderLayer_Native(const FLandscapeBrushParameters& InParameters);
 
 	UFUNCTION(BlueprintNativeEvent)
-	void Initialize(const FTransform& InLandscapeTransform, const FIntPoint& InLandscapeSize, const FIntPoint& InLandscapeRenderTargetSize);
+	LANDSCAPE_API void Initialize(const FTransform& InLandscapeTransform, const FIntPoint& InLandscapeSize, const FIntPoint& InLandscapeRenderTargetSize);
 
 	UFUNCTION(BlueprintCallable, Category = "Landscape")
-	void RequestLandscapeUpdate(bool bInUserTriggered = false);
+	LANDSCAPE_API void RequestLandscapeUpdate(bool bInUserTriggered = false);
 
 	UFUNCTION(BlueprintImplementableEvent, CallInEditor)
-	void GetBlueprintRenderDependencies(TArray<UObject*>& OutStreamableAssets);
+	LANDSCAPE_API void GetBlueprintRenderDependencies(TArray<UObject*>& OutStreamableAssets);
 
-	void SetCanAffectHeightmap(bool bInCanAffectHeightmap);
-	void SetCanAffectWeightmap(bool bInCanAffectWeightmap);
-	void SetCanAffectVisibilityLayer(bool bInCanAffectVisibilityLayer);
+	LANDSCAPE_API void SetCanAffectHeightmap(bool bInCanAffectHeightmap);
+	LANDSCAPE_API void SetCanAffectWeightmap(bool bInCanAffectWeightmap);
+	LANDSCAPE_API void SetCanAffectVisibilityLayer(bool bInCanAffectVisibilityLayer);
 
 #if WITH_EDITOR
-	virtual void CheckForErrors() override;
+	LANDSCAPE_API virtual void CheckForErrors() override;
 
-	virtual void GetRenderDependencies(TSet<UObject*>& OutDependencies);
+	LANDSCAPE_API virtual void GetRenderDependencies(TSet<UObject*>& OutDependencies);
 
-	virtual void SetOwningLandscape(class ALandscape* InOwningLandscape);
-	class ALandscape* GetOwningLandscape() const;
+	LANDSCAPE_API virtual void SetOwningLandscape(class ALandscape* InOwningLandscape);
+	LANDSCAPE_API class ALandscape* GetOwningLandscape() const;
 
 	UE_DEPRECATED(5.3, "Renamed CanAffectHeightmap")
 	bool IsAffectingHeightmap() const { return AffectHeightmap; }
 	UE_DEPRECATED(5.3, "Renamed CanAffectWeightmap")
 	bool IsAffectingWeightmap() const { return AffectWeightmap; }
 	UE_DEPRECATED(5.3, "Renamed AffectsVisibilityLayer")
-	virtual bool IsAffectingWeightmapLayer(const FName& InLayerName) const;
+	LANDSCAPE_API virtual bool IsAffectingWeightmapLayer(const FName& InLayerName) const;
 	UE_DEPRECATED(5.3, "Renamed CanAffectVisibilityLayer")
 	bool IsAffectingVisibilityLayer() const { return AffectVisibilityLayer; }
 
@@ -120,29 +120,29 @@ public:
 
 	virtual bool AffectsHeightmap() const { return CanAffectHeightmap(); }
 	virtual bool AffectsWeightmap() const { return CanAffectWeightmap(); }
-	virtual bool AffectsWeightmapLayer(const FName& InLayerName) const;
+	LANDSCAPE_API virtual bool AffectsWeightmapLayer(const FName& InLayerName) const;
 	virtual bool AffectsVisibilityLayer() const { return CanAffectVisibilityLayer(); }
 
 	bool IsVisible() const { return bIsVisible; }
-	bool IsLayerUpdatePending() const;
+	LANDSCAPE_API bool IsLayerUpdatePending() const;
 
-	void SetIsVisible(bool bInIsVisible);
+	LANDSCAPE_API void SetIsVisible(bool bInIsVisible);
 
 	UE_DEPRECATED(5.3, "Renamed SetCanAffectHeightmap")
-	void SetAffectsHeightmap(bool bInAffectsHeightmap);
+	LANDSCAPE_API void SetAffectsHeightmap(bool bInAffectsHeightmap);
 	UE_DEPRECATED(5.3, "Renamed SetCanAffectWeightmap")
-	void SetAffectsWeightmap(bool bInAffectsWeightmap);
+	LANDSCAPE_API void SetAffectsWeightmap(bool bInAffectsWeightmap);
 	UE_DEPRECATED(5.3, "Renamed SetCanAffectVisibilityLayer")
-	void SetAffectsVisibilityLayer(bool bInAffectsVisibilityLayer);
+	LANDSCAPE_API void SetAffectsVisibilityLayer(bool bInAffectsVisibilityLayer);
 
-	virtual bool ShouldTickIfViewportsOnly() const override;
-	virtual void Tick(float DeltaSeconds) override;
-	virtual void PostEditMove(bool bFinished) override;
-	virtual void PostEditChangeProperty(FPropertyChangedEvent& PropertyChangedEvent) override;
-	virtual void Destroyed() override;
+	LANDSCAPE_API virtual bool ShouldTickIfViewportsOnly() const override;
+	LANDSCAPE_API virtual void Tick(float DeltaSeconds) override;
+	LANDSCAPE_API virtual void PostEditMove(bool bFinished) override;
+	LANDSCAPE_API virtual void PostEditChangeProperty(FPropertyChangedEvent& PropertyChangedEvent) override;
+	LANDSCAPE_API virtual void Destroyed() override;
 
-	virtual void PushDeferredLayersContentUpdate();
+	LANDSCAPE_API virtual void PushDeferredLayersContentUpdate();
 
 	virtual bool CanChangeIsSpatiallyLoadedFlag() const override { return false; }
 #endif
-};
\ No newline at end of file
+};
diff --git a/Engine/Source/Runtime/Landscape/Public/LandscapeConfigHelper.h b/Engine/Source/Runtime/Landscape/Public/LandscapeConfigHelper.h
--- a/Engine/Source/Runtime/Landscape/Public/LandscapeConfigHelper.h
+++ b/Engine/Source/Runtime/Landscape/Public/LandscapeConfigHelper.h
@@ -25,7 +25,7 @@ struct FLandscapeImportLayerInfo;
 /**
  * 
  */
-struct LANDSCAPE_API FLandscapeConfig
+struct FLandscapeConfig
 {
 	FLandscapeConfig(int32 InComponentNumSubSections, int32 InSubsectionSizeQuads, int32 InGridSizeInComponents)
 		: ComponentNumSubsections(InComponentNumSubSections)
@@ -34,7 +34,7 @@ struct LANDSCAPE_API FLandscapeConfig
 	{
 	}
 
-	FLandscapeConfig(ULandscapeInfo* InLandscapeInfo);
+	LANDSCAPE_API FLandscapeConfig(ULandscapeInfo* InLandscapeInfo);
 
 	int32 GetComponentSizeQuads() const { return SubsectionSizeQuads * ComponentNumSubsections; }
 	int32 GetComponentSizeVerts() const { return (SubsectionSizeQuads + 1) * ComponentNumSubsections; }
@@ -44,11 +44,11 @@ struct LANDSCAPE_API FLandscapeConfig
 	int32 SubsectionSizeQuads;
 	int32 GridSizeInComponents;
 
-	static int32 NumSectionValues[2];
-	static int32 SubsectionSizeQuadsValues[6];
+	static LANDSCAPE_API int32 NumSectionValues[2];
+	static LANDSCAPE_API int32 SubsectionSizeQuadsValues[6];
 };
 
-struct LANDSCAPE_API FLandscapeConfigChange : public FLandscapeConfig
+struct FLandscapeConfigChange : public FLandscapeConfig
 {
 	FLandscapeConfigChange(int32 InComponentNumSubSections, int32 InSubsectionSizeQuads, int32 InGridSize, ELandscapeResizeMode InResizeMode, bool bInZeroBased)
 		: FLandscapeConfig(InComponentNumSubSections, InSubsectionSizeQuads, InGridSize)
@@ -57,7 +57,7 @@ struct LANDSCAPE_API FLandscapeConfigChange : public FLandscapeConfig
 	{
 	}
 
-	bool Validate() const;
+	LANDSCAPE_API bool Validate() const;
 
 	ELandscapeResizeMode ResizeMode;
 	bool bZeroBased;
@@ -193,4 +193,4 @@ public:
 	}
 };
 
-#endif
\ No newline at end of file
+#endif
diff --git a/Engine/Source/Runtime/Landscape/Public/LandscapeDataAccess.h b/Engine/Source/Runtime/Landscape/Public/LandscapeDataAccess.h
--- a/Engine/Source/Runtime/Landscape/Public/LandscapeDataAccess.h
+++ b/Engine/Source/Runtime/Landscape/Public/LandscapeDataAccess.h
@@ -136,7 +136,7 @@ private:
 
 struct FLandscapeComponentDataInterfaceBase
 {
-	LANDSCAPE_API FLandscapeComponentDataInterfaceBase() {}
+	FLandscapeComponentDataInterfaceBase() {}
 	LANDSCAPE_API FLandscapeComponentDataInterfaceBase(ULandscapeComponent* InComponent, int32 InMipLevel, bool InWorkOnEditingLayer = true);
 
 	// Accessors
diff --git a/Engine/Source/Runtime/Landscape/Public/LandscapeEdit.h b/Engine/Source/Runtime/Landscape/Public/LandscapeEdit.h
--- a/Engine/Source/Runtime/Landscape/Public/LandscapeEdit.h
+++ b/Engine/Source/Runtime/Landscape/Public/LandscapeEdit.h
@@ -115,74 +115,74 @@ private:
 	TArray<FMipInfo> MipInfo;
 };
 
-struct LANDSCAPE_API FLandscapeTextureDataInterface
+struct FLandscapeTextureDataInterface
 {
-	FLandscapeTextureDataInterface(bool bInUploadTextureChangesToGPU = true);
-	virtual ~FLandscapeTextureDataInterface();
+	LANDSCAPE_API FLandscapeTextureDataInterface(bool bInUploadTextureChangesToGPU = true);
+	LANDSCAPE_API virtual ~FLandscapeTextureDataInterface();
 		
 	void SetShouldDirtyPackage(bool bValue) { bShouldDirtyPackage = bValue; }
 	bool GetShouldDirtyPackage() const { return bShouldDirtyPackage; }
 
 	// Texture data access
-	FLandscapeTextureDataInfo* GetTextureDataInfo(UTexture2D* Texture);
+	LANDSCAPE_API FLandscapeTextureDataInfo* GetTextureDataInfo(UTexture2D* Texture);
 
 	// Flush texture updates
-	void Flush();
+	LANDSCAPE_API void Flush();
 
 	// Texture bulk operations for weightmap reallocation
-	void CopyTextureChannel(UTexture2D* Dest, int32 DestChannel, UTexture2D* Src, int32 SrcChannel);
-	void ZeroTextureChannel(UTexture2D* Dest, int32 DestChannel);
-	void CopyTextureFromHeightmap(UTexture2D* Dest, ULandscapeComponent* Comp, int32 MipIndex);
-	void CopyTextureFromHeightmap(UTexture2D* Dest, int32 DestChannel, ULandscapeComponent* Comp, int32 SrcChannel);
-	void CopyTextureFromWeightmap(UTexture2D* Dest, int32 DestChannel, ULandscapeComponent* Comp, ULandscapeLayerInfoObject* LayerInfo);
-	void CopyTextureFromWeightmap(UTexture2D* Dest, int32 DestChannel, ULandscapeComponent* Comp, ULandscapeLayerInfoObject* LayerInfo, int32 MipIndex);
+	LANDSCAPE_API void CopyTextureChannel(UTexture2D* Dest, int32 DestChannel, UTexture2D* Src, int32 SrcChannel);
+	LANDSCAPE_API void ZeroTextureChannel(UTexture2D* Dest, int32 DestChannel);
+	LANDSCAPE_API void CopyTextureFromHeightmap(UTexture2D* Dest, ULandscapeComponent* Comp, int32 MipIndex);
+	LANDSCAPE_API void CopyTextureFromHeightmap(UTexture2D* Dest, int32 DestChannel, ULandscapeComponent* Comp, int32 SrcChannel);
+	LANDSCAPE_API void CopyTextureFromWeightmap(UTexture2D* Dest, int32 DestChannel, ULandscapeComponent* Comp, ULandscapeLayerInfoObject* LayerInfo);
+	LANDSCAPE_API void CopyTextureFromWeightmap(UTexture2D* Dest, int32 DestChannel, ULandscapeComponent* Comp, ULandscapeLayerInfoObject* LayerInfo, int32 MipIndex);
 
 	template<typename TData>
 	void SetTextureValueTempl(UTexture2D* Dest, TData Value);
-	void ZeroTexture(UTexture2D* Dest);
-	void SetTextureValue(UTexture2D* Dest, FColor Value);
+	LANDSCAPE_API void ZeroTexture(UTexture2D* Dest);
+	LANDSCAPE_API void SetTextureValue(UTexture2D* Dest, FColor Value);
 
 	template<typename TData>
 	bool EqualTextureValueTempl(UTexture2D* Src, TData Value);
-	bool EqualTextureValue(UTexture2D* Src, FColor Value);
+	LANDSCAPE_API bool EqualTextureValue(UTexture2D* Src, FColor Value);
 
 private:
-	void CopyTextureFromWeightmap(FLandscapeTextureDataInfo* DestDataInfo, int32 DestChannel, ULandscapeComponent* Comp, ULandscapeLayerInfoObject* LayerInfo, int32 MipIndex);
+	LANDSCAPE_API void CopyTextureFromWeightmap(FLandscapeTextureDataInfo* DestDataInfo, int32 DestChannel, ULandscapeComponent* Comp, ULandscapeLayerInfoObject* LayerInfo, int32 MipIndex);
 
 	TMap<UTexture2D*, FLandscapeTextureDataInfo*> TextureDataMap;
 	bool bUploadTextureChangesToGPU;
 	bool bShouldDirtyPackage;
 };
 
-struct LANDSCAPE_API FLandscapeEditDataInterface : public FLandscapeTextureDataInterface
+struct FLandscapeEditDataInterface : public FLandscapeTextureDataInterface
 {
 	// this constructor will build an interface that works in the current edit layer (uses the ALandscape PrivateEditingLayer state)
-	FLandscapeEditDataInterface(ULandscapeInfo* InLandscape, bool bInUploadTextureChangesToGPU = true);
+	LANDSCAPE_API FLandscapeEditDataInterface(ULandscapeInfo* InLandscape, bool bInUploadTextureChangesToGPU = true);
 
 	// this constructor will build an interface that works in the specified edit layer : InEditLayerGUID
-	FLandscapeEditDataInterface(ULandscapeInfo* InLandscape, const FGuid& InEditLayerGUID, bool bInUploadTextureChangesToGPU = true);
+	LANDSCAPE_API FLandscapeEditDataInterface(ULandscapeInfo* InLandscape, const FGuid& InEditLayerGUID, bool bInUploadTextureChangesToGPU = true);
 
 	// switch the layer this interface works in
-	void SetEditLayer(const FGuid& InEditLayerGUID);
+	LANDSCAPE_API void SetEditLayer(const FGuid& InEditLayerGUID);
 
 	// returns the current layer this interface is working in
-	FGuid GetEditLayer() const;
+	LANDSCAPE_API FGuid GetEditLayer() const;
 
 	// Misc
-	bool GetComponentsInRegion(int32 X1, int32 Y1, int32 X2, int32 Y2, TSet<ULandscapeComponent*>* OutComponents = NULL);
+	LANDSCAPE_API bool GetComponentsInRegion(int32 X1, int32 Y1, int32 X2, int32 Y2, TSet<ULandscapeComponent*>* OutComponents = NULL);
 
 	//
 	// Heightmap access
 	//
-	void SetHeightData(int32 X1, int32 Y1, int32 X2, int32 Y2, const uint16* InData, int32 InStride, bool InCalcNormals, const uint16* InNormalData = nullptr, const uint16* InHeightAlphaBlendData = nullptr, const uint8* InHeightRaiseLowerData = nullptr, bool InCreateComponents = false, UTexture2D* InHeightmap = nullptr, UTexture2D* InXYOffsetmapTexture = nullptr,
+	LANDSCAPE_API void SetHeightData(int32 X1, int32 Y1, int32 X2, int32 Y2, const uint16* InData, int32 InStride, bool InCalcNormals, const uint16* InNormalData = nullptr, const uint16* InHeightAlphaBlendData = nullptr, const uint8* InHeightRaiseLowerData = nullptr, bool InCreateComponents = false, UTexture2D* InHeightmap = nullptr, UTexture2D* InXYOffsetmapTexture = nullptr,
 					   bool InUpdateBounds = true, bool InUpdateCollision = true, bool InGenerateMips = true);
 
 	// Helper accessor
-	FORCEINLINE uint16 GetHeightMapData(const ULandscapeComponent* Component, int32 TexU, int32 TexV, FColor* TextureData = NULL);
+	LANDSCAPE_API FORCEINLINE uint16 GetHeightMapData(const ULandscapeComponent* Component, int32 TexU, int32 TexV, FColor* TextureData = NULL);
 	// Helper accessor
-	FORCEINLINE uint16 GetHeightMapAlphaBlendData(const ULandscapeComponent* Component, int32 TexU, int32 TexV, FColor* TextureData = NULL);
+	LANDSCAPE_API FORCEINLINE uint16 GetHeightMapAlphaBlendData(const ULandscapeComponent* Component, int32 TexU, int32 TexV, FColor* TextureData = NULL);
 	// Helper accessor
-	FORCEINLINE uint8 GetHeightMapFlagsData(const ULandscapeComponent* Component, int32 TexU, int32 TexV, FColor* TextureData = NULL);
+	LANDSCAPE_API FORCEINLINE uint8 GetHeightMapFlagsData(const ULandscapeComponent* Component, int32 TexU, int32 TexV, FColor* TextureData = NULL);
 	// Generic
 	template<typename TStoreData>
 	void GetHeightDataTempl(int32& X1, int32& Y1, int32& X2, int32& Y2, TStoreData& StoreData);
@@ -194,106 +194,106 @@ struct LANDSCAPE_API FLandscapeEditDataInterface : public FLandscapeTextureDataI
 	template<typename TStoreData>
 	void GetHeightDataTemplFast(const int32 X1, const int32 Y1, const int32 X2, const int32 Y2, TStoreData& StoreData, UTexture2D* InHeightmap = nullptr, TStoreData* NormalData = NULL);
 	// Implementation for fixed array
-	void GetHeightData(int32& X1, int32& Y1, int32& X2, int32& Y2, uint16* Data, int32 Stride);
-	void GetHeightDataFast(const int32 X1, const int32 Y1, const int32 X2, const int32 Y2, uint16* Data, int32 Stride, uint16* NormalData = NULL, UTexture2D* InHeightmap = nullptr);
-	void GetHeightAlphaBlendData(int32& X1, int32& Y1, int32& X2, int32& Y2, uint16* Data, int32 Stride);
-	void GetHeightFlagsData(int32& X1, int32& Y1, int32& X2, int32& Y2, uint8* Data, int32 Stride);
+	LANDSCAPE_API void GetHeightData(int32& X1, int32& Y1, int32& X2, int32& Y2, uint16* Data, int32 Stride);
+	LANDSCAPE_API void GetHeightDataFast(const int32 X1, const int32 Y1, const int32 X2, const int32 Y2, uint16* Data, int32 Stride, uint16* NormalData = NULL, UTexture2D* InHeightmap = nullptr);
+	LANDSCAPE_API void GetHeightAlphaBlendData(int32& X1, int32& Y1, int32& X2, int32& Y2, uint16* Data, int32 Stride);
+	LANDSCAPE_API void GetHeightFlagsData(int32& X1, int32& Y1, int32& X2, int32& Y2, uint8* Data, int32 Stride);
 	// Implementation for sparse array
-	void GetHeightData(int32& X1, int32& Y1, int32& X2, int32& Y2, TMap<FIntPoint, uint16>& SparseData);
-	void GetHeightDataFast(const int32 X1, const int32 Y1, const int32 X2, const int32 Y2, TMap<FIntPoint, uint16>& SparseData, TMap<FIntPoint, uint16>* NormalData = NULL, UTexture2D* InHeightmap = nullptr);
+	LANDSCAPE_API void GetHeightData(int32& X1, int32& Y1, int32& X2, int32& Y2, TMap<FIntPoint, uint16>& SparseData);
+	LANDSCAPE_API void GetHeightDataFast(const int32 X1, const int32 Y1, const int32 X2, const int32 Y2, TMap<FIntPoint, uint16>& SparseData, TMap<FIntPoint, uint16>* NormalData = NULL, UTexture2D* InHeightmap = nullptr);
 
 	// Recaclulate normals for the entire landscape.
-	void RecalculateNormals();
+	LANDSCAPE_API void RecalculateNormals();
 
 	//
 	// Weightmap access
 	//
 	// Helper accessor
-	FORCEINLINE uint8 GetWeightMapData(const ULandscapeComponent* Component, ULandscapeLayerInfoObject* LayerInfo, int32 TexU, int32 TexV, uint8 Offset = 0, UTexture2D* Texture = NULL, uint8* TextureData = NULL);
+	LANDSCAPE_API FORCEINLINE uint8 GetWeightMapData(const ULandscapeComponent* Component, ULandscapeLayerInfoObject* LayerInfo, int32 TexU, int32 TexV, uint8 Offset = 0, UTexture2D* Texture = NULL, uint8* TextureData = NULL);
 	template<typename TStoreData>
 	void GetWeightDataTempl(ULandscapeLayerInfoObject* LayerInfo, int32& X1, int32& Y1, int32& X2, int32& Y2, TStoreData& StoreData);
 	// Without data interpolation
 	template<typename TStoreData>
 	void GetWeightDataTemplFast(ULandscapeLayerInfoObject* LayerInfo, const int32 X1, const int32 Y1, const int32 X2, const int32 Y2, TStoreData& StoreData);
 	// Implementation for fixed array
-	void GetWeightData(ULandscapeLayerInfoObject* LayerInfo, int32& X1, int32& Y1, int32& X2, int32& Y2, uint8* Data, int32 Stride);
+	LANDSCAPE_API void GetWeightData(ULandscapeLayerInfoObject* LayerInfo, int32& X1, int32& Y1, int32& X2, int32& Y2, uint8* Data, int32 Stride);
 	//void GetWeightData(FName LayerName, int32& X1, int32& Y1, int32& X2, int32& Y2, TArray<uint8>* Data, int32 Stride);
-	void GetWeightDataFast(ULandscapeLayerInfoObject* LayerInfo, const int32 X1, const int32 Y1, const int32 X2, const int32 Y2, uint8* Data, int32 Stride);
-	void GetWeightDataFast(ULandscapeLayerInfoObject* LayerInfo, const int32 X1, const int32 Y1, const int32 X2, const int32 Y2, TArray<uint8>* Data, int32 Stride);
+	LANDSCAPE_API void GetWeightDataFast(ULandscapeLayerInfoObject* LayerInfo, const int32 X1, const int32 Y1, const int32 X2, const int32 Y2, uint8* Data, int32 Stride);
+	LANDSCAPE_API void GetWeightDataFast(ULandscapeLayerInfoObject* LayerInfo, const int32 X1, const int32 Y1, const int32 X2, const int32 Y2, TArray<uint8>* Data, int32 Stride);
 	// Implementation for sparse array
-	void GetWeightData(ULandscapeLayerInfoObject* LayerInfo, int32& X1, int32& Y1, int32& X2, int32& Y2, TMap<FIntPoint, uint8>& SparseData);
+	LANDSCAPE_API void GetWeightData(ULandscapeLayerInfoObject* LayerInfo, int32& X1, int32& Y1, int32& X2, int32& Y2, TMap<FIntPoint, uint8>& SparseData);
 	//void GetWeightData(FName LayerName, int32& X1, int32& Y1, int32& X2, int32& Y2, TMap<uint64, TArray<uint8>>& SparseData);
-	void GetWeightDataFast(ULandscapeLayerInfoObject* LayerInfo, const int32 X1, const int32 Y1, const int32 X2, const int32 Y2, TMap<FIntPoint, uint8>& SparseData);
-	void GetWeightDataFast(ULandscapeLayerInfoObject* LayerInfo, const int32 X1, const int32 Y1, const int32 X2, const int32 Y2, TMap<FIntPoint, TArray<uint8>>& SparseData);
+	LANDSCAPE_API void GetWeightDataFast(ULandscapeLayerInfoObject* LayerInfo, const int32 X1, const int32 Y1, const int32 X2, const int32 Y2, TMap<FIntPoint, uint8>& SparseData);
+	LANDSCAPE_API void GetWeightDataFast(ULandscapeLayerInfoObject* LayerInfo, const int32 X1, const int32 Y1, const int32 X2, const int32 Y2, TMap<FIntPoint, TArray<uint8>>& SparseData);
 	// Updates weightmap for LayerInfo, optionally adjusting all other weightmaps.
-	void SetAlphaData(ULandscapeLayerInfoObject* LayerInfo, const int32 X1, const int32 Y1, const int32 X2, const int32 Y2, const uint8* Data, int32 Stride, ELandscapeLayerPaintingRestriction PaintingRestriction = ELandscapeLayerPaintingRestriction::None, bool bWeightAdjust = true, bool bTotalWeightAdjust = false);
+	LANDSCAPE_API void SetAlphaData(ULandscapeLayerInfoObject* LayerInfo, const int32 X1, const int32 Y1, const int32 X2, const int32 Y2, const uint8* Data, int32 Stride, ELandscapeLayerPaintingRestriction PaintingRestriction = ELandscapeLayerPaintingRestriction::None, bool bWeightAdjust = true, bool bTotalWeightAdjust = false);
 	// Updates weightmaps for all layers. Data points to packed data for all layers in the landscape info
-	void SetAlphaData(const TSet<ULandscapeLayerInfoObject*>& DirtyLayerInfos, const int32 X1, const int32 Y1, const int32 X2, const int32 Y2, const uint8* Data, int32 Stride, ELandscapeLayerPaintingRestriction PaintingRestriction = ELandscapeLayerPaintingRestriction::None);
+	LANDSCAPE_API void SetAlphaData(const TSet<ULandscapeLayerInfoObject*>& DirtyLayerInfos, const int32 X1, const int32 Y1, const int32 X2, const int32 Y2, const uint8* Data, int32 Stride, ELandscapeLayerPaintingRestriction PaintingRestriction = ELandscapeLayerPaintingRestriction::None);
 	// Delete a layer and re-normalize other layers
-	void DeleteLayer(ULandscapeLayerInfoObject* LayerInfo);
+	LANDSCAPE_API void DeleteLayer(ULandscapeLayerInfoObject* LayerInfo);
 	// Fill a layer and re-normalize other layers
-	void FillLayer(ULandscapeLayerInfoObject* LayerInfo);
+	LANDSCAPE_API void FillLayer(ULandscapeLayerInfoObject* LayerInfo);
 	// Fill all empty layers and re-normalize layers
-	void FillEmptyLayers(ULandscapeLayerInfoObject* LayerInfo);
+	LANDSCAPE_API void FillEmptyLayers(ULandscapeLayerInfoObject* LayerInfo);
 	
 	// Replace/merge a layer (across all edit layers and the main runtime layer)
-	void ReplaceLayer(ULandscapeLayerInfoObject* FromLayerInfo, ULandscapeLayerInfoObject* ToLayerInfo);
+	LANDSCAPE_API void ReplaceLayer(ULandscapeLayerInfoObject* FromLayerInfo, ULandscapeLayerInfoObject* ToLayerInfo);
 
 	template<typename TStoreData>
 	void GetEditToolTextureData(const int32 X1, const int32 Y1, const int32 X2, const int32 Y2, TStoreData& StoreData, TFunctionRef<UTexture2D*(ULandscapeComponent*)> GetComponentTexture);
-	void SetEditToolTextureData(int32 X1, int32 Y1, int32 X2, int32 Y2, const uint8* Data, int32 Stride, TFunctionRef<UTexture2D*&(ULandscapeComponent*)> GetComponentTexture, TextureGroup InTextureGroup = TEXTUREGROUP_Terrain_Weightmap);
+	LANDSCAPE_API void SetEditToolTextureData(int32 X1, int32 Y1, int32 X2, int32 Y2, const uint8* Data, int32 Stride, TFunctionRef<UTexture2D*&(ULandscapeComponent*)> GetComponentTexture, TextureGroup InTextureGroup = TEXTUREGROUP_Terrain_Weightmap);
 
 	// Without data interpolation, Select Data 
 	template<typename TStoreData>
 	void GetSelectDataTempl(const int32 X1, const int32 Y1, const int32 X2, const int32 Y2, TStoreData& StoreData);
-	void GetSelectData(const int32 X1, const int32 Y1, const int32 X2, const int32 Y2, uint8* Data, int32 Stride);
-	void GetSelectData(const int32 X1, const int32 Y1, const int32 X2, const int32 Y2, TMap<FIntPoint, uint8>& SparseData);
-	void SetSelectData(const int32 X1, const int32 Y1, const int32 X2, const int32 Y2, const uint8* Data, int32 Stride);
+	LANDSCAPE_API void GetSelectData(const int32 X1, const int32 Y1, const int32 X2, const int32 Y2, uint8* Data, int32 Stride);
+	LANDSCAPE_API void GetSelectData(const int32 X1, const int32 Y1, const int32 X2, const int32 Y2, TMap<FIntPoint, uint8>& SparseData);
+	LANDSCAPE_API void SetSelectData(const int32 X1, const int32 Y1, const int32 X2, const int32 Y2, const uint8* Data, int32 Stride);
 	
 	template<typename TStoreData>
 	void GetLayerContributionDataTempl(const int32 X1, const int32 Y1, const int32 X2, const int32 Y2, TStoreData& StoreData);
-	void GetLayerContributionData(const int32 X1, const int32 Y1, const int32 X2, const int32 Y2, uint8* Data, int32 Stride);
-	void GetLayerContributionData(const int32 X1, const int32 Y1, const int32 X2, const int32 Y2, TMap<FIntPoint, uint8>& SparseData);
-	void SetLayerContributionData(const int32 X1, const int32 Y1, const int32 X2, const int32 Y2, const uint8* Data, int32 Stride);
+	LANDSCAPE_API void GetLayerContributionData(const int32 X1, const int32 Y1, const int32 X2, const int32 Y2, uint8* Data, int32 Stride);
+	LANDSCAPE_API void GetLayerContributionData(const int32 X1, const int32 Y1, const int32 X2, const int32 Y2, TMap<FIntPoint, uint8>& SparseData);
+	LANDSCAPE_API void SetLayerContributionData(const int32 X1, const int32 Y1, const int32 X2, const int32 Y2, const uint8* Data, int32 Stride);
 
 	template<typename TStoreData>
 	void GetDirtyDataTempl(const int32 X1, const int32 Y1, const int32 X2, const int32 Y2, TStoreData& StoreData);
-	void GetDirtyData(const int32 X1, const int32 Y1, const int32 X2, const int32 Y2, uint8* Data, int32 Stride);
-	void GetDirtyData(const int32 X1, const int32 Y1, const int32 X2, const int32 Y2, TMap<FIntPoint, uint8>& SparseData);
-	void SetDirtyData(const int32 X1, const int32 Y1, const int32 X2, const int32 Y2, const uint8* Data, int32 Stride);
+	LANDSCAPE_API void GetDirtyData(const int32 X1, const int32 Y1, const int32 X2, const int32 Y2, uint8* Data, int32 Stride);
+	LANDSCAPE_API void GetDirtyData(const int32 X1, const int32 Y1, const int32 X2, const int32 Y2, TMap<FIntPoint, uint8>& SparseData);
+	LANDSCAPE_API void SetDirtyData(const int32 X1, const int32 Y1, const int32 X2, const int32 Y2, const uint8* Data, int32 Stride);
 
 	//
 	// XYOffsetmap access
 	//
 	template<typename T>
 	void SetXYOffsetDataTempl(int32 X1, int32 Y1, int32 X2, int32 Y2, const T* Data, int32 Stride);
-	void SetXYOffsetData(int32 X1, int32 Y1, int32 X2, int32 Y2, const FVector2D* Data, int32 Stride);
-	void SetXYOffsetData(int32 X1, int32 Y1, int32 X2, int32 Y2, const FVector* Data, int32 Stride);
+	LANDSCAPE_API void SetXYOffsetData(int32 X1, int32 Y1, int32 X2, int32 Y2, const FVector2D* Data, int32 Stride);
+	LANDSCAPE_API void SetXYOffsetData(int32 X1, int32 Y1, int32 X2, int32 Y2, const FVector* Data, int32 Stride);
 	// Helper accessor
-	FORCEINLINE FVector2D GetXYOffsetmapData(const ULandscapeComponent* Component, int32 TexU, int32 TexV, FColor* TextureData = NULL);
+	LANDSCAPE_API FORCEINLINE FVector2D GetXYOffsetmapData(const ULandscapeComponent* Component, int32 TexU, int32 TexV, FColor* TextureData = NULL);
 
 	template<typename TStoreData>
 	void GetXYOffsetDataTempl(int32& X1, int32& Y1, int32& X2, int32& Y2, TStoreData& StoreData);
-	void GetXYOffsetData(int32& X1, int32& Y1, int32& X2, int32& Y2, FVector2D* Data, int32 Stride);
-	void GetXYOffsetData(int32& X1, int32& Y1, int32& X2, int32& Y2, TMap<FIntPoint, FVector2D>& SparseData);
-	void GetXYOffsetData(int32& X1, int32& Y1, int32& X2, int32& Y2, FVector* Data, int32 Stride);
-	void GetXYOffsetData(int32& X1, int32& Y1, int32& X2, int32& Y2, TMap<FIntPoint, FVector>& SparseData);
+	LANDSCAPE_API void GetXYOffsetData(int32& X1, int32& Y1, int32& X2, int32& Y2, FVector2D* Data, int32 Stride);
+	LANDSCAPE_API void GetXYOffsetData(int32& X1, int32& Y1, int32& X2, int32& Y2, TMap<FIntPoint, FVector2D>& SparseData);
+	LANDSCAPE_API void GetXYOffsetData(int32& X1, int32& Y1, int32& X2, int32& Y2, FVector* Data, int32 Stride);
+	LANDSCAPE_API void GetXYOffsetData(int32& X1, int32& Y1, int32& X2, int32& Y2, TMap<FIntPoint, FVector>& SparseData);
 	// Without data interpolation
 	template<typename TStoreData>
 	void GetXYOffsetDataTemplFast(const int32 X1, const int32 Y1, const int32 X2, const int32 Y2, TStoreData& StoreData);
 	// Without data interpolation, able to get normal data
-	void GetXYOffsetDataFast(const int32 X1, const int32 Y1, const int32 X2, const int32 Y2, FVector2D* Data, int32 Stride);
-	void GetXYOffsetDataFast(const int32 X1, const int32 Y1, const int32 X2, const int32 Y2, TMap<FIntPoint, FVector2D>& SparseData);
-	void GetXYOffsetDataFast(const int32 X1, const int32 Y1, const int32 X2, const int32 Y2, FVector* Data, int32 Stride);
-	void GetXYOffsetDataFast(const int32 X1, const int32 Y1, const int32 X2, const int32 Y2, TMap<FIntPoint, FVector>& SparseData);
+	LANDSCAPE_API void GetXYOffsetDataFast(const int32 X1, const int32 Y1, const int32 X2, const int32 Y2, FVector2D* Data, int32 Stride);
+	LANDSCAPE_API void GetXYOffsetDataFast(const int32 X1, const int32 Y1, const int32 X2, const int32 Y2, TMap<FIntPoint, FVector2D>& SparseData);
+	LANDSCAPE_API void GetXYOffsetDataFast(const int32 X1, const int32 Y1, const int32 X2, const int32 Y2, FVector* Data, int32 Stride);
+	LANDSCAPE_API void GetXYOffsetDataFast(const int32 X1, const int32 Y1, const int32 X2, const int32 Y2, TMap<FIntPoint, FVector>& SparseData);
 
 	template<typename T>
 	static void ShrinkData(TArray<T>& Data, int32 OldMinX, int32 OldMinY, int32 OldMaxX, int32 OldMaxY, int32 NewMinX, int32 NewMinY, int32 NewMaxX, int32 NewMaxY);
 
-	const ALandscape* GetTargetLandscape() const;
+	LANDSCAPE_API const ALandscape* GetTargetLandscape() const;
 
-	bool CanHaveLandscapeLayersContent() const;
-	bool HasLandscapeLayersContent() const;
+	LANDSCAPE_API bool CanHaveLandscapeLayersContent() const;
+	LANDSCAPE_API bool HasLandscapeLayersContent() const;
 
 private:
 	int32 ComponentSizeQuads;
@@ -307,7 +307,7 @@ private:
 	bool bUseSharedLandscapeEditLayer;
 	FGuid LocalEditLayerGUID;
 
-	void FillLayer(ULandscapeLayerInfoObject* LayerInfo, bool bEmptyLayersOnly);
+	LANDSCAPE_API void FillLayer(ULandscapeLayerInfoObject* LayerInfo, bool bEmptyLayersOnly);
 
 	// Only for Missing Data interpolation... only internal usage
 	template<typename TData, typename TStoreData, typename FType>
@@ -320,7 +320,7 @@ private:
 	template<typename TStoreData, typename TGetHeightMapDataFunction>
 	void GetHeightDataInternal(int32& ValidX1, int32& ValidY1, int32& ValidX2, int32& ValidY2, TStoreData& StoreData, TGetHeightMapDataFunction GetHeightMapDataFunction);
 
-	FORCEINLINE FColor& GetHeightMapColor(const ULandscapeComponent* Component, int32 TexU, int32 TexV, FColor* TextureData);
+	LANDSCAPE_API FORCEINLINE FColor& GetHeightMapColor(const ULandscapeComponent* Component, int32 TexU, int32 TexV, FColor* TextureData);
 
 	// test if layer is allowed for a given texel
 	inline bool IsLayerAllowed(const ULandscapeLayerInfoObject* LayerInfo,
@@ -331,10 +331,10 @@ private:
 	inline TMap<const ULandscapeLayerInfoObject*, uint32> CountWeightBlendedLayerInfluence(int32 ComponentIndexX, int32 ComponentIndexY, TOptional<TArrayView<const uint8* const>> LayerDataPtrs);
 
 	// chooses a replacement layer to use when erasing from 100% influence on a texel
-	const ULandscapeLayerInfoObject* ChooseReplacementLayer(const ULandscapeLayerInfoObject* LayerInfo, int32 ComponentIndexX, int32 SubIndexX, int32 SubX, int32 ComponentIndexY, int32 SubIndexY, int32 SubY, TMap<FIntPoint, TMap<const ULandscapeLayerInfoObject*, uint32>>& LayerInfluenceCache, TArrayView<const uint8* const> LayerDataPtrs);
+	LANDSCAPE_API const ULandscapeLayerInfoObject* ChooseReplacementLayer(const ULandscapeLayerInfoObject* LayerInfo, int32 ComponentIndexX, int32 SubIndexX, int32 SubX, int32 ComponentIndexY, int32 SubIndexY, int32 SubY, TMap<FIntPoint, TMap<const ULandscapeLayerInfoObject*, uint32>>& LayerInfluenceCache, TArrayView<const uint8* const> LayerDataPtrs);
 };
 
-struct LANDSCAPE_API FLandscapeDoNotDirtyScope
+struct FLandscapeDoNotDirtyScope
 {
 	FLandscapeDoNotDirtyScope(FLandscapeEditDataInterface& InEditInterface, bool bInScopeEnabled = true)
 		: EditInterface(InEditInterface), bScopeEnabled(bInScopeEnabled)
diff --git a/Engine/Source/Runtime/Landscape/Public/LandscapeHLODBuilder.h b/Engine/Source/Runtime/Landscape/Public/LandscapeHLODBuilder.h
--- a/Engine/Source/Runtime/Landscape/Public/LandscapeHLODBuilder.h
+++ b/Engine/Source/Runtime/Landscape/Public/LandscapeHLODBuilder.h
@@ -7,18 +7,18 @@
 #include "WorldPartition/HLOD/HLODBuilder.h"
 #include "LandscapeHLODBuilder.generated.h"
 
-UCLASS(HideDropdown)
-class LANDSCAPE_API ULandscapeHLODBuilder : public UHLODBuilder
+UCLASS(HideDropdown, MinimalAPI)
+class ULandscapeHLODBuilder : public UHLODBuilder
 {
 	GENERATED_UCLASS_BODY()
 
 public:
 #if WITH_EDITOR
-	virtual uint32 ComputeHLODHash(const UActorComponent* InSourceComponent) const override;
+	LANDSCAPE_API virtual uint32 ComputeHLODHash(const UActorComponent* InSourceComponent) const override;
 
 	/**
 	 * Components created with this method need to be properly outered & assigned to your target actor.
 	 */
-	virtual TArray<UActorComponent*> Build(const FHLODBuildContext& InHLODBuildContext, const TArray<UActorComponent*>& InSourceComponents) const override;
+	LANDSCAPE_API virtual TArray<UActorComponent*> Build(const FHLODBuildContext& InHLODBuildContext, const TArray<UActorComponent*>& InSourceComponents) const override;
 #endif
 };
diff --git a/Engine/Source/Runtime/Landscape/Public/LandscapeModule.h b/Engine/Source/Runtime/Landscape/Public/LandscapeModule.h
--- a/Engine/Source/Runtime/Landscape/Public/LandscapeModule.h
+++ b/Engine/Source/Runtime/Landscape/Public/LandscapeModule.h
@@ -16,6 +16,6 @@ class ILandscapeModule : public IModuleInterface
 public:
 	virtual TSharedPtr<FLandscapeSceneViewExtension, ESPMode::ThreadSafe> GetLandscapeSceneViewExtension() const = 0;
 	
-	virtual LANDSCAPE_API void SetLandscapeEditorServices(ILandscapeEditorServices* InLandscapeEditorServices) = 0;
-	virtual LANDSCAPE_API ILandscapeEditorServices* GetLandscapeEditorServices() const = 0;
+	virtual void SetLandscapeEditorServices(ILandscapeEditorServices* InLandscapeEditorServices) = 0;
+	virtual ILandscapeEditorServices* GetLandscapeEditorServices() const = 0;
 };
diff --git a/Engine/Source/Runtime/Landscape/Public/LandscapeRenderMobile.h b/Engine/Source/Runtime/Landscape/Public/LandscapeRenderMobile.h
--- a/Engine/Source/Runtime/Landscape/Public/LandscapeRenderMobile.h
+++ b/Engine/Source/Runtime/Landscape/Public/LandscapeRenderMobile.h
@@ -43,7 +43,7 @@ public:
 	}
 
 	UE_DEPRECATED(5.1, "FLandscapeVertexBufferMobile is now deprecated and will be removed.")
-	static LANDSCAPE_API void UpdateMemoryStat(int32 Delta) 
+	static void UpdateMemoryStat(int32 Delta) 
 	{
 	}
 };
diff --git a/Engine/Source/Runtime/Landscape/Public/LandscapeSettings.h b/Engine/Source/Runtime/Landscape/Public/LandscapeSettings.h
--- a/Engine/Source/Runtime/Landscape/Public/LandscapeSettings.h
+++ b/Engine/Source/Runtime/Landscape/Public/LandscapeSettings.h
@@ -20,8 +20,8 @@ enum class ELandscapeDirtyingMode : uint8
 	InLandscapeModeAndUserTriggeredChanges
 };
 
-UCLASS(config = Engine, defaultconfig, meta = (DisplayName = "Landscape"))
-class LANDSCAPE_API ULandscapeSettings : public UDeveloperSettings
+UCLASS(config = Engine, defaultconfig, meta = (DisplayName = "Landscape"), MinimalAPI)
+class ULandscapeSettings : public UDeveloperSettings
 {
 	GENERATED_BODY()
 
diff --git a/Engine/Source/Runtime/Landscape/Public/LandscapeSubsystem.h b/Engine/Source/Runtime/Landscape/Public/LandscapeSubsystem.h
--- a/Engine/Source/Runtime/Landscape/Public/LandscapeSubsystem.h
+++ b/Engine/Source/Runtime/Landscape/Public/LandscapeSubsystem.h
@@ -74,9 +74,9 @@ public:
 	LANDSCAPE_API void BuildPhysicalMaterial();
 
 	UE_DEPRECATED(5.3, "BuildGIBakedTextures is officially deprecated nowe")
-	LANDSCAPE_API void BuildGIBakedTextures() {}
+	void BuildGIBakedTextures() {}
 	UE_DEPRECATED(5.3, "GetOutdatedGIBakedTextureComponentsCount is officially deprecated now")
-	LANDSCAPE_API int32 GetOutdatedGIBakedTextureComponentsCount() { return 0; }
+	int32 GetOutdatedGIBakedTextureComponentsCount() { return 0; }
 
 	/**
 	 * Updates the Nanite mesh on all landscape actors whose mesh is not up to date.
@@ -104,7 +104,7 @@ public:
 	LANDSCAPE_API bool GetDirtyOnlyInMode() const;
 	FLandscapeNotificationManager* GetNotificationManager() { return NotificationManager; }
 	FOnHeightmapStreamedDelegate& GetOnHeightmapStreamedDelegate() { return OnHeightmapStreamed; }
-	LANDSCAPE_API bool AnyViewShowCollisions() const { return bAnyViewShowCollisions; }  //! Returns true if any view has view collisions enabled.
+	bool AnyViewShowCollisions() const { return bAnyViewShowCollisions; }  //! Returns true if any view has view collisions enabled.
 	FDateTime GetAppCurrentDateTime();
 	LANDSCAPE_API void AddAsyncEvent(FGraphEventRef GraphEventRef);
 
diff --git a/Engine/Source/Runtime/Landscape/Public/LandscapeUtils.h b/Engine/Source/Runtime/Landscape/Public/LandscapeUtils.h
--- a/Engine/Source/Runtime/Landscape/Public/LandscapeUtils.h
+++ b/Engine/Source/Runtime/Landscape/Public/LandscapeUtils.h
@@ -50,7 +50,7 @@ struct FTextureCopyChannelMapping
 	int8 Mappings[4];
 };
 
-class LANDSCAPE_API FBatchTextureCopy
+class FBatchTextureCopy
 {
 public:
 	/**
@@ -61,10 +61,10 @@ public:
 	* @param	InLayerInfo		The layer info used to retrieve the proper source weightmap and channel.
 	* @return True if the copy has been successfully added.
 	*/
-	bool AddWeightmapCopy(UTexture2D* InDestination, int8 InDestinationChannel, const ULandscapeComponent* InComponent, ULandscapeLayerInfoObject* InLayerInfo);
+	LANDSCAPE_API bool AddWeightmapCopy(UTexture2D* InDestination, int8 InDestinationChannel, const ULandscapeComponent* InComponent, ULandscapeLayerInfoObject* InLayerInfo);
 
 	/** Process pending internal copy requests. */
-	bool ProcessTextureCopies();
+	LANDSCAPE_API bool ProcessTextureCopies();
 
 private:
 	using FTextureCopyChannelMappingMap = TMap<FTextureCopyRequest, FTextureCopyChannelMapping>;
diff --git a/Engine/Source/Runtime/LevelSequence/Public/AnimSequenceLevelSequenceLink.h b/Engine/Source/Runtime/LevelSequence/Public/AnimSequenceLevelSequenceLink.h
--- a/Engine/Source/Runtime/LevelSequence/Public/AnimSequenceLevelSequenceLink.h
+++ b/Engine/Source/Runtime/LevelSequence/Public/AnimSequenceLevelSequenceLink.h
@@ -7,8 +7,8 @@
 class ULevelSequence;
 
 /** Link To Level Sequence That may be driving the anim sequence*/
-UCLASS(BlueprintType)
-class LEVELSEQUENCE_API UAnimSequenceLevelSequenceLink : public UAssetUserData
+UCLASS(BlueprintType, MinimalAPI)
+class UAnimSequenceLevelSequenceLink : public UAssetUserData
 {
 	GENERATED_UCLASS_BODY()
 
@@ -21,6 +21,6 @@ public:
 	UPROPERTY(BlueprintReadWrite, AssetRegistrySearchable, Category = Property)
 	FSoftObjectPath PathToLevelSequence;
 
-	void SetLevelSequence(ULevelSequence* InLevelSequence);
-	ULevelSequence* ResolveLevelSequence();
-};
\ No newline at end of file
+	LEVELSEQUENCE_API void SetLevelSequence(ULevelSequence* InLevelSequence);
+	LEVELSEQUENCE_API ULevelSequence* ResolveLevelSequence();
+};
diff --git a/Engine/Source/Runtime/LevelSequence/Public/DefaultLevelSequenceInstanceData.h b/Engine/Source/Runtime/LevelSequence/Public/DefaultLevelSequenceInstanceData.h
--- a/Engine/Source/Runtime/LevelSequence/Public/DefaultLevelSequenceInstanceData.h
+++ b/Engine/Source/Runtime/LevelSequence/Public/DefaultLevelSequenceInstanceData.h
@@ -10,8 +10,8 @@
 #include "DefaultLevelSequenceInstanceData.generated.h"
 
 /** Default instance data class that level sequences understand. Implements IMovieSceneTransformOrigin. */
-UCLASS(BlueprintType)
-class LEVELSEQUENCE_API UDefaultLevelSequenceInstanceData
+UCLASS(BlueprintType, MinimalAPI)
+class UDefaultLevelSequenceInstanceData
 	: public UObject
 	, public IMovieSceneTransformOrigin
 {
diff --git a/Engine/Source/Runtime/LevelSequence/Public/LevelSequence.h b/Engine/Source/Runtime/LevelSequence/Public/LevelSequence.h
--- a/Engine/Source/Runtime/LevelSequence/Public/LevelSequence.h
+++ b/Engine/Source/Runtime/LevelSequence/Public/LevelSequence.h
@@ -20,8 +20,8 @@ class UAssetUserData;
 /**
  * Movie scene animation for Actors.
  */
-UCLASS(BlueprintType)
-class LEVELSEQUENCE_API ULevelSequence
+UCLASS(BlueprintType, MinimalAPI)
+class ULevelSequence
 	: public UMovieSceneSequence, public IInterface_AssetUserData
 {
 	GENERATED_UCLASS_BODY()
@@ -35,62 +35,62 @@ public:
 public:
 
 	/** Initialize this level sequence. */
-	virtual void Initialize();
+	LEVELSEQUENCE_API virtual void Initialize();
 
 	/** Convert old-style lazy object ptrs to new-style references using the specified context */
-	void ConvertPersistentBindingsToDefault(UObject* FixupContext);
+	LEVELSEQUENCE_API void ConvertPersistentBindingsToDefault(UObject* FixupContext);
 
 public:
 
 	// UMovieSceneSequence interface
-	virtual void BindPossessableObject(const FGuid& ObjectId, UObject& PossessedObject, UObject* Context) override;
-	virtual bool CanPossessObject(UObject& Object, UObject* InPlaybackContext) const override;
-	virtual void LocateBoundObjects(const FGuid& ObjectId, UObject* Context, TArray<UObject*, TInlineAllocator<1>>& OutObjects) const override;
-	virtual FGuid FindBindingFromObject(UObject* InObject, UObject* Context) const override;
-	virtual void GatherExpiredObjects(const FMovieSceneObjectCache& InObjectCache, TArray<FGuid>& OutInvalidIDs) const override;
-	virtual UMovieScene* GetMovieScene() const override;
-	virtual UObject* GetParentObject(UObject* Object) const override;
-	virtual void UnbindPossessableObjects(const FGuid& ObjectId) override;
-	virtual void UnbindObjects(const FGuid& ObjectId, const TArray<UObject*>& InObjects, UObject* InContext) override;
-	virtual void UnbindInvalidObjects(const FGuid& ObjectId, UObject* InContext) override;
-	virtual bool AllowsSpawnableObjects() const override;
-	virtual bool CanRebindPossessable(const FMovieScenePossessable& InPossessable) const override;
-	virtual UObject* MakeSpawnableTemplateFromInstance(UObject& InSourceObject, FName ObjectName) override;
-	virtual bool CanAnimateObject(UObject& InObject) const override;
-	virtual UObject* CreateDirectorInstance(IMovieScenePlayer& Player, FMovieSceneSequenceID SequenceID) override;
-	virtual void PostLoad() override;
+	LEVELSEQUENCE_API virtual void BindPossessableObject(const FGuid& ObjectId, UObject& PossessedObject, UObject* Context) override;
+	LEVELSEQUENCE_API virtual bool CanPossessObject(UObject& Object, UObject* InPlaybackContext) const override;
+	LEVELSEQUENCE_API virtual void LocateBoundObjects(const FGuid& ObjectId, UObject* Context, TArray<UObject*, TInlineAllocator<1>>& OutObjects) const override;
+	LEVELSEQUENCE_API virtual FGuid FindBindingFromObject(UObject* InObject, UObject* Context) const override;
+	LEVELSEQUENCE_API virtual void GatherExpiredObjects(const FMovieSceneObjectCache& InObjectCache, TArray<FGuid>& OutInvalidIDs) const override;
+	LEVELSEQUENCE_API virtual UMovieScene* GetMovieScene() const override;
+	LEVELSEQUENCE_API virtual UObject* GetParentObject(UObject* Object) const override;
+	LEVELSEQUENCE_API virtual void UnbindPossessableObjects(const FGuid& ObjectId) override;
+	LEVELSEQUENCE_API virtual void UnbindObjects(const FGuid& ObjectId, const TArray<UObject*>& InObjects, UObject* InContext) override;
+	LEVELSEQUENCE_API virtual void UnbindInvalidObjects(const FGuid& ObjectId, UObject* InContext) override;
+	LEVELSEQUENCE_API virtual bool AllowsSpawnableObjects() const override;
+	LEVELSEQUENCE_API virtual bool CanRebindPossessable(const FMovieScenePossessable& InPossessable) const override;
+	LEVELSEQUENCE_API virtual UObject* MakeSpawnableTemplateFromInstance(UObject& InSourceObject, FName ObjectName) override;
+	LEVELSEQUENCE_API virtual bool CanAnimateObject(UObject& InObject) const override;
+	LEVELSEQUENCE_API virtual UObject* CreateDirectorInstance(IMovieScenePlayer& Player, FMovieSceneSequenceID SequenceID) override;
+	LEVELSEQUENCE_API virtual void PostLoad() override;
 #if WITH_EDITORONLY_DATA
-	static void DeclareConstructClasses(TArray<FTopLevelAssetPath>& OutConstructClasses, const UClass* SpecificSubclass);
+	static LEVELSEQUENCE_API void DeclareConstructClasses(TArray<FTopLevelAssetPath>& OutConstructClasses, const UClass* SpecificSubclass);
 #endif
 
-	virtual void PostInitProperties() override;
-	virtual bool Rename(const TCHAR* NewName = nullptr, UObject* NewOuter = nullptr, ERenameFlags Flags = REN_None) override;
+	LEVELSEQUENCE_API virtual void PostInitProperties() override;
+	LEVELSEQUENCE_API virtual bool Rename(const TCHAR* NewName = nullptr, UObject* NewOuter = nullptr, ERenameFlags Flags = REN_None) override;
 
 
 	//~ Begin IInterface_AssetUserData Interface
-	virtual void AddAssetUserData(UAssetUserData* InUserData) override;
-	virtual void RemoveUserDataOfClass(TSubclassOf<UAssetUserData> InUserDataClass) override;
-	virtual UAssetUserData* GetAssetUserDataOfClass(TSubclassOf<UAssetUserData> InUserDataClass) override;
-	virtual const TArray<UAssetUserData*>* GetAssetUserDataArray() const override;
+	LEVELSEQUENCE_API virtual void AddAssetUserData(UAssetUserData* InUserData) override;
+	LEVELSEQUENCE_API virtual void RemoveUserDataOfClass(TSubclassOf<UAssetUserData> InUserDataClass) override;
+	LEVELSEQUENCE_API virtual UAssetUserData* GetAssetUserDataOfClass(TSubclassOf<UAssetUserData> InUserDataClass) override;
+	LEVELSEQUENCE_API virtual const TArray<UAssetUserData*>* GetAssetUserDataArray() const override;
 	//~ End IInterface_AssetUserData Interface
 
 
 #if WITH_EDITOR
-	virtual ETrackSupport IsTrackSupported(TSubclassOf<class UMovieSceneTrack> InTrackClass) const override;
-	virtual void GetAssetRegistryTagMetadata(TMap<FName, FAssetRegistryTagMetadata>& OutMetadata) const override;
-	virtual void GetAssetRegistryTags(TArray<FAssetRegistryTag>& OutTags) const override;
-	virtual void PostLoadAssetRegistryTags(const FAssetData& InAssetData, TArray<FAssetRegistryTag>& OutTagsAndValuesToUpdate) const override;
+	LEVELSEQUENCE_API virtual ETrackSupport IsTrackSupported(TSubclassOf<class UMovieSceneTrack> InTrackClass) const override;
+	LEVELSEQUENCE_API virtual void GetAssetRegistryTagMetadata(TMap<FName, FAssetRegistryTagMetadata>& OutMetadata) const override;
+	LEVELSEQUENCE_API virtual void GetAssetRegistryTags(TArray<FAssetRegistryTag>& OutTags) const override;
+	LEVELSEQUENCE_API virtual void PostLoadAssetRegistryTags(const FAssetData& InAssetData, TArray<FAssetRegistryTag>& OutTagsAndValuesToUpdate) const override;
 	
 	DECLARE_DELEGATE_RetVal_OneParam(void, FPostDuplicateEvent, ULevelSequence*);
-	static FPostDuplicateEvent PostDuplicateEvent;
+	static LEVELSEQUENCE_API FPostDuplicateEvent PostDuplicateEvent;
 #endif
 
-	virtual void PostDuplicate(bool bDuplicateForPIE) override;
+	LEVELSEQUENCE_API virtual void PostDuplicate(bool bDuplicateForPIE) override;
 		
 	UE_DEPRECATED(5.3, "Use LocateBoundObjects with FLocateBoundObjectsParams parameter instead")
 	void LocateBoundObjects(const FGuid& ObjectId, UObject* Context, const FTopLevelAssetPath& StreamedLevelAssetPath, TArray<UObject*, TInlineAllocator<1>>& OutObjects) const	{}
 
-	void LocateBoundObjects(const FGuid& ObjectId, UObject* Context, const FLevelSequenceBindingReference::FResolveBindingParams& InResolveBindingParams, TArray<UObject*, TInlineAllocator<1>>& OutObjects) const;
+	LEVELSEQUENCE_API void LocateBoundObjects(const FGuid& ObjectId, UObject* Context, const FLevelSequenceBindingReference::FResolveBindingParams& InResolveBindingParams, TArray<UObject*, TInlineAllocator<1>>& OutObjects) const;
 #if WITH_EDITOR
 
 
@@ -98,26 +98,26 @@ public:
 
 	/**
 	 * Assign a new director blueprint to this level sequence. The specified blueprint *must* be contained within this object.?	 */
-	void SetDirectorBlueprint(UBlueprint* NewDirectorBlueprint);
+	LEVELSEQUENCE_API void SetDirectorBlueprint(UBlueprint* NewDirectorBlueprint);
 
 	/**
 	 * Retrieve the currently assigned director blueprint for this level sequence
 	 */
-	UBlueprint* GetDirectorBlueprint() const;
+	LEVELSEQUENCE_API UBlueprint* GetDirectorBlueprint() const;
 
-	FString GetDirectorBlueprintName() const;
+	LEVELSEQUENCE_API FString GetDirectorBlueprintName() const;
 
 protected:
 
-	virtual FGuid CreatePossessable(UObject* ObjectToPossess) override;
-	virtual FGuid CreateSpawnable(UObject* ObjectToSpawn) override;
+	LEVELSEQUENCE_API virtual FGuid CreatePossessable(UObject* ObjectToPossess) override;
+	LEVELSEQUENCE_API virtual FGuid CreateSpawnable(UObject* ObjectToSpawn) override;
 
-	FGuid FindOrAddBinding(UObject* ObjectToPossess);
+	LEVELSEQUENCE_API FGuid FindOrAddBinding(UObject* ObjectToPossess);
 
 	/**
 	 * Invoked when this level sequence's director blueprint has been recompiled
 	 */
-	void OnDirectorRecompiled(UBlueprint*);
+	LEVELSEQUENCE_API void OnDirectorRecompiled(UBlueprint*);
 
 #endif // WITH_EDITOR
 
diff --git a/Engine/Source/Runtime/LevelSequence/Public/LevelSequenceActor.h b/Engine/Source/Runtime/LevelSequence/Public/LevelSequenceActor.h
--- a/Engine/Source/Runtime/LevelSequence/Public/LevelSequenceActor.h
+++ b/Engine/Source/Runtime/LevelSequence/Public/LevelSequenceActor.h
@@ -24,26 +24,26 @@ class ULevelSequenceBurnIn;
 class ULevelSequencePlayer;
 class UMovieSceneSequenceTickManager;
 
-UCLASS(Blueprintable, DefaultToInstanced)
-class LEVELSEQUENCE_API ULevelSequenceBurnInInitSettings : public UObject
+UCLASS(Blueprintable, DefaultToInstanced, MinimalAPI)
+class ULevelSequenceBurnInInitSettings : public UObject
 {
 	GENERATED_BODY()
 };
 
-UCLASS(config=EditorPerProjectUserSettings, PerObjectConfig, DefaultToInstanced, BlueprintType)
-class LEVELSEQUENCE_API ULevelSequenceBurnInOptions : public UObject
+UCLASS(config=EditorPerProjectUserSettings, PerObjectConfig, DefaultToInstanced, BlueprintType, MinimalAPI)
+class ULevelSequenceBurnInOptions : public UObject
 {
 public:
 
 	GENERATED_BODY()
-	ULevelSequenceBurnInOptions(const FObjectInitializer& Init);
+	LEVELSEQUENCE_API ULevelSequenceBurnInOptions(const FObjectInitializer& Init);
 
 	/** Loads the specified class path and initializes an instance, then stores it in Settings. */
 	UFUNCTION(BlueprintCallable, Category = "General")
-	void SetBurnIn(FSoftClassPath InBurnInClass);
+	LEVELSEQUENCE_API void SetBurnIn(FSoftClassPath InBurnInClass);
 
 	/** Ensure the settings object is up-to-date */
-	void ResetSettings();
+	LEVELSEQUENCE_API void ResetSettings();
 
 public:
 
@@ -59,15 +59,15 @@ public:
 protected:
 
 #if WITH_EDITOR
-	virtual void PostEditChangeProperty(FPropertyChangedEvent& PropertyChangedEvent) override;
+	LEVELSEQUENCE_API virtual void PostEditChangeProperty(FPropertyChangedEvent& PropertyChangedEvent) override;
 #endif //WITH_EDITOR
 };
 
 /**
  * Actor responsible for controlling a specific level sequence in the world.
  */
-UCLASS(hideCategories=(Rendering, Physics, HLOD, Activation, Input))
-class LEVELSEQUENCE_API ALevelSequenceActor
+UCLASS(hideCategories=(Rendering, Physics, HLOD, Activation, Input), MinimalAPI)
+class ALevelSequenceActor
 	: public AActor
 	, public IMovieScenePlaybackClient
 	, public IMovieSceneBindingOwnerInterface
@@ -80,7 +80,7 @@ public:
 	GENERATED_BODY()
 
 	/** Create and initialize a new instance. */
-	ALevelSequenceActor(const FObjectInitializer& Init);
+	LEVELSEQUENCE_API ALevelSequenceActor(const FObjectInitializer& Init);
 
 public:
 
@@ -132,7 +132,7 @@ public:
 	 * @see SetSequence
 	 */
 	UFUNCTION(BlueprintCallable, Category="Sequencer|Player")
-	ULevelSequence* GetSequence() const;
+	LEVELSEQUENCE_API ULevelSequence* GetSequence() const;
 
 	/**
 	 * Set the level sequence being played by this actor.
@@ -141,30 +141,30 @@ public:
 	 * @see GetSequence
 	 */
 	UFUNCTION(BlueprintCallable, Category="Sequencer|Player")
-	void SetSequence(ULevelSequence* InSequence);
+	LEVELSEQUENCE_API void SetSequence(ULevelSequence* InSequence);
 
 	/**
 	 * Set whether or not to replicate playback for this actor
 	 */
 	UFUNCTION(BlueprintSetter)
-	void SetReplicatePlayback(bool ReplicatePlayback);
+	LEVELSEQUENCE_API void SetReplicatePlayback(bool ReplicatePlayback);
 
 	/**
 	 * Access this actor's sequence player, or None if it is not yet initialized
 	 */
 	UFUNCTION(BlueprintGetter)
-	ULevelSequencePlayer* GetSequencePlayer() const;
+	LEVELSEQUENCE_API ULevelSequencePlayer* GetSequencePlayer() const;
 
 	/* Hide burnin */
 	UFUNCTION(BlueprintCallable, Category = "Sequencer|Player")
-	void HideBurnin();
+	LEVELSEQUENCE_API void HideBurnin();
 
 	/* Show burnin */
 	UFUNCTION(BlueprintCallable, Category = "Sequencer|Player")
-	void ShowBurnin();
+	LEVELSEQUENCE_API void ShowBurnin();
 
 	/** Refresh this actor's burn in */
-	void RefreshBurnIn();
+	LEVELSEQUENCE_API void RefreshBurnIn();
 
 public:
 
@@ -178,7 +178,7 @@ public:
 	 *								  set in Sequencer UI. This function will not modify the original asset.
 	 */
 	UFUNCTION(BlueprintCallable, Category = "Sequencer|Player|Bindings")
-	void SetBinding(FMovieSceneObjectBindingID Binding, const TArray<AActor*>& Actors, bool bAllowBindingsFromAsset = false);
+	LEVELSEQUENCE_API void SetBinding(FMovieSceneObjectBindingID Binding, const TArray<AActor*>& Actors, bool bAllowBindingsFromAsset = false);
 
 	/**
 	 * Assigns an set of actors to all the bindings tagged with the specified name in this sequence. Object Bindings can be tagged within the sequence UI by RMB -> Tags... on the object binding in the tree.
@@ -190,7 +190,7 @@ public:
 	 *								  set in Sequencer UI. This function will not modify the original asset.
 	 */
 	UFUNCTION(BlueprintCallable, Category = "Sequencer|Player|Bindings")
-	void SetBindingByTag(FName BindingTag, const TArray<AActor*>& Actors, bool bAllowBindingsFromAsset = false);
+	LEVELSEQUENCE_API void SetBindingByTag(FName BindingTag, const TArray<AActor*>& Actors, bool bAllowBindingsFromAsset = false);
 
 	/**
 	 * Adds the specified actor to the overridden bindings for the specified binding ID, optionally still allowing the bindings defined in the Level Sequence asset
@@ -202,7 +202,7 @@ public:
 	 *								  set in Sequencer UI. This function will not modify the original asset.
 	 */
 	UFUNCTION(BlueprintCallable, Category = "Sequencer|Player|Bindings")
-	void AddBinding(FMovieSceneObjectBindingID Binding, AActor* Actor, bool bAllowBindingsFromAsset = false);
+	LEVELSEQUENCE_API void AddBinding(FMovieSceneObjectBindingID Binding, AActor* Actor, bool bAllowBindingsFromAsset = false);
 
 	/**
 	 * Binds an actor to all the bindings tagged with the specified name in this sequence. Does not remove any exising bindings that have been set up through this API. Object Bindings can be tagged within the sequence UI by RMB -> Tags... on the object binding in the tree.
@@ -214,37 +214,37 @@ public:
 	 *								  set in Sequencer UI. This function will not modify the original asset.
 	 */
 	UFUNCTION(BlueprintCallable, Category = "Sequencer|Player|Bindings")
-	void AddBindingByTag(FName BindingTag, AActor* Actor, bool bAllowBindingsFromAsset = false);
+	LEVELSEQUENCE_API void AddBindingByTag(FName BindingTag, AActor* Actor, bool bAllowBindingsFromAsset = false);
 
 	/**
 	 * Removes the specified actor from the specified binding's actor array
 	 */
 	UFUNCTION(BlueprintCallable, Category = "Sequencer|Player|Bindings")
-	void RemoveBinding(FMovieSceneObjectBindingID Binding, AActor* Actor);
+	LEVELSEQUENCE_API void RemoveBinding(FMovieSceneObjectBindingID Binding, AActor* Actor);
 
 	/**
 	 * Removes the specified actor from the specified binding's actor array
 	 */
 	UFUNCTION(BlueprintCallable, Category = "Sequencer|Player|Bindings")
-	void RemoveBindingByTag(FName Tag, AActor* Actor);
+	LEVELSEQUENCE_API void RemoveBindingByTag(FName Tag, AActor* Actor);
 
 	/**
 	 * Resets the specified binding back to the defaults defined by the Level Sequence asset
 	 */
 	UFUNCTION(BlueprintCallable, Category = "Sequencer|Player|Bindings")
-	void ResetBinding(FMovieSceneObjectBindingID Binding);
+	LEVELSEQUENCE_API void ResetBinding(FMovieSceneObjectBindingID Binding);
 
 	/**
 	 * Resets all overridden bindings back to the defaults defined by the Level Sequence asset
 	 */
 	UFUNCTION(BlueprintCallable, Category = "Sequencer|Player|Bindings")
-	void ResetBindings();
+	LEVELSEQUENCE_API void ResetBindings();
 
 	/**
 	 * Retrieve the first object binding that has been tagged with the specified name
 	 */
 	UFUNCTION(BlueprintCallable, Category = "Sequencer|Player|Bindings", DisplayName="Find Binding by Tag")
-	FMovieSceneObjectBindingID FindNamedBinding(FName Tag) const;
+	LEVELSEQUENCE_API FMovieSceneObjectBindingID FindNamedBinding(FName Tag) const;
 
 	/**
 	 * Retrieve all the bindings that have been tagged with the specified name
@@ -253,35 +253,35 @@ public:
 	 * @return An array containing all the bindings that are tagged with this name, potentially empty.
 	 */
 	UFUNCTION(BlueprintCallable, Category = "Sequencer|Player|Bindings", DisplayName="Find Bindings by Tag")
-	const TArray<FMovieSceneObjectBindingID>& FindNamedBindings(FName Tag) const;
+	LEVELSEQUENCE_API const TArray<FMovieSceneObjectBindingID>& FindNamedBindings(FName Tag) const;
 
 protected:
 
 	//~ Begin IMovieScenePlaybackClient interface
-	virtual bool RetrieveBindingOverrides(const FGuid& InBindingId, FMovieSceneSequenceID InSequenceID, TArray<UObject*, TInlineAllocator<1>>& OutObjects) const override;
-	virtual UObject* GetInstanceData() const override;
-	virtual TOptional<EAspectRatioAxisConstraint> GetAspectRatioAxisConstraint() const override;
-	virtual bool GetIsReplicatedPlayback() const override;
+	LEVELSEQUENCE_API virtual bool RetrieveBindingOverrides(const FGuid& InBindingId, FMovieSceneSequenceID InSequenceID, TArray<UObject*, TInlineAllocator<1>>& OutObjects) const override;
+	LEVELSEQUENCE_API virtual UObject* GetInstanceData() const override;
+	LEVELSEQUENCE_API virtual TOptional<EAspectRatioAxisConstraint> GetAspectRatioAxisConstraint() const override;
+	LEVELSEQUENCE_API virtual bool GetIsReplicatedPlayback() const override;
 	//~ End IMovieScenePlaybackClient interface
 
 	//~ Begin UObject interface
-	virtual bool ReplicateSubobjects(UActorChannel* Channel, FOutBunch* Bunch, FReplicationFlags *RepFlags) override;
-	virtual void PostInitProperties() override;
-	virtual void PostLoad() override;
+	LEVELSEQUENCE_API virtual bool ReplicateSubobjects(UActorChannel* Channel, FOutBunch* Bunch, FReplicationFlags *RepFlags) override;
+	LEVELSEQUENCE_API virtual void PostInitProperties() override;
+	LEVELSEQUENCE_API virtual void PostLoad() override;
 public:
 #if WITH_EDITORONLY_DATA
-	static void DeclareConstructClasses(TArray<FTopLevelAssetPath>& OutConstructClasses, const UClass* SpecificSubclass);
+	static LEVELSEQUENCE_API void DeclareConstructClasses(TArray<FTopLevelAssetPath>& OutConstructClasses, const UClass* SpecificSubclass);
 #endif
 protected:
 
 	//~ End UObject interface
 
 	//~ Begin AActor interface
-	virtual void PostInitializeComponents() override;
-	virtual void BeginPlay() override;
-	virtual void EndPlay(const EEndPlayReason::Type EndPlayReason) override;
-	virtual void RewindForReplay() override;
-	virtual void PostNetReceive() override;
+	LEVELSEQUENCE_API virtual void PostInitializeComponents() override;
+	LEVELSEQUENCE_API virtual void BeginPlay() override;
+	LEVELSEQUENCE_API virtual void EndPlay(const EEndPlayReason::Type EndPlayReason) override;
+	LEVELSEQUENCE_API virtual void RewindForReplay() override;
+	LEVELSEQUENCE_API virtual void PostNetReceive() override;
 #if WITH_EDITOR
 	virtual bool CanChangeIsSpatiallyLoadedFlag() const override { return false; }
 #endif
@@ -299,20 +299,20 @@ public:
 		WorldPartitionResolveData = InWorldPartitionResolveData;
 	}
 
-	virtual bool GetReferencedContentObjects(TArray<UObject*>& Objects) const override;
+	LEVELSEQUENCE_API virtual bool GetReferencedContentObjects(TArray<UObject*>& Objects) const override;
 #endif //WITH_EDITOR
 
 	/** Initialize the player object by loading the asset, using async loading when necessary */
-	void InitializePlayer();
+	LEVELSEQUENCE_API void InitializePlayer();
 
 	/** Initialize the player object with the specified asset */
-	void InitializePlayerWithSequence(ULevelSequence* LevelSequenceAsset);
-	void OnSequenceLoaded(const FName& PackageName, UPackage* Package, EAsyncLoadingResult::Type Result);
+	LEVELSEQUENCE_API void InitializePlayerWithSequence(ULevelSequence* LevelSequenceAsset);
+	LEVELSEQUENCE_API void OnSequenceLoaded(const FName& PackageName, UPackage* Package, EAsyncLoadingResult::Type Result);
 
 #if WITH_EDITOR
-	virtual TSharedPtr<FStructOnScope> GetObjectPickerProxy(TSharedPtr<IPropertyHandle> PropertyHandle) override;
-	virtual void UpdateObjectFromProxy(FStructOnScope& Proxy, IPropertyHandle& ObjectPropertyHandle) override;
-	virtual UMovieSceneSequence* RetrieveOwnedSequence() const override;
+	LEVELSEQUENCE_API virtual TSharedPtr<FStructOnScope> GetObjectPickerProxy(TSharedPtr<IPropertyHandle> PropertyHandle) override;
+	LEVELSEQUENCE_API virtual void UpdateObjectFromProxy(FStructOnScope& Proxy, IPropertyHandle& ObjectPropertyHandle) override;
+	LEVELSEQUENCE_API virtual UMovieSceneSequence* RetrieveOwnedSequence() const override;
 #endif
 
 private:
@@ -328,7 +328,7 @@ private:
 };
 
 USTRUCT()
-struct LEVELSEQUENCE_API FBoundActorProxy
+struct FBoundActorProxy
 {
 	GENERATED_BODY()
 
@@ -338,9 +338,9 @@ struct LEVELSEQUENCE_API FBoundActorProxy
 	UPROPERTY(EditInstanceOnly, AdvancedDisplay, Category="General")
 	TObjectPtr<AActor> BoundActor = nullptr;
 
-	void Initialize(TSharedPtr<IPropertyHandle> InPropertyHandle);
+	LEVELSEQUENCE_API void Initialize(TSharedPtr<IPropertyHandle> InPropertyHandle);
 
-	void OnReflectedPropertyChanged();
+	LEVELSEQUENCE_API void OnReflectedPropertyChanged();
 
 	TSharedPtr<IPropertyHandle> ReflectedProperty;
 
@@ -350,13 +350,13 @@ struct LEVELSEQUENCE_API FBoundActorProxy
 /**
  * A level sequence actor that is set to always be relevant for networking purposes
  */
-UCLASS()
-class LEVELSEQUENCE_API AReplicatedLevelSequenceActor
+UCLASS(MinimalAPI)
+class AReplicatedLevelSequenceActor
 	: public ALevelSequenceActor
 {
 	GENERATED_BODY()
 
 public:
 	/** Create and initialize a new instance. */
-	AReplicatedLevelSequenceActor(const FObjectInitializer& Init);
+	LEVELSEQUENCE_API AReplicatedLevelSequenceActor(const FObjectInitializer& Init);
 };
diff --git a/Engine/Source/Runtime/LevelSequence/Public/LevelSequenceActorSpawner.h b/Engine/Source/Runtime/LevelSequence/Public/LevelSequenceActorSpawner.h
--- a/Engine/Source/Runtime/LevelSequence/Public/LevelSequenceActorSpawner.h
+++ b/Engine/Source/Runtime/LevelSequence/Public/LevelSequenceActorSpawner.h
@@ -4,14 +4,14 @@
 
 #include "IMovieSceneObjectSpawner.h"
 
-class LEVELSEQUENCE_API FLevelSequenceActorSpawner : public IMovieSceneObjectSpawner
+class FLevelSequenceActorSpawner : public IMovieSceneObjectSpawner
 {
 public:
 
-	static TSharedRef<IMovieSceneObjectSpawner> CreateObjectSpawner();
+	static LEVELSEQUENCE_API TSharedRef<IMovieSceneObjectSpawner> CreateObjectSpawner();
 
 	// IMovieSceneObjectSpawner interface
-	virtual UClass* GetSupportedTemplateType() const override;
-	virtual UObject* SpawnObject(FMovieSceneSpawnable& Spawnable, FMovieSceneSequenceIDRef TemplateID, IMovieScenePlayer& Player) override;
-	virtual void DestroySpawnedObject(UObject& Object) override;
-};
\ No newline at end of file
+	LEVELSEQUENCE_API virtual UClass* GetSupportedTemplateType() const override;
+	LEVELSEQUENCE_API virtual UObject* SpawnObject(FMovieSceneSpawnable& Spawnable, FMovieSceneSequenceIDRef TemplateID, IMovieScenePlayer& Player) override;
+	LEVELSEQUENCE_API virtual void DestroySpawnedObject(UObject& Object) override;
+};
diff --git a/Engine/Source/Runtime/LevelSequence/Public/LevelSequenceAnimSequenceLink.h b/Engine/Source/Runtime/LevelSequence/Public/LevelSequenceAnimSequenceLink.h
--- a/Engine/Source/Runtime/LevelSequence/Public/LevelSequenceAnimSequenceLink.h
+++ b/Engine/Source/Runtime/LevelSequence/Public/LevelSequenceAnimSequenceLink.h
@@ -17,7 +17,7 @@ class UObject;
 
 /** Link To Anim Sequence that we are linked too.*/
 USTRUCT(BlueprintType)
-struct LEVELSEQUENCE_API FLevelSequenceAnimSequenceLinkItem
+struct FLevelSequenceAnimSequenceLinkItem
 {
 	GENERATED_BODY()
 
@@ -46,14 +46,14 @@ public:
 	UPROPERTY(BlueprintReadWrite, Category = Property)
 	bool bEvaluateAllSkeletalMeshComponents = true;
 
-	void SetAnimSequence(UAnimSequence* InAnimSequence);
-	UAnimSequence* ResolveAnimSequence();
+	LEVELSEQUENCE_API void SetAnimSequence(UAnimSequence* InAnimSequence);
+	LEVELSEQUENCE_API UAnimSequence* ResolveAnimSequence();
 
 };
 
 /** Link To Set of Anim Sequences that we may be linked to.*/
-UCLASS(BlueprintType)
-class LEVELSEQUENCE_API ULevelSequenceAnimSequenceLink : public UAssetUserData
+UCLASS(BlueprintType, MinimalAPI)
+class ULevelSequenceAnimSequenceLink : public UAssetUserData
 {
 	GENERATED_UCLASS_BODY()
 
@@ -62,4 +62,4 @@ public:
 
 	UPROPERTY(BlueprintReadWrite, Category = Links)
 	TArray< FLevelSequenceAnimSequenceLinkItem> AnimSequenceLinks;
-};
\ No newline at end of file
+};
diff --git a/Engine/Source/Runtime/LevelSequence/Public/LevelSequenceBurnIn.h b/Engine/Source/Runtime/LevelSequence/Public/LevelSequenceBurnIn.h
--- a/Engine/Source/Runtime/LevelSequence/Public/LevelSequenceBurnIn.h
+++ b/Engine/Source/Runtime/LevelSequence/Public/LevelSequenceBurnIn.h
@@ -13,30 +13,30 @@
 /**
  * Base class for level sequence burn ins
  */
-UCLASS()
-class LEVELSEQUENCE_API ULevelSequenceBurnIn : public UUserWidget
+UCLASS(MinimalAPI)
+class ULevelSequenceBurnIn : public UUserWidget
 {
 public:
 	GENERATED_BODY()
 
-	ULevelSequenceBurnIn(const FObjectInitializer& ObjectInitializer);
+	LEVELSEQUENCE_API ULevelSequenceBurnIn(const FObjectInitializer& ObjectInitializer);
 
 	/** Initialize this burn in */
-	void TakeSnapshotsFrom(ALevelSequenceActor& InActor);
+	LEVELSEQUENCE_API void TakeSnapshotsFrom(ALevelSequenceActor& InActor);
 
 	/** Called when this burn in is receiving its settings */
 	UFUNCTION(BlueprintImplementableEvent, Category="Burn In")
-	void SetSettings(UObject* InSettings);
+	LEVELSEQUENCE_API void SetSettings(UObject* InSettings);
 
 	/** Get the settings class to use for this burn in */
 	UFUNCTION(BlueprintNativeEvent, Category="Burn In")
-	TSubclassOf<ULevelSequenceBurnInInitSettings> GetSettingsClass() const;
+	LEVELSEQUENCE_API TSubclassOf<ULevelSequenceBurnInInitSettings> GetSettingsClass() const;
 	virtual TSubclassOf<ULevelSequenceBurnInInitSettings> GetSettingsClass_Implementation() const { return nullptr; }
 
 protected:
 
 	/** Called as part of the game tick loop when the sequence has been updated */
-	void OnSequenceUpdated(const UMovieSceneSequencePlayer& Player, FFrameTime CurrentTime, FFrameTime PreviousTime);
+	LEVELSEQUENCE_API void OnSequenceUpdated(const UMovieSceneSequencePlayer& Player, FFrameTime CurrentTime, FFrameTime PreviousTime);
 
 protected:
 
