Landscape  : Prevent crash on loading a map with landscape when dropping unwanted mips.

diff --git a/Engine/Source/Runtime/Landscape/Private/Landscape.cpp b/Engine/Source/Runtime/Landscape/Private/Landscape.cpp
--- a/Engine/Source/Runtime/Landscape/Private/Landscape.cpp
+++ b/Engine/Source/Runtime/Landscape/Private/Landscape.cpp
@@ -1182,16 +1182,14 @@ void ULandscapeComponent::PostLoad()
 			return;
 		}
 
-		uint64 NumBytes = InTexture->Source.GetSizeX() * InTexture->Source.GetSizeY() * FTextureSource::GetBytesPerPixel(InTexture->Source.GetFormat());
-		uint8* MipDataCopy = static_cast<uint8*>(FMemory::Malloc(NumBytes));
-		uint8 const* MipData = InTexture->Source.LockMipReadOnly(0);
-		FMemory::Memcpy(MipDataCopy, MipData, NumBytes);
+		TArray64<uint8> TopMipData;
+		InTexture->Source.GetMipData(TopMipData, 0);
 
-		InTexture->Source.Init(InTexture->Source.GetSizeX(), InTexture->Source.GetSizeY(), 1, 1, InTexture->Source.GetFormat(), MipDataCopy);
+		InTexture->PreEditChange(nullptr);
+		InTexture->Source.Init(InTexture->Source.GetSizeX(), InTexture->Source.GetSizeY(), 1, 1, InTexture->Source.GetFormat(), TopMipData.GetData());
 		InTexture->UpdateResource();
 
-		FMemory::Free(MipDataCopy);
-
+		InTexture->PostEditChange();
 	};
 
 	// Remove Non zero mip levels found in layer textures
