Enables landscape shared properties silence fixup by default.

diff --git a/Engine/Source/Runtime/Landscape/Private/Landscape.cpp b/Engine/Source/Runtime/Landscape/Private/Landscape.cpp
--- a/Engine/Source/Runtime/Landscape/Private/Landscape.cpp
+++ b/Engine/Source/Runtime/Landscape/Private/Landscape.cpp
@@ -177,7 +177,7 @@ static FAutoConsoleVariable CVarForceInvalidateNaniteOnLoad(
 
 static FAutoConsoleVariable CVarSilenceSharedPropertyDeprecationFixup(
 	TEXT("landscape.SilenceSharedPropertyDeprecationFixup"),
-	false,
+	true,
 	TEXT("Silently performs the fixup of discrepancies in shared properties when handling data modified before the enforcement introduction."));
 #endif // WITH_EDITOR
 
