Landscape : Improvements to landscape to the update of physical materials.

* Only write the hash of the Material to the ULandscapeComponent if the task is complete.
* Prior to LandscapeProxy save issue a warning if Physical Material renders are in progress
* Change the behavior of BuildPhysical Materials to actually rebuild all physical materials.

diff --git a/Engine/Source/Runtime/Landscape/Classes/LandscapeComponent.h b/Engine/Source/Runtime/Landscape/Classes/LandscapeComponent.h
--- a/Engine/Source/Runtime/Landscape/Classes/LandscapeComponent.h
+++ b/Engine/Source/Runtime/Landscape/Classes/LandscapeComponent.h
@@ -1021,6 +1021,8 @@ public:
 	bool CanUpdatePhysicalMaterial();
 	/** Update physical material render tasks. */
 	void UpdatePhysicalMaterialTasks();
+	/** Write the physical materials into the LandscapeComponent from the Render & Immediately Rebuild physics if requested */
+	void FinalizePhysicalMaterial(bool bInImmediatePhysicsRebuild);
 	/** Update collision component physical materials from render task results. */
 	void UpdateCollisionPhysicalMaterialData(TArray<UPhysicalMaterial*> const& InPhysicalMaterials, TArray<uint8> const& InMaterialIds);
 
diff --git a/Engine/Source/Runtime/Landscape/Classes/LandscapePhysicalMaterial.h b/Engine/Source/Runtime/Landscape/Classes/LandscapePhysicalMaterial.h
--- a/Engine/Source/Runtime/Landscape/Classes/LandscapePhysicalMaterial.h
+++ b/Engine/Source/Runtime/Landscape/Classes/LandscapePhysicalMaterial.h
@@ -16,9 +16,7 @@
 class FLandscapePhysicalMaterialRenderTask
 {
 public:
-	/* Task constructs with IsValid() false. */
 	FLandscapePhysicalMaterialRenderTask()
-		: PoolHandle(-1)
 	{}
 
 	~FLandscapePhysicalMaterialRenderTask()
@@ -27,7 +25,7 @@ public:
 	}
 
 	/** Initialize the task for a component. A task in progress can safely call Init() again to rekick the task. Returns false if no physical material needs to be rendered. */
-	bool Init(class ULandscapeComponent const* LandscapeComponent);
+	bool Init(class ULandscapeComponent const* LandscapeComponent, uint32 InHash);
 	/** Release the task. After calling this IsValid() will return false until we call Init() again. */
 	void Release();
 
@@ -35,6 +33,9 @@ public:
 	bool IsValid() const;
 	/** Returns true if the task is complete and the result is available. */
 	bool IsComplete() const;
+	/** Is this Task in progress (valid and incomplete) */
+	bool IsInProgress() const;
+
 
 	/** Update a task. Does nothing if the task is complete. */
 	void Tick();
@@ -46,14 +47,17 @@ public:
 	/** Get the result data. Assumes that IsComplete() is true. This doesn't Release() the data. */
 	TArray <class UPhysicalMaterial* > const& GetResultMaterials() const;
 
+	uint32 GetHash() const { return Hash; }
 private:
 	void UpdateInternal(bool bInFlush);
 
 private:
-	int32 PoolHandle;
+	int32 PoolHandle = -1;
 	
 	friend class FLandscapePhysicalMaterialRenderTaskPool;
 	FLandscapePhysicalMaterialRenderTask(FLandscapePhysicalMaterialRenderTask const&);
+
+	uint32 Hash = 0;
 };
 
 namespace LandscapePhysicalMaterial
diff --git a/Engine/Source/Runtime/Landscape/Classes/LandscapeProxy.h b/Engine/Source/Runtime/Landscape/Classes/LandscapeProxy.h
--- a/Engine/Source/Runtime/Landscape/Classes/LandscapeProxy.h
+++ b/Engine/Source/Runtime/Landscape/Classes/LandscapeProxy.h
@@ -1128,6 +1128,7 @@ public:
 	UE_DEPRECATED(5.3, "GetOutdatedGIBakedTextureComponentsCount is officially deprecated now returns 0")
 	int32 GetOutdatedGIBakedTextureComponentsCount() const { return 0; }
 	LANDSCAPE_API void BuildPhysicalMaterial(struct FScopedSlowTask* InSlowTask = nullptr);
+	LANDSCAPE_API void InvalidatePhysicalMaterial();
 	LANDSCAPE_API int32 GetOudatedPhysicalMaterialComponentsCount() const;
 	LANDSCAPE_API virtual void CreateSplineComponent() override;
 	LANDSCAPE_API virtual void CreateSplineComponent(const FVector& Scale3D) override;
@@ -1545,6 +1546,7 @@ class FLandscapePhysicalMaterialBuilder
 public:
 	LANDSCAPE_API FLandscapePhysicalMaterialBuilder(UWorld* InWorld);
 	LANDSCAPE_API void Build();
+	LANDSCAPE_API void Rebuild();
 	LANDSCAPE_API int32 GetOudatedPhysicalMaterialComponentsCount();
 private:
 	UWorld* World;
diff --git a/Engine/Source/Runtime/Landscape/Private/Landscape.cpp b/Engine/Source/Runtime/Landscape/Private/Landscape.cpp
--- a/Engine/Source/Runtime/Landscape/Private/Landscape.cpp
+++ b/Engine/Source/Runtime/Landscape/Private/Landscape.cpp
@@ -3191,6 +3191,16 @@ void ALandscapeProxy::PreSave(FObjectPreSaveContext ObjectSaveContext)
 		UpdateRenderingMethod();
 	}
 
+	// Warn if there are any active Physical Material Renders on this Proxy
+	// don.boogert todo: We should block or wait for this to complete within some timeout period.
+	for (ULandscapeComponent* LandscapeComponent : LandscapeComponents)
+	{
+		if (LandscapeComponent->PhysicalMaterialTask.IsInProgress())
+		{
+			UE_LOG(LogLandscape, Warning, TEXT("Physical material render on component: '%s' in progress."), *LandscapeComponent->GetFullName());
+		}
+	}
+
 	if (LandscapeGuid.IsValid())
 	{
 		if (ULandscapeInfo* LandscapeInfo = GetLandscapeInfo())
@@ -5698,6 +5708,18 @@ void FLandscapePhysicalMaterialBuilder::Build()
 	}
 }
 
+void FLandscapePhysicalMaterialBuilder::Rebuild()
+{
+	if (World)
+	{
+		for (TActorIterator<ALandscapeProxy> ProxyIt(World); ProxyIt; ++ProxyIt)
+		{
+			ProxyIt->InvalidatePhysicalMaterial();
+			ProxyIt->BuildPhysicalMaterial();
+		}
+	}
+}
+
 int32 FLandscapePhysicalMaterialBuilder::GetOudatedPhysicalMaterialComponentsCount()
 {
 	if (World)
@@ -5740,6 +5762,14 @@ UE::Landscape::EOutdatedDataFlags ALandscapeProxy::GetOutdatedDataFlags() const
 	return OutdatedDataFlags;
 }
 
+void ALandscapeProxy::InvalidatePhysicalMaterial()
+{
+	for (ULandscapeComponent* Component : LandscapeComponents)
+	{
+		Component->PhysicalMaterialHash = 0;
+	}
+}
+
 void ALandscapeProxy::BuildPhysicalMaterial(struct FScopedSlowTask* InSlowTask)
 {
 	if (!HasAnyFlags(RF_ClassDefaultObject))
@@ -5752,7 +5782,7 @@ void ALandscapeProxy::BuildPhysicalMaterial(struct FScopedSlowTask* InSlowTask)
 void ALandscapeProxy::UpdatePhysicalMaterialTasksStatus(TSet<ULandscapeComponent*>* OutdatedComponents, int32* OutdatedComponentsCount) const
 {
 	int32 OutdatedCount = 0;
-	for (auto Component : LandscapeComponents)
+	for (ULandscapeComponent* Component : LandscapeComponents)
 	{
 		uint32 Hash = Component->CalculatePhysicalMaterialTaskHash();
 		if (Component->PhysicalMaterialHash != Hash || Component->PhysicalMaterialTask.IsValid())
@@ -5767,7 +5797,7 @@ void ALandscapeProxy::UpdatePhysicalMaterialTasksStatus(TSet<ULandscapeComponent
 
 	if (OutdatedCount == 0)
 	{
-		for (auto Component : LandscapeComponents)
+		for (ULandscapeComponent* Component : LandscapeComponents)
 		{
 			const bool bIsDirty = Component->GetPackage()->IsDirty();
 			if (Component->LastSavedPhysicalMaterialHash != Component->PhysicalMaterialHash && !bIsDirty)
diff --git a/Engine/Source/Runtime/Landscape/Private/LandscapeEdit.cpp b/Engine/Source/Runtime/Landscape/Private/LandscapeEdit.cpp
--- a/Engine/Source/Runtime/Landscape/Private/LandscapeEdit.cpp
+++ b/Engine/Source/Runtime/Landscape/Private/LandscapeEdit.cpp
@@ -1816,6 +1816,31 @@ bool ULandscapeComponent::CanUpdatePhysicalMaterial()
 	return GetCollisionComponent() != nullptr;
 }
 
+void ULandscapeComponent::FinalizePhysicalMaterial(bool bInImmediatePhysicsRebuild)
+{
+	if (!PhysicalMaterialTask.IsValid())
+	{
+		return;
+	}
+
+	if (!PhysicalMaterialTask.IsComplete())
+	{
+		return;
+	}
+
+	UpdateCollisionPhysicalMaterialData(PhysicalMaterialTask.GetResultMaterials(), PhysicalMaterialTask.GetResultIds());
+	PhysicalMaterialHash = PhysicalMaterialTask.GetHash();
+
+	PhysicalMaterialTask.Release();
+
+	if (bInImmediatePhysicsRebuild)
+	{
+		ULandscapeHeightfieldCollisionComponent* CollisionComponent = GetCollisionComponent();
+		check(CollisionComponent != nullptr);
+
+		CollisionComponent->RecreateCollision();
+	}
+}
 
 void ULandscapeComponent::UpdatePhysicalMaterialTasks()
 {
@@ -1825,8 +1850,6 @@ void ULandscapeComponent::UpdatePhysicalMaterialTasks()
 		if (PhysicalMaterialTask.IsValid())
 		{
 			PhysicalMaterialTask.Release();
-			// Resetting the hash ensures we will re-start the task once we CAN update physical materials.
-			PhysicalMaterialHash = 0;
 		}
 		return;
 	}
@@ -1836,12 +1859,18 @@ void ULandscapeComponent::UpdatePhysicalMaterialTasks()
 	ULandscapeHeightfieldCollisionComponent* CollisionComponent = GetCollisionComponent();
 	check(CollisionComponent != nullptr);
 
-	if (PhysicalMaterialHash != Hash)
+	if (Hash == PhysicalMaterialHash)
+	{
+		check(!PhysicalMaterialTask.IsValid());
+		return;
+	}
+
+	if (PhysicalMaterialTask.GetHash() != Hash)
 	{
 		TArray<UPhysicalMaterial*> PhysicalMaterials;
 		if (GetRenderPhysicalMaterials(PhysicalMaterials))
 		{
-			bool bSuccess = PhysicalMaterialTask.Init(this);
+			bool bSuccess = PhysicalMaterialTask.Init(this, Hash);
 			check(bSuccess && PhysicalMaterialTask.IsValid());
 		}
 		else
@@ -1851,7 +1880,7 @@ void ULandscapeComponent::UpdatePhysicalMaterialTasks()
 			CollisionComponent->PhysicalMaterialRenderData.RemoveBulkData();
 			PhysicalMaterialTask.Release();
 		}
-		PhysicalMaterialHash = Hash;
+	
 	}
 
 	// If we have a current task, update it
@@ -1859,16 +1888,9 @@ void ULandscapeComponent::UpdatePhysicalMaterialTasks()
 	{
 		if (PhysicalMaterialTask.IsComplete())
 		{
-			UpdateCollisionPhysicalMaterialData(PhysicalMaterialTask.GetResultMaterials(), PhysicalMaterialTask.GetResultIds());
-
-			PhysicalMaterialTask.Release();
-
 			// Potentially, we do not force an update of the physics data here (behind a CVar, as we don't necessarily need the 
 			//  information immediately in the editor and update will happen on cook or PIE) :
-			if (CVarLandscapeApplyPhysicalMaterialChangesImmediately.GetValueOnGameThread() != 0)
-			{
-				CollisionComponent->RecreateCollision();
-			}
+			FinalizePhysicalMaterial(CVarLandscapeApplyPhysicalMaterialChangesImmediately.GetValueOnGameThread() != 0);
 		}
 		else
 		{
diff --git a/Engine/Source/Runtime/Landscape/Private/LandscapePhysicalMaterial.cpp b/Engine/Source/Runtime/Landscape/Private/LandscapePhysicalMaterial.cpp
--- a/Engine/Source/Runtime/Landscape/Private/LandscapePhysicalMaterial.cpp
+++ b/Engine/Source/Runtime/Landscape/Private/LandscapePhysicalMaterial.cpp
@@ -550,13 +550,14 @@ public:
 static TGlobalResource< FLandscapePhysicalMaterialRenderTaskPool > GTaskPool;
 
 
-bool FLandscapePhysicalMaterialRenderTask::Init(ULandscapeComponent const* LandscapeComponent)
+bool FLandscapePhysicalMaterialRenderTask::Init(ULandscapeComponent const* LandscapeComponent, uint32 InHash)
 {
 	check(IsInGameThread());
 	if (IsValid())
 	{
 		GTaskPool.Free(*this);
 	}
+	Hash = InHash;
 	return GTaskPool.Allocate(*this, LandscapeComponent);
 }
 
@@ -566,6 +567,7 @@ void FLandscapePhysicalMaterialRenderTask::Release()
 	if (IsValid())
 	{
 		GTaskPool.Free(*this);
+		Hash = 0;
 	}
 }
 
@@ -582,6 +584,11 @@ bool FLandscapePhysicalMaterialRenderTask::IsComplete() const
 	return GTaskPool.Pool[PoolHandle].CompletionState == ECompletionState::Complete;
 }
 
+bool FLandscapePhysicalMaterialRenderTask::IsInProgress() const
+{
+	return IsValid() && !IsComplete();
+}
+
 void FLandscapePhysicalMaterialRenderTask::Tick()
 {
 	UpdateInternal(/*bInFlush = */false);
diff --git a/Engine/Source/Runtime/Landscape/Private/LandscapeSubsystem.cpp b/Engine/Source/Runtime/Landscape/Private/LandscapeSubsystem.cpp
--- a/Engine/Source/Runtime/Landscape/Private/LandscapeSubsystem.cpp
+++ b/Engine/Source/Runtime/Landscape/Private/LandscapeSubsystem.cpp
@@ -424,7 +424,7 @@ void ULandscapeSubsystem::BuildGrassMaps()
 
 void ULandscapeSubsystem::BuildPhysicalMaterial()
 {
-	PhysicalMaterialBuilder->Build();
+	PhysicalMaterialBuilder->Rebuild();
 }
 
 TArray<ALandscapeProxy*> ULandscapeSubsystem::GetOutdatedProxies(UE::Landscape::EOutdatedDataFlags InMatchingOutdatedDataFlags, bool bInMustMatchAllFlags) const
