Adds unmarked landscape shared properties synchronizing.

diff --git a/Engine/Source/Runtime/Landscape/Classes/LandscapeProxy.h b/Engine/Source/Runtime/Landscape/Classes/LandscapeProxy.h
--- a/Engine/Source/Runtime/Landscape/Classes/LandscapeProxy.h
+++ b/Engine/Source/Runtime/Landscape/Classes/LandscapeProxy.h
@@ -1079,7 +1079,7 @@ public:
 
 #if WITH_EDITOR
 	// Copy properties from parent Landscape actor
-	LANDSCAPE_API void CopySharedProperties(ALandscapeProxy* Landscape);
+	LANDSCAPE_API void CopySharedProperties(ALandscapeProxy* InLandscape);
 
 	/**
 	* Enforce property sharing and requirements related to the parent Landscape actor
@@ -1451,6 +1451,13 @@ protected:
 	/** Fixup any internal representation for shared properties. Will be used when deprecating or renaming a property. */
 	virtual void FixupOverriddenSharedProperties() {}
 
+	/**
+	* Synchronizes all shared properties that cannot be marked by regular meta tags.
+	* @param The landscape to use as a template.
+	* @return An array containing all synchronized properties.
+	*/
+	LANDSCAPE_API virtual TArray<FName> SynchronizeUnmarkedSharedProperties(ALandscapeProxy* InLandscape);
+
 protected:
 	FLandscapeMaterialChangedDelegate LandscapeMaterialChangedDelegate;
 
diff --git a/Engine/Source/Runtime/Landscape/Private/Landscape.cpp b/Engine/Source/Runtime/Landscape/Private/Landscape.cpp
--- a/Engine/Source/Runtime/Landscape/Private/Landscape.cpp
+++ b/Engine/Source/Runtime/Landscape/Private/Landscape.cpp
@@ -4000,8 +4000,10 @@ namespace UE::Landscape::Private
 	}
 } // namespace UE::Landscape::Private
 
-void ALandscapeProxy::CopySharedProperties(ALandscapeProxy* Landscape)
+void ALandscapeProxy::CopySharedProperties(ALandscapeProxy* InLandscape)
 {
+	SynchronizeUnmarkedSharedProperties(InLandscape);
+
 	for (TFieldIterator<FProperty> PropertyIterator(GetClass()); PropertyIterator; ++PropertyIterator)
 	{
 		FProperty* Property = *PropertyIterator;
@@ -4013,14 +4015,14 @@ void ALandscapeProxy::CopySharedProperties(ALandscapeProxy* Landscape)
 
 		if (IsSharedProperty(Property))
 		{
-			UE::Landscape::Private::CopyProperty(Property, Landscape, this);
+			UE::Landscape::Private::CopyProperty(Property, InLandscape, this);
 		}
 	}
 }
 
 TArray<FName> ALandscapeProxy::SynchronizeSharedProperties(ALandscapeProxy* InLandscape)
 {
-	TArray<FName> SynchronizedProperties;
+	TArray<FName> SynchronizedProperties = SynchronizeUnmarkedSharedProperties(InLandscape);
 
 	for (TFieldIterator<FProperty> PropertyIterator(GetClass()); PropertyIterator; ++PropertyIterator)
 	{
diff --git a/Engine/Source/Runtime/Landscape/Private/LandscapeEditLayers.cpp b/Engine/Source/Runtime/Landscape/Private/LandscapeEditLayers.cpp
--- a/Engine/Source/Runtime/Landscape/Private/LandscapeEditLayers.cpp
+++ b/Engine/Source/Runtime/Landscape/Private/LandscapeEditLayers.cpp
@@ -8844,6 +8844,28 @@ void ALandscapeProxy::InitializeLayerWithEmptyContent(const FGuid& InLayerGuid)
 		}
 	}
 }
+
+TArray<FName> ALandscapeProxy::SynchronizeUnmarkedSharedProperties(ALandscapeProxy* InLandscape)
+{
+	check(InLandscape != nullptr);
+	TArray<FName> SynchronizedProperties;
+	USceneComponent* OwnRootComponent = GetRootComponent();
+	USceneComponent* ProxyRootComponent = InLandscape->GetRootComponent();
+
+	if ((OwnRootComponent != nullptr) && (ProxyRootComponent != nullptr))
+	{
+		FVector ProxyScale3D = ProxyRootComponent->GetComponentToWorld().GetScale3D();
+
+		if (!OwnRootComponent->GetRelativeScale3D().Equals(ProxyScale3D))
+		{
+			OwnRootComponent->SetRelativeScale3D(ProxyScale3D);
+			SynchronizedProperties.Emplace(TEXT("RelativeScale3D"));
+		}
+	}
+
+	return SynchronizedProperties;
+}
+
 #endif
 
 void ALandscape::BeginDestroy()
