Landscape : Ensure ToolTexture creation is undoable.

Create a zero filled Tool texture which is correctly added to the Undo history. Causes a slight sutter in mask painting when creating a new mask texture per component but I think this is better than the render artifacts on undo

diff --git a/Engine/Source/Runtime/Landscape/Private/LandscapeEdit.cpp b/Engine/Source/Runtime/Landscape/Private/LandscapeEdit.cpp
--- a/Engine/Source/Runtime/Landscape/Private/LandscapeEdit.cpp
+++ b/Engine/Source/Runtime/Landscape/Private/LandscapeEdit.cpp
@@ -7958,7 +7958,13 @@ UTexture2D* ALandscapeProxy::CreateLandscapeToolTexture(int32 InSizeX, int32 InS
 {
 	UObject* TexOuter = const_cast<ALandscapeProxy*>(this);
 	UTexture2D* NewTexture = NewObject<UTexture2D>(TexOuter, GenerateUniqueLandscapeTextureName(TexOuter, InLODGroup));
-	NewTexture->Source.Init(InSizeX, InSizeY, 1, 1, InFormat);
+
+	int32 BytesPerPixel = FTextureSource::GetBytesPerPixel(InFormat);
+	int32 ZeroBufferSize = BytesPerPixel * InSizeX * InSizeY;
+	std::unique_ptr<uint8[]> ZeroBuffer(new uint8[ZeroBufferSize]);
+	memset(ZeroBuffer.get(), 0, ZeroBufferSize);
+	
+	NewTexture->Source.Init(InSizeX, InSizeY, 1, 1, InFormat, ZeroBuffer.get());
 	NewTexture->SRGB = false;
 	NewTexture->CompressionNone = true;
 	NewTexture->MipGenSettings = TMGS_NoMipmaps;
diff --git a/Engine/Source/Runtime/Landscape/Private/LandscapeEditInterface.cpp b/Engine/Source/Runtime/Landscape/Private/LandscapeEditInterface.cpp
--- a/Engine/Source/Runtime/Landscape/Private/LandscapeEditInterface.cpp
+++ b/Engine/Source/Runtime/Landscape/Private/LandscapeEditInterface.cpp
@@ -4753,8 +4753,7 @@ void FLandscapeEditDataInterface::SetEditToolTextureData(int32 X1, int32 Y1, int
 				int32 WeightmapSize = (Component->SubsectionSizeQuads+1) * Component->NumSubsections;
 				EditToolTexture = Component->GetLandscapeProxy()->CreateLandscapeToolTexture(WeightmapSize, WeightmapSize, InTextureGroup, TSF_G8);
 				EditToolTexture->PostEditChange();
-
-				ZeroTexture(EditToolTexture);
+				
 				Component->UpdateEditToolRenderData();
 			}
 
