Landscape : Fix non unity build (favor UE memory functions)

diff --git a/Engine/Source/Runtime/Landscape/Private/LandscapeEdit.cpp b/Engine/Source/Runtime/Landscape/Private/LandscapeEdit.cpp
--- a/Engine/Source/Runtime/Landscape/Private/LandscapeEdit.cpp
+++ b/Engine/Source/Runtime/Landscape/Private/LandscapeEdit.cpp
@@ -7961,10 +7961,10 @@ UTexture2D* ALandscapeProxy::CreateLandscapeToolTexture(int32 InSizeX, int32 InS
 
 	int32 BytesPerPixel = FTextureSource::GetBytesPerPixel(InFormat);
 	int32 ZeroBufferSize = BytesPerPixel * InSizeX * InSizeY;
-	std::unique_ptr<uint8[]> ZeroBuffer(new uint8[ZeroBufferSize]);
-	memset(ZeroBuffer.get(), 0, ZeroBufferSize);
-	
-	NewTexture->Source.Init(InSizeX, InSizeY, 1, 1, InFormat, ZeroBuffer.get());
+	uint8* ZeroBuffer = reinterpret_cast<uint8*>(FMemory::MallocZeroed(ZeroBufferSize));
+	NewTexture->Source.Init(InSizeX, InSizeY, 1, 1, InFormat, ZeroBuffer);
+	FMemory::Free(ZeroBuffer);
+
 	NewTexture->SRGB = false;
 	NewTexture->CompressionNone = true;
 	NewTexture->MipGenSettings = TMGS_NoMipmaps;
