Landscape : Only remove landscape layer mipchain if requested by landscape.StripLayerMipsOnLoad CVar

diff --git a/Engine/Source/Runtime/Landscape/Private/Landscape.cpp b/Engine/Source/Runtime/Landscape/Private/Landscape.cpp
--- a/Engine/Source/Runtime/Landscape/Private/Landscape.cpp
+++ b/Engine/Source/Runtime/Landscape/Private/Landscape.cpp
@@ -179,6 +179,11 @@ static FAutoConsoleVariable CVarSilenceSharedPropertyDeprecationFixup(
 	TEXT("landscape.SilenceSharedPropertyDeprecationFixup"),
 	true,
 	TEXT("Silently performs the fixup of discrepancies in shared properties when handling data modified before the enforcement introduction."));
+
+static FAutoConsoleVariable CVarStripLayerTextureMipsOnLoad(
+	TEXT("landscape.StripLayerMipsOnLoad"),
+	false,
+	TEXT("Remove (on load) the mip chain from textures used in layers which don't require them"));
 #endif // WITH_EDITOR
 
 int32 GRenderNaniteLandscape = 1;
@@ -1175,30 +1180,33 @@ void ULandscapeComponent::PostLoad()
 		}
 	}
 
-	auto DropMipChain = [](UTexture2D* InTexture) 
+	if (CVarStripLayerTextureMipsOnLoad->GetBool())
 	{
-		if (InTexture->Source.GetNumMips() <= 1)
+		auto DropMipChain = [](UTexture2D* InTexture)
 		{
-			return;
-		}
+			if (InTexture->Source.GetNumMips() <= 1)
+			{
+				return;
+			}
 
-		TArray64<uint8> TopMipData;
-		InTexture->Source.GetMipData(TopMipData, 0);
+			TArray64<uint8> TopMipData;
+			InTexture->Source.GetMipData(TopMipData, 0);
 
-		InTexture->PreEditChange(nullptr);
-		InTexture->Source.Init(InTexture->Source.GetSizeX(), InTexture->Source.GetSizeY(), 1, 1, InTexture->Source.GetFormat(), TopMipData.GetData());
-		InTexture->UpdateResource();
+			InTexture->PreEditChange(nullptr);
+			InTexture->Source.Init(InTexture->Source.GetSizeX(), InTexture->Source.GetSizeY(), 1, 1, InTexture->Source.GetFormat(), TopMipData.GetData());
+			InTexture->UpdateResource();
 
-		InTexture->PostEditChange();
-	};
+			InTexture->PostEditChange();
+		};
 
-	// Remove Non zero mip levels found in layer textures
-	for (auto& LayerIt : LayersData)
-	{
-		DropMipChain(LayerIt.Value.HeightmapData.Texture);
-		for (int32 i = 0; i < LayerIt.Value.WeightmapData.Textures.Num(); ++i)
+		// Remove Non zero mip levels found in layer textures
+		for (auto& LayerIt : LayersData)
 		{
-			DropMipChain(LayerIt.Value.WeightmapData.Textures[i]);
+			DropMipChain(LayerIt.Value.HeightmapData.Texture);
+			for (int32 i = 0; i < LayerIt.Value.WeightmapData.Textures.Num(); ++i)
+			{
+				DropMipChain(LayerIt.Value.WeightmapData.Textures[i]);
+			}
 		}
 	}
 	
