Fix landscape layer selection actions not being recorded in Undo history

diff --git a/Engine/Source/Runtime/Landscape/Private/LandscapeEditLayers.cpp b/Engine/Source/Runtime/Landscape/Private/LandscapeEditLayers.cpp
--- a/Engine/Source/Runtime/Landscape/Private/LandscapeEditLayers.cpp
+++ b/Engine/Source/Runtime/Landscape/Private/LandscapeEditLayers.cpp
@@ -9336,11 +9336,24 @@ void ALandscape::ShowOnlySelectedLayer(int32 InLayerIndex)
 	const FLandscapeLayer* VisibleLayer = GetLayer(InLayerIndex);
 	if (VisibleLayer)
 	{
+		bool bModified = false;
 		for (FLandscapeLayer& Layer : LandscapeLayers)
 		{
-			Layer.bVisible = (&Layer == VisibleLayer);
+			bool bDesiredVisible = (&Layer == VisibleLayer);
+			if (Layer.bVisible != bDesiredVisible)
+			{
+				if (!bModified)
+				{
+					Modify();
+					bModified = true;
+				}
+				Layer.bVisible = bDesiredVisible;
+			}
+		}
+		if (bModified)
+		{
+			RequestLayersContentUpdateForceAll();
 		}
-		RequestLayersContentUpdateForceAll();
 	}
 }
 
@@ -9348,11 +9361,23 @@ void ALandscape::ShowAllLayers()
 {
 	if (LandscapeLayers.Num() > 0)
 	{
+		bool bModified = false;
 		for (FLandscapeLayer& Layer : LandscapeLayers)
 		{
-			Layer.bVisible = true;
+			if (Layer.bVisible != true)
+			{
+				if (!bModified)
+				{
+					Modify();
+					bModified = true;
+				}
+				Layer.bVisible = true;
+			}
+		}
+		if (bModified)
+		{
+			RequestLayersContentUpdateForceAll();
 		}
-		RequestLayersContentUpdateForceAll();
 	}
 }
 
