Fix crash when using Weightmaps stripped to 1 mip

diff --git a/Engine/Source/Runtime/Landscape/Private/LandscapeEdit.cpp b/Engine/Source/Runtime/Landscape/Private/LandscapeEdit.cpp
--- a/Engine/Source/Runtime/Landscape/Private/LandscapeEdit.cpp
+++ b/Engine/Source/Runtime/Landscape/Private/LandscapeEdit.cpp
@@ -2311,6 +2311,9 @@ void ULandscapeComponent::UpdateMipsTempl(int32 InNumSubsections, int32 InSubsec
 	// Find the maximum mip where each texel's data comes from just one subsection.
 	int32 MaxWholeSubsectionMip = FMath::FloorLog2(InSubsectionSizeQuads + 1) - 1;
 
+	// clamp to actual number of mips
+	MaxWholeSubsectionMip = FMath::Min(MaxWholeSubsectionMip, TextureMipData.Num() - 1);
+
 	// Update the mip where each texel's data comes from just one subsection.
 	for (int32 SubsectionY = 0; SubsectionY < InNumSubsections; SubsectionY++)
 	{
@@ -2429,7 +2432,7 @@ void ULandscapeComponent::UpdateMipsTempl(int32 InNumSubsections, int32 InSubsec
 
 	// Handle mips that have texels from multiple subsections
 	// not valid weight data, so just average the texels of the previous mip.
-	for (int32 Mip = MaxWholeSubsectionMip + 1;; ++Mip)
+	for (int32 Mip = MaxWholeSubsectionMip + 1; Mip < TextureMipData.Num(); ++Mip)
 	{
 		int32 MipSubsectionSizeQuads = ((InSubsectionSizeQuads + 1) >> Mip) - 1;
 		checkSlow(MipSubsectionSizeQuads <= 0);
