Fix foliage being removed when Landscape height changes.

diff --git a/Engine/Source/Runtime/Landscape/Private/LandscapeCollision.cpp b/Engine/Source/Runtime/Landscape/Private/LandscapeCollision.cpp
--- a/Engine/Source/Runtime/Landscape/Private/LandscapeCollision.cpp
+++ b/Engine/Source/Runtime/Landscape/Private/LandscapeCollision.cpp
@@ -2058,26 +2058,29 @@ void ULandscapeHeightfieldCollisionComponent::SnapFoliageInstances(const FBox& I
 			const auto* InstanceSet = MeshInfo.ComponentHash.Find(BaseId);
 			if (InstanceSet)
 			{
+				const FVector ZUnitAxis = GetOwner()->GetRootComponent()->GetComponentTransform().GetUnitAxis(EAxis::Z);
 				const float TraceExtentSize = static_cast<float>(Bounds.SphereRadius) * 2.f + 10.f; // extend a little
-				const FVector TraceVector = GetOwner()->GetRootComponent()->GetComponentTransform().GetUnitAxis(EAxis::Z) * TraceExtentSize;
+				const FVector TraceVector = ZUnitAxis * TraceExtentSize;
 
 				TArray<int32> InstancesToRemove;
 				TSet<UHierarchicalInstancedStaticMeshComponent*> AffectedFoliageComponents;
-				
+
 				bool bIsMeshInfoDirty = false;
 				for (int32 InstanceIndex : *InstanceSet)
 				{
 					FFoliageInstance& Instance = MeshInfo.Instances[InstanceIndex];
 
 					// Test location should remove any Z offset
-					FVector TestLocation = FMath::Abs(Instance.ZOffset) > KINDA_SMALL_NUMBER
+					FVector InstanceLocation = FMath::Abs(Instance.ZOffset) > KINDA_SMALL_NUMBER
 						? Instance.GetInstanceWorldTransform().TransformPosition(FVector(0, 0, -Instance.ZOffset))
 						: Instance.Location;
 
-					if (InInstanceBox.IsInside(TestLocation))
+					if (InInstanceBox.IsInside(InstanceLocation))
 					{
-						FVector Start = TestLocation + TraceVector;
-						FVector End = TestLocation - TraceVector;
+						const double HitDistance = FVector::DotProduct((Bounds.Origin - InstanceLocation), ZUnitAxis);
+						const FVector TestLocation = InstanceLocation + ZUnitAxis * HitDistance;
+						const FVector Start = TestLocation + TraceVector;
+						const FVector End = TestLocation - TraceVector;
 
 						TArray<FHitResult> Results;
 						UWorld* World = GetWorld();
@@ -2091,7 +2094,7 @@ void ULandscapeHeightfieldCollisionComponent::SnapFoliageInstances(const FBox& I
 							if (Hit.Component == this)
 							{
 								bFoundHit = true;
-								if ((TestLocation - Hit.Location).SizeSquared() > KINDA_SMALL_NUMBER)
+								if ((InstanceLocation - Hit.Location).SizeSquared() > KINDA_SMALL_NUMBER)
 								{
 									IFA->Modify();
 
