Fix for crash when cooking when landscape weightmap allocations are invalid

diff --git a/Engine/Source/Runtime/Landscape/Private/LandscapeUtils.cpp b/Engine/Source/Runtime/Landscape/Private/LandscapeUtils.cpp
--- a/Engine/Source/Runtime/Landscape/Private/LandscapeUtils.cpp
+++ b/Engine/Source/Runtime/Landscape/Private/LandscapeUtils.cpp
@@ -99,7 +99,9 @@ bool FBatchTextureCopy::AddWeightmapCopy(UTexture2D* InDestination, int8 InDesti
 	// Find the proper Source Texture and channel from Layer Allocations
 	for (const FWeightmapLayerAllocationInfo& ComponentWeightmapLayerAllocation : ComponentWeightmapLayerAllocations)
 	{
-		if (ComponentWeightmapLayerAllocation.LayerInfo == InLayerInfo)
+		if ((ComponentWeightmapLayerAllocation.LayerInfo == InLayerInfo) &&
+			ComponentWeightmapLayerAllocation.IsAllocated() &&
+			ComponentWeightmapTextures.IsValidIndex(ComponentWeightmapLayerAllocation.WeightmapTextureIndex))
 		{
 			CopyRequest.Source = ComponentWeightmapTextures[ComponentWeightmapLayerAllocation.WeightmapTextureIndex];
 			SourceChannel = ComponentWeightmapLayerAllocation.WeightmapTextureChannel;
