Cache landscape culling setting inside landscape proxy to avoid incosistency on whether culling is enabled or not
Re-enable landscape culling

diff --git a/Engine/Source/Runtime/Landscape/Private/LandscapeRender.cpp b/Engine/Source/Runtime/Landscape/Private/LandscapeRender.cpp
--- a/Engine/Source/Runtime/Landscape/Private/LandscapeRender.cpp
+++ b/Engine/Source/Runtime/Landscape/Private/LandscapeRender.cpp
@@ -1060,6 +1060,7 @@ FLandscapeComponentSceneProxy::FLandscapeComponentSceneProxy(ULandscapeComponent
 	, ComponentSizeVerts(InComponent->ComponentSizeQuads + 1)
 	, SectionBase(InComponent->GetSectionBase())
 	, LandscapeComponent(InComponent)
+	, bUsesLandscapeCulling(false)
 	, WeightmapScaleBias(InComponent->WeightmapScaleBias)
 	, VisibilityWeightmapTexture(nullptr)
 	, VisibilityWeightmapChannel(-1)
@@ -1100,6 +1101,7 @@ FLandscapeComponentSceneProxy::FLandscapeComponentSceneProxy(ULandscapeComponent
 	}
 
 	bNaniteActive = InComponent->IsNaniteActive();
+	bUsesLandscapeCulling = (!bNaniteActive && Culling::UseCulling(GetScene().GetShaderPlatform()));
 
 	EnableGPUSceneSupportFlags();
 
@@ -1237,7 +1239,7 @@ FLandscapeComponentSceneProxy::FLandscapeComponentSceneProxy(ULandscapeComponent
 	}
 
 	// Landscape GPU culling uses VF that requires primitive UB
-	bVFRequiresPrimitiveUniformBuffer |= (!bNaniteActive && Culling::UseCulling(GetScene().GetShaderPlatform()));
+	bVFRequiresPrimitiveUniformBuffer |= bUsesLandscapeCulling;
 
 	ComponentLightInfo = MakeUnique<FLandscapeLCI>(InComponent, FeatureLevel, bVFRequiresPrimitiveUniformBuffer != 0);
 	check(ComponentLightInfo);
@@ -1400,7 +1402,7 @@ void FLandscapeComponentSceneProxy::CreateRenderThreadResources(FRHICommandListB
 	VertexFactory = SharedBuffers->VertexFactory;
 	FixedGridVertexFactory = SharedBuffers->FixedGridVertexFactory;
 
-	if (!bNaniteActive && Culling::UseCulling(GetScene().GetShaderPlatform()))
+	if (bUsesLandscapeCulling)
 	{
 		Culling::RegisterLandscape(RHICmdList, *SharedBuffers, FeatureLevel, LandscapeKey, SubsectionSizeVerts, NumSubsections);
 	}
@@ -2116,7 +2118,7 @@ bool FLandscapeComponentSceneProxy::GetStaticMeshElement(int32 LODIndex, bool bF
 		BatchElement.MinVertexIndex = SharedBuffers->IndexRanges[LODIndex].MinIndexFull;
 		BatchElement.MaxVertexIndex = SharedBuffers->IndexRanges[LODIndex].MaxIndexFull;
 				
-		if (!bForToolMesh && !bNaniteActive && Culling::UseCulling(GetScene().GetShaderPlatform()))
+		if (!bForToolMesh && bUsesLandscapeCulling)
 		{
 			Culling::SetupMeshBatch(*SharedBuffers, MeshBatch);
 		}
diff --git a/Engine/Source/Runtime/Landscape/Public/LandscapeRender.h b/Engine/Source/Runtime/Landscape/Public/LandscapeRender.h
--- a/Engine/Source/Runtime/Landscape/Public/LandscapeRender.h
+++ b/Engine/Source/Runtime/Landscape/Public/LandscapeRender.h
@@ -762,6 +762,7 @@ protected:
 	TArray<FLandscapeBatchElementParams> StaticBatchParamArray;
 
 	bool bNaniteActive;
+	bool bUsesLandscapeCulling;
 
 
 #if WITH_EDITOR
