Fix for nanite enabled landscape not rendering at all after switching preview rendering mode to SM5 in editor (running D3D12)

diff --git a/Engine/Source/Runtime/Landscape/Private/Landscape.cpp b/Engine/Source/Runtime/Landscape/Private/Landscape.cpp
--- a/Engine/Source/Runtime/Landscape/Private/Landscape.cpp
+++ b/Engine/Source/Runtime/Landscape/Private/Landscape.cpp
@@ -5897,16 +5897,13 @@ void ALandscapeProxy::UpdateRenderingMethod()
 	bool bNaniteActive = false;
 	if ((GRenderNaniteLandscape != 0) && NaniteComponent)
 	{
-		bNaniteActive = UseNanite(GShaderPlatformForFeatureLevel[GMaxRHIFeatureLevel]);
+		bNaniteActive = UseNanite(GShaderPlatformForFeatureLevel[GEngine->GetDefaultWorldFeatureLevel()]);
 #if WITH_EDITOR
 		if (ALandscape* LandscapeActor = GetLandscapeActor())
 		{
 			if (UWorld* World = LandscapeActor->GetWorld())
 			{
-				if ((GEngine->GetDefaultWorldFeatureLevel() == ERHIFeatureLevel::ES3_1) || (World->GetFeatureLevel() <= ERHIFeatureLevel::ES3_1))
-				{
-					bNaniteActive = UseNanite(GShaderPlatformForFeatureLevel[ERHIFeatureLevel::ES3_1]);
-				}
+				bNaniteActive = UseNanite(GShaderPlatformForFeatureLevel[World->GetFeatureLevel()]);
 			}
 		}
 #endif //WITH_EDITOR
