Don't delete cooked collision data on use

diff --git a/Engine/Source/Runtime/Landscape/Classes/LandscapeHeightfieldCollisionComponent.h b/Engine/Source/Runtime/Landscape/Classes/LandscapeHeightfieldCollisionComponent.h
--- a/Engine/Source/Runtime/Landscape/Classes/LandscapeHeightfieldCollisionComponent.h
+++ b/Engine/Source/Runtime/Landscape/Classes/LandscapeHeightfieldCollisionComponent.h
@@ -156,6 +156,7 @@ public:
 	 *	Cooked HeightField data. Serialized only with cooked content 
 	 *	Stored as array instead of BulkData to take advantage of precaching during async loading
 	 */
+	bool CookedCollisionDataWasDeleted = false;
 	TArray<uint8>								CookedCollisionData;
 	
 	/** This is a list of physical materials that is actually used by a cooked HeightField */
diff --git a/Engine/Source/Runtime/Landscape/Private/LandscapeCollision.cpp b/Engine/Source/Runtime/Landscape/Private/LandscapeCollision.cpp
--- a/Engine/Source/Runtime/Landscape/Private/LandscapeCollision.cpp
+++ b/Engine/Source/Runtime/Landscape/Private/LandscapeCollision.cpp
@@ -101,6 +101,11 @@ static TAutoConsoleVariable<bool> CVarLandscapeCollisionMeshShowPhysicalMaterial
 	TEXT("When enabled, vertex colors of the collision mesh are chosen based on the physical material"),
 	ECVF_RenderThreadSafe);
 
+static TAutoConsoleVariable<bool> CVarLandscapeCollisionDeleteCookedDataAfterUse(
+	TEXT("landscape.Collision.DeleteCookedDataAfterUse"),
+	true,
+	TEXT("When enabled, landscape cooked collision data is deleted after it is used to construct a physics heightfield.  Enabling this will reclaim memory, but may cause problems if the components are ever re-registered and collision must be recreated."));
+
 
 #if ENABLE_COOK_STATS
 namespace LandscapeCollisionCookStats
@@ -856,7 +861,15 @@ void ULandscapeHeightfieldCollisionComponent::CreateCollisionObject()
 			SpeculativeDDCRequest.Reset();
 #endif //WITH_EDITOR
 
-			if (CookedCollisionData.Num())
+			if (CookedCollisionData.IsEmpty())
+			{
+				if (CookedCollisionDataWasDeleted)
+				{
+					// only complain if we actually deleted the data.. otherwise it may have been intentional?
+					UE_LOG(LogLandscape, Warning, TEXT("Tried to create collision for component '%s', but the collision data was deleted!  To avoid set: landscape.Collision.DeleteCookedDataAfterUse 0"), *GetName());
+				}
+			}
+			else
 			{
 				HeightfieldRef = GSharedHeightfieldRefs.Add(HeightfieldGuid, new FHeightfieldGeometryRef(HeightfieldGuid));
 
@@ -889,7 +902,11 @@ void ULandscapeHeightfieldCollisionComponent::CreateCollisionObject()
 				// In cooked builds created collision object will never be deleted while component is alive, so we don't need this data anymore
 				if(FPlatformProperties::RequiresCookedData() || World->IsGameWorld())
 				{
-					CookedCollisionData.Empty();
+					if (CVarLandscapeCollisionDeleteCookedDataAfterUse.GetValueOnGameThread())
+					{
+						CookedCollisionData.Empty();
+						CookedCollisionDataWasDeleted = true;
+					}
 				}
 
 #if WITH_EDITOR
@@ -1537,7 +1554,15 @@ void ULandscapeMeshCollisionComponent::CreateCollisionObject()
 			CookCollisionData(PhysicsFormatName, false, bCheckDDC, CookedCollisionData, MutableView(CookedPhysicalMaterials));
 #endif // WITH_EDITOR
 
-			if (CookedCollisionData.Num())
+			if (CookedCollisionData.IsEmpty())
+			{
+				if (CookedCollisionDataWasDeleted)
+				{
+					// only complain if we actually deleted the data.. otherwise it may have been intentional?
+					UE_LOG(LogLandscape, Warning, TEXT("Tried to create collision for component '%s', but the collision data was deleted!  To avoid set: landscape.Collision.DeleteCookedDataAfterUse 0"), *GetName());
+				}
+			}
+			else
 			{
 				MeshRef = GSharedMeshRefs.Add(MeshGuid, new FTriMeshGeometryRef(MeshGuid));
 
@@ -1555,7 +1580,11 @@ void ULandscapeMeshCollisionComponent::CreateCollisionObject()
 				// In cooked builds created collision object will never be deleted while component is alive, so we don't need this data anymore
 				if (FPlatformProperties::RequiresCookedData() || GetWorld()->IsGameWorld())
 				{
-					CookedCollisionData.Empty();
+					if (CVarLandscapeCollisionDeleteCookedDataAfterUse.GetValueOnGameThread())
+					{
+						CookedCollisionData.Empty();
+						CookedCollisionDataWasDeleted = true;
+					}
 				}
 
 #if WITH_EDITOR
