Ensure the LOD level is within valid ranges when building Nanite landscape data

diff --git a/Engine/Source/Runtime/Landscape/Private/LandscapeEdit.cpp b/Engine/Source/Runtime/Landscape/Private/LandscapeEdit.cpp
--- a/Engine/Source/Runtime/Landscape/Private/LandscapeEdit.cpp
+++ b/Engine/Source/Runtime/Landscape/Private/LandscapeEdit.cpp
@@ -3778,7 +3778,10 @@ TSharedRef<UE::Landscape::Nanite::FAsyncBuildData> ALandscapeProxy::MakeAsyncNan
 {
 	TSharedRef<UE::Landscape::Nanite::FAsyncBuildData> AsyncBuildData = MakeShared<UE::Landscape::Nanite::FAsyncBuildData>();
 
-	AsyncBuildData->LOD = InLODToExport;
+	// Make sure the requested LOD is valid
+	int32 FinalLODToExport = FMath::Clamp<int32>(InLODToExport, 0, FMath::CeilLogTwo(SubsectionSizeQuads + 1) - 1);
+
+	AsyncBuildData->LOD = FinalLODToExport;
 	AsyncBuildData->LandscapeWeakRef = MakeWeakObjectPtr(const_cast<ALandscapeProxy*>(this));
 	AsyncBuildData->LandscapeSubSystemWeakRef = MakeWeakObjectPtr(GetWorld()->GetSubsystem<ULandscapeSubsystem>());
 
@@ -3806,10 +3809,10 @@ TSharedRef<UE::Landscape::Nanite::FAsyncBuildData> ALandscapeProxy::MakeAsyncNan
 
 	// take a copy of the height and visility data for each component.
 	// Add an sync version of ForEachComponent?  
-	ForEachComponent<ULandscapeComponent>(true, [AsyncBuildData, InLODToExport](ULandscapeComponent* LandscapeComponent)
+	ForEachComponent<ULandscapeComponent>(true, [AsyncBuildData, FinalLODToExport](ULandscapeComponent* LandscapeComponent)
 		{
 			TRACE_CPUPROFILER_EVENT_SCOPE(ALandscapeProxy::MakeAsyncBuildData-CopyHeightAndVisibility);
-			FLandscapeComponentDataInterface DataInterface(LandscapeComponent, InLODToExport, false);
+			FLandscapeComponentDataInterface DataInterface(LandscapeComponent, FinalLODToExport, false);
 
 			UE::Landscape::Nanite::FAsyncComponentData AsyncComponentData;
 
@@ -3817,8 +3820,8 @@ TSharedRef<UE::Landscape::Nanite::FAsyncBuildData> ALandscapeProxy::MakeAsyncNan
 			DataInterface.GetWeightmapTextureData(LandscapeComponent->GetVisibilityLayer(), AsyncComponentData.Visibility);
 
 		
-			AsyncComponentData.ComponentDataInterface = MakeShared<FLandscapeComponentDataInterfaceBase>(LandscapeComponent, InLODToExport, false);
-			int32 HeightmapSize = ((LandscapeComponent->SubsectionSizeQuads + 1) * LandscapeComponent->NumSubsections) >> InLODToExport;
+			AsyncComponentData.ComponentDataInterface = MakeShared<FLandscapeComponentDataInterfaceBase>(LandscapeComponent, FinalLODToExport, false);
+			int32 HeightmapSize = ((LandscapeComponent->SubsectionSizeQuads + 1) * LandscapeComponent->NumSubsections) >> FinalLODToExport;
 			AsyncComponentData.ComponentDataInterface->HeightmapStride = HeightmapSize;
 			AsyncComponentData.ComponentDataInterface->HeightmapComponentOffsetX = 0;
 			AsyncComponentData.ComponentDataInterface->HeightmapComponentOffsetY = 0;
@@ -3832,10 +3835,7 @@ TSharedRef<UE::Landscape::Nanite::FAsyncBuildData> ALandscapeProxy::MakeAsyncNan
 bool ALandscapeProxy::ExportToRawMesh(const FRawMeshExportParams& InExportParams, FMeshDescription& OutRawMesh) const
 {
 	FRawMeshExportParams ExportParams = InExportParams;
-	if (InExportParams.ExportLOD != INDEX_NONE)
-	{
-		ExportParams.ExportLOD = FMath::Clamp<int32>(InExportParams.ExportLOD, 0, FMath::CeilLogTwo(SubsectionSizeQuads + 1) - 1);
-	}
+	ExportParams.ExportLOD = FMath::Clamp<int32>(InExportParams.ExportLOD, 0, FMath::CeilLogTwo(SubsectionSizeQuads + 1) - 1);
 
 	TSharedRef<UE::Landscape::Nanite::FAsyncBuildData> AsyncBuildData = MakeAsyncNaniteBuildData(ExportParams.ExportLOD);
 	return ExportToRawMeshDataCopy(ExportParams, OutRawMesh, AsyncBuildData.Get());
@@ -3843,6 +3843,7 @@ bool ALandscapeProxy::ExportToRawMesh(const FRawMeshExportParams& InExportParams
 
 bool ALandscapeProxy::ExportToRawMeshDataCopy(const FRawMeshExportParams& InExportParams, FMeshDescription& OutRawMesh, const UE::Landscape::Nanite::FAsyncBuildData& AsyncData) const
 {
+	check(InExportParams.ExportLOD == AsyncData.LOD);
 	if (CVarLandscapeMarchingSquaresVisibility.GetValueOnAnyThread())
 	{
 		return ExportToRawMeshDataCopyNew(InExportParams, OutRawMesh, AsyncData);
