Enable compressed landscape texture data for all platforms (not just windows)

diff --git a/Engine/Source/Runtime/Landscape/Private/Landscape.cpp b/Engine/Source/Runtime/Landscape/Private/Landscape.cpp
--- a/Engine/Source/Runtime/Landscape/Private/Landscape.cpp
+++ b/Engine/Source/Runtime/Landscape/Private/Landscape.cpp
@@ -3221,8 +3221,7 @@ void ALandscapeProxy::PreSave(FObjectPreSaveContext ObjectSaveContext)
 	{
 		FString PlatformName = ObjectSaveContext.GetTargetPlatform()->PlatformName();
 
-		// TODO [chris.tchou] : remove the limitation to only work on windows, once we validate other platforms work as well...
-		if (PlatformName.StartsWith(TEXT("Windows")))
+		// TODO [chris.tchou] : add user and/or platform data controllable toggles here
 		{
 			for (ULandscapeComponent* LandscapeComponent : LandscapeComponents)
 			{
