Landscape : Optionally scale grass instance size based on weight ( relative density )

diff --git a/Engine/Source/Runtime/Landscape/Classes/LandscapeGrassType.h b/Engine/Source/Runtime/Landscape/Classes/LandscapeGrassType.h
--- a/Engine/Source/Runtime/Landscape/Classes/LandscapeGrassType.h
+++ b/Engine/Source/Runtime/Landscape/Classes/LandscapeGrassType.h
@@ -91,6 +91,14 @@ struct FGrassVariety
 	UPROPERTY(EditAnywhere, Category=Grass, meta = (EditCondition = "Scaling == EGrassScaling::Free || Scaling == EGrassScaling::LockXY"))
 	FFloatInterval ScaleZ;
 
+	/** If enabled the the scale of instances is reduced as the weight (density) decreases */
+	UPROPERTY(EditAnywhere, AdvancedDisplay, Category = Grass)
+	bool bWeightAttenuatesMaxScale = false;
+
+	/** Modulate the scale of the instances based on weight (normalized density). The weight range (ScaleWeightAttenuation, 1.0) maps to (scaleMin, scaleMax), weight values less than ScaleWeightAttenuation are set to minScale */
+	UPROPERTY(EditAnywhere, AdvancedDisplay, Category = Grass, meta = (UIMin = 0, ClampMin = 0, UIMax = 1, ClampMax = 1, EditCondition = "bWeightAttenuatesMaxScale == true"))
+	float MaxScaleWeightAttenuation = 0.5f;
+
 	/** Whether the grass instances should be placed at random rotation (true) or all at the same rotation (false) */
 	UPROPERTY(EditAnywhere, Category = Grass)
 	bool RandomRotation;
diff --git a/Engine/Source/Runtime/Landscape/Private/LandscapeGrass.cpp b/Engine/Source/Runtime/Landscape/Private/LandscapeGrass.cpp
--- a/Engine/Source/Runtime/Landscape/Private/LandscapeGrass.cpp
+++ b/Engine/Source/Runtime/Landscape/Private/LandscapeGrass.cpp
@@ -2450,6 +2450,8 @@ struct FAsyncGrassBuilder : public FGrassBuilderBase
 	FFloatInterval ScaleX;
 	FFloatInterval ScaleY;
 	FFloatInterval ScaleZ;
+	bool bWeightAttenuatesMaxScale;
+	float MaxScaleWeightAttenuation ;
 	bool RandomRotation;
 	bool RandomScale;
 	bool AlignToSurface;
@@ -2487,6 +2489,8 @@ struct FAsyncGrassBuilder : public FGrassBuilderBase
 		, ScaleX(GrassVariety.ScaleX)
 		, ScaleY(GrassVariety.ScaleY)
 		, ScaleZ(GrassVariety.ScaleZ)
+		, bWeightAttenuatesMaxScale(GrassVariety.bWeightAttenuatesMaxScale)
+		, MaxScaleWeightAttenuation(GrassVariety.MaxScaleWeightAttenuation)
 		, RandomRotation(GrassVariety.RandomRotation)
 		, RandomScale(false)
 		, AlignToSurface(GrassVariety.AlignToSurface)
@@ -2640,26 +2644,32 @@ struct FAsyncGrassBuilder : public FGrassBuilderBase
 		return Result;
 	}
 
-	FVector GetRandomScale() const
+	FVector GetRandomScale(float InWeight) const
 	{
 		FVector Result(1.0f);
 
+		float WeightAttenuationFactor = 1.0f;		
+		if (bWeightAttenuatesMaxScale)
+		{
+			WeightAttenuationFactor = MaxScaleWeightAttenuation < 1.0f ? FMath::Clamp((InWeight - MaxScaleWeightAttenuation) / (1.0f - MaxScaleWeightAttenuation), 0.0f, 1.0f) : 0.0f;
+		}
+		
 		switch (Scaling)
 		{
 		case EGrassScaling::Uniform:
-			Result.X = ScaleX.Interpolate(RandomStream.GetFraction());
+			Result.X = ScaleX.Interpolate(RandomStream.GetFraction() * WeightAttenuationFactor);
 			Result.Y = Result.X;
 			Result.Z = Result.X;
 			break;
 		case EGrassScaling::Free:
-			Result.X = ScaleX.Interpolate(RandomStream.GetFraction());
-			Result.Y = ScaleY.Interpolate(RandomStream.GetFraction());
-			Result.Z = ScaleZ.Interpolate(RandomStream.GetFraction());
+			Result.X = ScaleX.Interpolate(RandomStream.GetFraction() * WeightAttenuationFactor);
+			Result.Y = ScaleY.Interpolate(RandomStream.GetFraction() * WeightAttenuationFactor);
+			Result.Z = ScaleZ.Interpolate(RandomStream.GetFraction() * WeightAttenuationFactor);
 			break;
 		case EGrassScaling::LockXY:
-			Result.X = ScaleX.Interpolate(RandomStream.GetFraction());
+			Result.X = ScaleX.Interpolate(RandomStream.GetFraction() * WeightAttenuationFactor);
 			Result.Y = Result.X;
-			Result.Z = ScaleZ.Interpolate(RandomStream.GetFraction());
+			Result.Z = ScaleZ.Interpolate(RandomStream.GetFraction() * WeightAttenuationFactor);
 			break;
 		default:
 			check(0);
@@ -2715,7 +2725,7 @@ struct FAsyncGrassBuilder : public FGrassBuilderBase
 				bool bKeep = Weight > 0.0f && Weight >= RandomStream.GetFraction() && !IsExcluded(LocationWithHeight);
 				if (bKeep)
 				{
-					const FVector Scale = RandomScale ? GetRandomScale() : DefaultScale;
+					const FVector Scale = RandomScale ? GetRandomScale(Weight) : DefaultScale;
 					const float Rot = RandomRotation ? RandomStream.GetFraction() * 360.0f : 0.0f;
 					const FMatrix BaseXForm = FScaleRotationTranslationMatrix(Scale, FRotator(0.0f, Rot, 0.0f), FVector::ZeroVector);
 					FMatrix OutXForm;
@@ -2756,6 +2766,7 @@ struct FAsyncGrassBuilder : public FGrassBuilderBase
 			{
 				FVector Pos;
 				bool bKeep;
+				float Weight;
 			};
 			TArray<FInstanceLocal> Instances;
 			Instances.AddUninitialized(SqrtMaxInstances * SqrtMaxInstances);
@@ -2778,6 +2789,7 @@ struct FAsyncGrassBuilder : public FGrassBuilderBase
 						float Weight = 0.f;
 						SampleLandscapeAtLocationLocal(Location, Instance.Pos, Weight);
 						Instance.bKeep = Weight > 0.0f && Weight >= RandomStream.GetFraction() && !IsExcluded(Instance.Pos);
+						Instance.Weight = Weight;
 						if (Instance.bKeep)
 						{
 							NumKept++;
@@ -2801,7 +2813,7 @@ struct FAsyncGrassBuilder : public FGrassBuilderBase
 							const FInstanceLocal& Instance = Instances[InstanceIndex];
 							if (Instance.bKeep)
 							{
-								const FVector Scale = RandomScale ? GetRandomScale() : DefaultScale;
+								const FVector Scale = RandomScale ? GetRandomScale(Instance.Weight) : DefaultScale;
 								const float Rot = RandomRotation ? RandomStream.GetFraction() * 360.0f : 0.0f;
 								const FMatrix BaseXForm = FScaleRotationTranslationMatrix(Scale, FRotator(0.0f, Rot, 0.0f), FVector::ZeroVector);
 								FMatrix OutXForm;
