Fix invalid material test in  ALandscapeProxy::PostLoad

diff --git a/Engine/Source/Runtime/Landscape/Private/Landscape.cpp b/Engine/Source/Runtime/Landscape/Private/Landscape.cpp
--- a/Engine/Source/Runtime/Landscape/Private/Landscape.cpp
+++ b/Engine/Source/Runtime/Landscape/Private/Landscape.cpp
@@ -3698,6 +3698,8 @@ void ALandscapeProxy::PostLoad()
 
 	bool bFixedUpInvalidMaterialInstances = false;
 
+	UMaterial* BaseLandscapeMaterial = GetLandscapeMaterial()->GetMaterial();
+
 	for (ULandscapeComponent* Comp : LandscapeComponents)
 	{
 		if (Comp == nullptr)
@@ -3718,7 +3720,7 @@ void ALandscapeProxy::PostLoad()
 		if (!Comp->GetOutermost()->HasAnyPackageFlags(PKG_FilterEditorOnly) && (GIsEditor && FApp::CanEverRender()))
 		{
 			// MaterialInstance is different from the used LandscapeMaterial, we need to update the material as we cannot properly validate used combinations.
-			if (MaterialInstance->GetMaterial() != GetLandscapeMaterial())
+			if (MaterialInstance->GetMaterial() != BaseLandscapeMaterial)
 			{
 				Comp->UpdateMaterialInstances();
 				bFixedUpInvalidMaterialInstances = true;
