Fixed landscape GPU readback pool not being depleted in most cases.

diff --git a/Engine/Source/Runtime/Landscape/Private/LandscapeEditLayers.cpp b/Engine/Source/Runtime/Landscape/Private/LandscapeEditLayers.cpp
--- a/Engine/Source/Runtime/Landscape/Private/LandscapeEditLayers.cpp
+++ b/Engine/Source/Runtime/Landscape/Private/LandscapeEditLayers.cpp
@@ -8169,6 +8169,14 @@ void ALandscape::UpdateLayersContent(bool bInWaitForStreaming, bool bInSkipMonit
 			WaitingForLandscapeTextureResourcesStartTime = -1.0;
 			WaitingForLandscapeBrushResourcesStartTime = -1.0;
 		}
+
+		// If nothing to do, let's do some garbage collecting on async readback tasks so that we slowly get rid of staging textures 
+		//  (don't do it while waiting for read backs because something might prevent us from updating the readbacks (e.g. waiting for resources to compiling...), which would 
+		//  lead to FLandscapeEditReadbackTaskPool's frame count increasing while readback tasks don't have the chance to complete, leading to the "readback leak" warning to incorrectly be triggered) :
+		if ((LayerContentUpdateModes == 0) && !FLandscapeEditLayerReadback::HasWork())
+		{
+			FLandscapeEditLayerReadback::GarbageCollectTasks();
+		}
 	};
 
 	// Note : no early-out allowed before this : even if not actually updating edit layers, we need to poll our resources in order to make sure we register to streaming events when needed: 
@@ -8395,8 +8403,6 @@ void ALandscape::UpdateLayersContent(bool bInWaitForStreaming, bool bInSkipMonit
 
 	// Additional validation that at the end of an update, we haven't screwed up anything in the weightmap allocations/usages : 
 	ValidateProxyLayersWeightmapUsage();
-
-	FLandscapeEditLayerReadback::GarbageCollectTasks();
 }
 
 // not thread safe
diff --git a/Engine/Source/Runtime/Landscape/Private/LandscapeEditReadback.cpp b/Engine/Source/Runtime/Landscape/Private/LandscapeEditReadback.cpp
--- a/Engine/Source/Runtime/Landscape/Private/LandscapeEditReadback.cpp
+++ b/Engine/Source/Runtime/Landscape/Private/LandscapeEditReadback.cpp
@@ -37,7 +37,7 @@ struct FLandscapeEditReadbackTaskImpl
 };
 
 /** Initialize the read back task data that is written by game thread. */
-bool InitTask_GameThread(FLandscapeEditReadbackTaskImpl& Task, UTexture2D const* InTexture, FLandscapeEditLayerReadback::FReadbackContext&& InReadbackContext, uint32 InFrameId)
+void InitTask_GameThread(FLandscapeEditReadbackTaskImpl& Task, UTexture2D const* InTexture, FLandscapeEditLayerReadback::FReadbackContext&& InReadbackContext, uint32 InFrameId)
 {
 	Task.TextureResource = InTexture->GetResource();
 	Task.ReadbackContext = MoveTemp(InReadbackContext);
@@ -46,14 +46,12 @@ bool InitTask_GameThread(FLandscapeEditReadbackTaskImpl& Task, UTexture2D const*
 	Task.NumMips = InTexture->GetNumMips();
 	Task.Format = InTexture->GetPixelFormat();
 	Task.CompletionState = FLandscapeEditReadbackTaskImpl::ECompletionState::None;
-	
-	return Task.TextureResource != nullptr;
 }
 
 /** Initialize the read back task resources. */
 bool InitTask_RenderThread(FLandscapeEditReadbackTaskImpl& Task)
 {
-	if (Task.StagingTextures.Num() == 0 || !Task.StagingTextures[0].IsValid() || Task.StagingTextures[0]->GetSizeXYZ() != FIntVector(Task.Size.X, Task.Size.Y, 1))
+	if (Task.StagingTextures.Num() == 0 || !Task.StagingTextures[0].IsValid() || Task.StagingTextures[0]->GetSizeXYZ() != FIntVector(Task.Size.X, Task.Size.Y, 1) || (Task.StagingTextures[0]->GetFormat() != Task.Format))
 	{
 		Task.StagingTextures.SetNum(Task.NumMips);
 
@@ -190,24 +188,34 @@ public:
 	/** Allocate task data from the pool. */
 	int32 Allocate(UTexture2D const* InTexture, FLandscapeEditLayerReadback::FReadbackContext&& InReadbackContext)
 	{
-		int32 Index = 0;
+		int32 CurrentIndex = 0;
+		int32 BestEntryIndex = INDEX_NONE;
+		FIntVector TextureSize(InTexture->GetSizeX(), InTexture->GetSizeY(), 1);
 		auto ItEnd = Pool.end();
-		for (auto It = Pool.begin(); It != ItEnd; ++It, ++Index)
+		for (auto It = Pool.begin(); It != ItEnd; ++It, ++CurrentIndex)
 		{
-			if ((*It).TextureResource == nullptr)
+			FLandscapeEditReadbackTaskImpl& Task = *It;
+			// If the entry is unused, it's a candidate 
+			if (Task.TextureResource == nullptr)
 			{
-				break;
+				BestEntryIndex = CurrentIndex;
+				// Check the entry's texture size to ensure it's the best possible candidate. If so, no need to look further :
+				if (!Task.StagingTextures.IsEmpty() && Task.StagingTextures[0].IsValid() && (Task.StagingTextures[0]->GetSizeXYZ() == TextureSize) && (Task.Format == InTexture->GetPixelFormat()))
+				{
+					break;
+				}
 			}
 		}
 
-		if (Index == Pool.Num())
+		if (BestEntryIndex == INDEX_NONE)
 		{
 			Pool.Add();
+			BestEntryIndex = Pool.Num() - 1;
 		}
 
-		const bool bSuccess = InitTask_GameThread(Pool[Index], InTexture, MoveTemp(InReadbackContext), FrameCount);
-		AllocCount += bSuccess ? 1: 0;
-		return bSuccess ? Index : -1;
+		InitTask_GameThread(Pool[BestEntryIndex], InTexture, MoveTemp(InReadbackContext), FrameCount);
+		++AllocCount;
+		return BestEntryIndex;
 	}
 
 	/** Return task data to the pool. */
@@ -245,12 +253,15 @@ public:
 					Task->ReadbackContext.Empty();
 					Task->Result.Empty();
 
-					// Release the render resources (which may already be released)
-					ENQUEUE_RENDER_COMMAND(FLandscapeEditLayerReadback_Release)([Task](FRHICommandListImmediate& RHICmdList)
+					if (!Task->StagingTextures.IsEmpty() || Task->ReadbackFence.IsValid())
 					{
-						Task->StagingTextures.Reset();
-						Task->ReadbackFence.SafeRelease();
-					});
+						// Release the render resources (which may already be released)
+						ENQUEUE_RENDER_COMMAND(FLandscapeEditLayerReadback_Release)([Task](FRHICommandListImmediate& RHICmdList)
+						{
+							Task->StagingTextures.Reset();
+							Task->ReadbackFence.SafeRelease();
+						});
+					}
 				}
 			}
 		}
