* Prevented computing the landscape texture hash twice when it has changed
* Fixed landscape getting dirty even when the detected difference is lesser than the detection threshold

diff --git a/Engine/Source/Runtime/Landscape/Private/LandscapeEditLayers.cpp b/Engine/Source/Runtime/Landscape/Private/LandscapeEditLayers.cpp
--- a/Engine/Source/Runtime/Landscape/Private/LandscapeEditLayers.cpp
+++ b/Engine/Source/Runtime/Landscape/Private/LandscapeEditLayers.cpp
@@ -5564,10 +5564,9 @@ bool ALandscape::ResolveLayersTexture(
 					uint64 Hash = 0;
 					if (HasTextureDataChanged(MakeArrayView(reinterpret_cast<const FColor*>(TextureData), TextureSize), MakeArrayView(OutMipsData[MipIndex].GetData(), TextureSize), bIsWeightmap, InCPUReadback->GetHash(), Hash))
 					{
-						TRACE_CPUPROFILER_EVENT_SCOPE(LandscapeLayers_CalculateHash);
-						Hash = FLandscapeEditLayerReadback::CalculateHash((uint8*)OutMipsData[MipIndex].GetData(), TextureSize * sizeof(FColor));
+						bChanged |= InCPUReadback->SetHash(Hash);
+						check(bChanged);
 					}
-					bChanged |= InCPUReadback->SetHash(Hash);
 
 					if (bChanged)
 					{
