Add an option to strip collision and grass data from landscape at cook time (for clients and/or servers).
This is intended to save disk space, when the data is not needed, or can be recreated by user systems at runtime.
(Original review under 26452362 in UE5/main)

diff --git a/Engine/Source/Runtime/Landscape/Classes/LandscapeProxy.h b/Engine/Source/Runtime/Landscape/Classes/LandscapeProxy.h
--- a/Engine/Source/Runtime/Landscape/Classes/LandscapeProxy.h
+++ b/Engine/Source/Runtime/Landscape/Classes/LandscapeProxy.h
@@ -824,6 +824,22 @@ public:
 	UPROPERTY(EditAnywhere, Category = Landscape, AdvancedDisplay, meta = (LandscapeInherited))
 	bool bUseCompressedHeightmapStorage = false;
 
+	/** Strip Physics/collision components when cooked for client */
+	UPROPERTY(EditAnywhere, Category = Landscape, AdvancedDisplay, meta = (LandscapeOverridable))
+	bool bStripPhysicsWhenCookedClient = false;
+
+	/** Strip Physics/collision components when cooked for server */
+	UPROPERTY(EditAnywhere, Category = Landscape, AdvancedDisplay, meta = (LandscapeOverridable))
+	bool bStripPhysicsWhenCookedServer = false;
+
+	/** Strip Grass data when cooked for client */
+	UPROPERTY(EditAnywhere, Category = Landscape, AdvancedDisplay, meta = (LandscapeOverridable))
+	bool bStripGrassWhenCookedClient = false;
+
+	/** Strip Grass data when cooked for server */
+	UPROPERTY(EditAnywhere, Category = Landscape, AdvancedDisplay, meta = (LandscapeOverridable))
+	bool bStripGrassWhenCookedServer = false;
+
 #if WITH_EDITOR
 	LANDSCAPE_API static ULandscapeLayerInfoObject* VisibilityLayer;
 
diff --git a/Engine/Source/Runtime/Landscape/Private/Landscape.cpp b/Engine/Source/Runtime/Landscape/Private/Landscape.cpp
--- a/Engine/Source/Runtime/Landscape/Private/Landscape.cpp
+++ b/Engine/Source/Runtime/Landscape/Private/Landscape.cpp
@@ -3234,6 +3234,57 @@ void ALandscapeProxy::PreSave(FObjectPreSaveContext ObjectSaveContext)
 			}
 		}
 	}
+
+	// Strip data according to flags when cooked (and not cooking for editor)
+	if (ObjectSaveContext.IsCooking() && !ObjectSaveContext.GetTargetPlatform()->AllowsEditorObjects())
+	{
+		if ((bStripPhysicsWhenCookedClient && bStripPhysicsWhenCookedServer) ||
+			(bStripPhysicsWhenCookedClient && ObjectSaveContext.GetTargetPlatform()->IsClientOnly()) ||
+			(bStripPhysicsWhenCookedServer && ObjectSaveContext.GetTargetPlatform()->IsServerOnly()))
+		{
+			// Clear old CollisionComponent containers
+			CollisionComponents.Empty();
+
+			// Destroy any owned collision components
+			TInlineComponentArray<ULandscapeHeightfieldCollisionComponent*> CollisionComps;
+			GetComponents(CollisionComps);
+			for (ULandscapeHeightfieldCollisionComponent* Component : CollisionComps)
+			{
+				Component->DestroyComponent();
+			}
+
+			// Destroy any attached but un-owned collision components (not sure what would be here that is not already caught above ^^)
+			TArray<USceneComponent*> AttachedCollisionComponents = RootComponent->GetAttachChildren().FilterByPredicate(
+				[](USceneComponent* Component)
+				{
+					return Cast<ULandscapeHeightfieldCollisionComponent>(Component);
+				});
+
+			for (USceneComponent* Component : AttachedCollisionComponents)
+			{
+				Component->DestroyComponent();
+			}
+
+			// ULandscapeComponents have refs... null them
+			for (ULandscapeComponent* Component : LandscapeComponents)
+			{
+				Component->DestroyCollisionData();
+				Component->SetCollisionComponent(nullptr);
+			}
+		}
+		
+		if ((bStripGrassWhenCookedClient && bStripGrassWhenCookedServer) ||
+			(bStripGrassWhenCookedClient && ObjectSaveContext.GetTargetPlatform()->IsClientOnly()) ||
+			(bStripGrassWhenCookedServer && ObjectSaveContext.GetTargetPlatform()->IsServerOnly()))
+		{
+			for (ULandscapeComponent* Component : LandscapeComponents)
+			{
+				TUniquePtr<FLandscapeComponentGrassData> NewGrassData = MakeUnique<FLandscapeComponentGrassData>();
+				Component->GrassData = MakeShareable(NewGrassData.Release());
+				Component->GrassData->NumElements = 0;
+			}
+		}
+	}
 #endif // WITH_EDITOR
 }
 
