Fix an issue where landscape splines were selecting the wrong control point mesh connection locations to auto connect to (selecting farther points instead of the nearest one).
When the code was updated to work on doubles, the logic was accidentally changed, leading to this regression.  This CL reverts to the original behavior.

diff --git a/Engine/Source/Runtime/Landscape/Private/LandscapeSplines.cpp b/Engine/Source/Runtime/Landscape/Private/LandscapeSplines.cpp
--- a/Engine/Source/Runtime/Landscape/Private/LandscapeSplines.cpp
+++ b/Engine/Source/Runtime/Landscape/Private/LandscapeSplines.cpp
@@ -1909,8 +1909,16 @@ FName ULandscapeSplineControlPoint::GetBestConnectionTo(FVector Destination) con
 			FVector SocketLocation = SocketTransform.GetTranslation();
 			FRotator SocketRotation = SocketTransform.GetRotation().Rotator();
 
-			const double Score = (Destination - Location).Size() - (Destination - SocketLocation).Size() // Score closer sockets higher
-				* FMath::Abs(FVector::DotProduct((Destination - SocketLocation), SocketRotation.Vector())); // score closer rotation higher
+			const double DistanceToCP = (Destination - Location).Size();
+
+			// score closer locations higher
+			const FVector SocketDelta = Destination - SocketLocation;
+			const double DistanceToSocket = SocketDelta.Size();
+
+			// score closer rotations higher
+			const double DirectionWeight = FMath::Abs(FVector::DotProduct(SocketDelta, SocketRotation.Vector()));
+
+			const double Score = (DistanceToCP - DistanceToSocket) * DirectionWeight;
 
 			if (Score > BestScore)
 			{
