Landscape : Enable Marching squares nanite visibility

diff --git a/Engine/Source/Runtime/Landscape/Private/LandscapeEdit.cpp b/Engine/Source/Runtime/Landscape/Private/LandscapeEdit.cpp
--- a/Engine/Source/Runtime/Landscape/Private/LandscapeEdit.cpp
+++ b/Engine/Source/Runtime/Landscape/Private/LandscapeEdit.cpp
@@ -100,7 +100,7 @@ static TAutoConsoleVariable<int32> CVarLandscapeApplyPhysicalMaterialChangesImme
 
 static TAutoConsoleVariable<bool> CVarLandscapeMarchingSquaresVisibility(
 	TEXT("landscape.Nanite.MarchingSquaresVisibility"),
-	false,
+	true,
 	TEXT("Use Marching Squares to make the visibility edges more accurate"));
 
 #if WITH_EDITOR
