Landscape : Fix crash on large landscape creation

diff --git a/Engine/Source/Editor/LandscapeEditor/Private/LandscapeEditorDetailCustomization_NewLandscape.cpp b/Engine/Source/Editor/LandscapeEditor/Private/LandscapeEditorDetailCustomization_NewLandscape.cpp
--- a/Engine/Source/Editor/LandscapeEditor/Private/LandscapeEditorDetailCustomization_NewLandscape.cpp
+++ b/Engine/Source/Editor/LandscapeEditor/Private/LandscapeEditorDetailCustomization_NewLandscape.cpp
@@ -1057,11 +1057,7 @@ FReply FLandscapeEditorDetailCustomization_NewLandscape::OnCreateButtonClicked()
 
 	TArray<uint16> OutHeightData;
 	
-	if (bIsNewLandscape || bNeedsLandscapeRegions)
-	{
-		UISettings->InitializeDefaultHeightData(OutHeightData);
-	}
-	else
+	if (!(bIsNewLandscape || bNeedsLandscapeRegions))
 	{
 		UISettings->ExpandImportData(OutHeightData, MaterialImportLayers);
 	}
diff --git a/Engine/Source/Runtime/Landscape/Private/LandscapeEdit.cpp b/Engine/Source/Runtime/Landscape/Private/LandscapeEdit.cpp
--- a/Engine/Source/Runtime/Landscape/Private/LandscapeEdit.cpp
+++ b/Engine/Source/Runtime/Landscape/Private/LandscapeEdit.cpp
@@ -2784,7 +2784,7 @@ ULandscapeLayerInfoObject* ALandscapeProxy::CreateLayerInfo(const TCHAR* InLayer
 	return LayerInfo;
 }
 
-#define HEIGHTDATA(X,Y) (HeightData[ FMath::Clamp<int32>(Y,0,VertsY) * VertsX + FMath::Clamp<int32>(X,0,VertsX) ])
+#define HEIGHTDATA(X,Y) (HeightData.Num() == 0 ? LandscapeDataAccess::GetTexHeight(0.0f) : HeightData[ FMath::Clamp<int32>(Y,0,VertsY) * VertsX + FMath::Clamp<int32>(X,0,VertsX) ])
 ENGINE_API extern bool GDisableAutomaticTextureMaterialUpdateDependencies;
 
 LANDSCAPE_API void ALandscapeProxy::Import(const FGuid& InGuid, int32 InMinX, int32 InMinY, int32 InMaxX, int32 InMaxY, int32 InNumSubsections, int32 InSubsectionSizeQuads, const TMap<FGuid, TArray<uint16>>& InImportHeightData, 
@@ -3489,7 +3489,7 @@ LANDSCAPE_API void ALandscapeProxy::Import(const FGuid& InGuid, int32 InMinX, in
 
 			const TArray<uint16>* ImportHeightData = InImportHeightData.Find(ImportSettings.SourceLayerGuid);
 
-			if (ImportHeightData != nullptr)
+			if (ImportHeightData != nullptr && ImportHeightData->Num() != 0)
 			{
 				LandscapeEdit.SetHeightData(InMinX, InMinY, InMaxX, InMaxY, (uint16*)ImportHeightData->GetData(), 0, false, nullptr);
 			}
