Fix crash on MetalRHI with incorrect flags in LandscapeEditResources

diff --git a/Engine/Source/Runtime/Landscape/Private/LandscapeEditResources.cpp b/Engine/Source/Runtime/Landscape/Private/LandscapeEditResources.cpp
--- a/Engine/Source/Runtime/Landscape/Private/LandscapeEditResources.cpp
+++ b/Engine/Source/Runtime/Landscape/Private/LandscapeEditResources.cpp
@@ -96,6 +96,7 @@ FLandscapeTexture2DArrayResource::FLandscapeTexture2DArrayResource(uint32 InSize
 	, Format(InFormat)
 	, NumMips(InNumMips)
 	, bCreateUAVs(bInNeedUAVs)
+    , bCreateSRV(bInNeedSRV)
 {
 
 }
@@ -103,26 +104,24 @@ FLandscapeTexture2DArrayResource::FLandscapeTexture2DArrayResource(uint32 InSize
 void FLandscapeTexture2DArrayResource::InitRHI(FRHICommandListBase& RHICmdList)
 {
 	FTextureResource::InitRHI(RHICmdList);
-
-	FRHIResourceCreateInfo CreateInfo(TEXT("FLandscapeTexture2DArrayResource"));
-	ETextureCreateFlags Flags = TexCreate_NoTiling | TexCreate_OfflineProcessed;
-
-	if (bCreateUAVs)
-	{
-		Flags |= TexCreate_UAV;
-	}
-
-	if (bCreateSRV)
-	{
-		Flags |= TexCreate_ShaderResource;
-	}
-
-	const FRHITextureCreateDesc Desc =
+    
+	FRHITextureCreateDesc Desc =
 		FRHITextureCreateDesc::Create2DArray(TEXT("FLandscapeTexture2DArrayResource"), SizeX, SizeY, SizeZ, Format)
-		.SetNumMips(static_cast<uint8>(NumMips))
-		.SetFlags(Flags);
-
-	TextureRHI = RHICreateTexture(Desc);
+        .SetNumMips(static_cast<uint8>(NumMips));
+
+    Desc.AddFlags(ETextureCreateFlags::NoTiling | ETextureCreateFlags::OfflineProcessed);
+    
+    if (bCreateUAVs)
+    {
+        Desc.AddFlags(ETextureCreateFlags::UAV);
+    }
+
+    if (bCreateSRV)
+    {
+        Desc.AddFlags(ETextureCreateFlags::ShaderResource);
+    }
+    
+    TextureRHI = RHICreateTexture(Desc);
 
 	if (bCreateUAVs)
 	{
