Fix for crash in ULandscapeMeshProxyComponent::InitializeForLandscape() while building HLOD

Switched to read from the local LODGroupKey value and not require the Landscape Actor to be loaded.

diff --git a/Engine/Source/Runtime/Landscape/Private/Landscape.cpp b/Engine/Source/Runtime/Landscape/Private/Landscape.cpp
--- a/Engine/Source/Runtime/Landscape/Private/Landscape.cpp
+++ b/Engine/Source/Runtime/Landscape/Private/Landscape.cpp
@@ -5662,7 +5662,7 @@ ULandscapeMeshProxyComponent::ULandscapeMeshProxyComponent(const FObjectInitiali
 void ULandscapeMeshProxyComponent::InitializeForLandscape(ALandscapeProxy* Landscape, int8 InProxyLOD)
 {
 	LandscapeGuid = Landscape->GetLandscapeGuid();
-	LODGroupKey = Landscape->GetLandscapeActor()->GetLODGroupKey();
+	LODGroupKey = Landscape->LODGroupKey;
 
 	for (ULandscapeComponent* Component : Landscape->LandscapeComponents)
 	{
