Suppress deprecation warning for ~FTriMeshGeometryRef() due to MSVC 14.37 raising a warning for the TUniquePtr destructor

#rnx

diff --git a/Engine/Source/Runtime/Landscape/Classes/LandscapeMeshCollisionComponent.h b/Engine/Source/Runtime/Landscape/Classes/LandscapeMeshCollisionComponent.h
--- a/Engine/Source/Runtime/Landscape/Classes/LandscapeMeshCollisionComponent.h
+++ b/Engine/Source/Runtime/Landscape/Classes/LandscapeMeshCollisionComponent.h
@@ -46,12 +46,14 @@ public:
 		TArray<Chaos::FMaterialHandle> UsedChaosMaterials;
 		Chaos::FTriangleMeshImplicitObjectPtr TrimeshGeometry;
 
+		// When deleted, remove PRAGMA_DISABLE_DEPRECATION_WARNINGS / PRAGMA_ENABLE_DEPRECATION_WARNINGS from ~FTriMeshGeometryRef()
 		UE_DEPRECATED(5.4, "Please use TrimeshGeometry instead")
 		TUniquePtr<Chaos::FTriangleMeshImplicitObject> Trimesh;
 		
 #if WITH_EDITORONLY_DATA
 		Chaos::FTriangleMeshImplicitObjectPtr EditorTrimeshGeometry;
 		
+		// When deleted, remove PRAGMA_DISABLE_DEPRECATION_WARNINGS / PRAGMA_ENABLE_DEPRECATION_WARNINGS from ~FTriMeshGeometryRef()
 		UE_DEPRECATED(5.4, "Please use EditorTrimeshGeometry instead")
 		TUniquePtr<Chaos::FTriangleMeshImplicitObject> EditorTrimesh;
 #endif // WITH_EDITORONLY_DATA
diff --git a/Engine/Source/Runtime/Landscape/Private/LandscapeCollision.cpp b/Engine/Source/Runtime/Landscape/Private/LandscapeCollision.cpp
--- a/Engine/Source/Runtime/Landscape/Private/LandscapeCollision.cpp
+++ b/Engine/Source/Runtime/Landscape/Private/LandscapeCollision.cpp
@@ -169,11 +169,13 @@ ULandscapeMeshCollisionComponent::FTriMeshGeometryRef::FTriMeshGeometryRef(FGuid
 	: Guid(InGuid)
 {}
 
+PRAGMA_DISABLE_DEPRECATION_WARNINGS
 ULandscapeMeshCollisionComponent::FTriMeshGeometryRef::~FTriMeshGeometryRef()
 {
 	// Remove ourselves from the shared map.
 	GSharedMeshRefs.Remove(Guid);
 }
+PRAGMA_ENABLE_DEPRECATION_WARNINGS
 
 void ULandscapeMeshCollisionComponent::FTriMeshGeometryRef::GetResourceSizeEx(FResourceSizeEx& CumulativeResourceSize)
 {
