Landscape : Fix Nanite and HLOD export when LOD is not zero

diff --git a/Engine/Source/Runtime/Landscape/Private/LandscapeEdit.cpp b/Engine/Source/Runtime/Landscape/Private/LandscapeEdit.cpp
--- a/Engine/Source/Runtime/Landscape/Private/LandscapeEdit.cpp
+++ b/Engine/Source/Runtime/Landscape/Private/LandscapeEdit.cpp
@@ -3875,10 +3875,10 @@ bool ALandscapeProxy::ExportToRawMeshDataCopyNew(const FRawMeshExportParams& InE
 
 	
 	checkf(!InExportParams.ComponentsUVConfiguration.IsSet() || InExportParams.ComponentsUVConfiguration->Num() == ComponentsToExport.Num(), TEXT("If ComponentsUVConfiguration is passed (per-component UV configuration), it must have the same number of entries as the number of components to export."))
-		checkf(!InExportParams.ComponentsMaterialSlotName.IsSet() || InExportParams.ComponentsMaterialSlotName->Num() == ComponentsToExport.Num(), TEXT("If ComponentsMaterialSlotName is passed (per-component material slot), it must have the same number of entries as the number of components to export."))
+	checkf(!InExportParams.ComponentsMaterialSlotName.IsSet() || InExportParams.ComponentsMaterialSlotName->Num() == ComponentsToExport.Num(), TEXT("If ComponentsMaterialSlotName is passed (per-component material slot), it must have the same number of entries as the number of components to export."))
 
-		// Get the tight bounds around the proxy's component (in quads, relative to the proxy's origin) :
-		const FIntRect LandscapeProxyBoundsRect = GetBoundingRect();
+	// Get the tight bounds around the proxy's component (in quads, relative to the proxy's origin) :
+	const FIntRect LandscapeProxyBoundsRect = GetBoundingRect();
 	const FVector2f LandscapeProxyBoundsRectUVScale = FVector2f(1.0f, 1.0f) / FVector2f(LandscapeProxyBoundsRect.Size());
 
 	FStaticMeshAttributes Attributes(OutRawMesh);
@@ -4015,16 +4015,18 @@ bool ALandscapeProxy::ExportToRawMeshDataCopyNew(const FRawMeshExportParams& InE
 			const bool bIsInside = ClampedVertexX == VertexX && ClampedVertexY == VertexY;
 
 			FVector Position = CDI.GetLocalVertex(ClampedVertexX, ClampedVertexY, HeightAndNormals);
-
+			
 			if (!bIsInside)
 			{
-				const float X = VertexX - ClampedVertexX;
-				const float Y = VertexY - ClampedVertexY;
-
-				Position -= FVector(0.0f, 0.0f, SkirtDepth);
+				int32 SignX = FMath::Sign(VertexX - ClampedVertexX);
+				int32 SignY = FMath::Sign(VertexY - ClampedVertexY);
 
+				const float Diff = CDI.GetLocalHeight(ClampedVertexX - SignX, ClampedVertexY - SignY, HeightAndNormals) - Position.Z;
+				
 				Position.X = CDI.GetScaleFactor() * VertexX;
 				Position.Y = CDI.GetScaleFactor() * VertexY;
+
+				Position -= FVector(0.0f, 0.0f, Diff + SkirtDepth);
 			}
 			return Position;
 		};
@@ -4044,8 +4046,10 @@ bool ALandscapeProxy::ExportToRawMeshDataCopyNew(const FRawMeshExportParams& InE
 			return VisDataMap[CDI.TexelXYToIndex(TexelX, TexelY)] / 255.0f;
 		};
 
-		auto GetBasis = [&CDI, ComponentSizeQuadsLOD, &HeightAndNormals](float X, float Y, FVector& OutLocalTangentX, FVector& OutLocalTangentY, FVector& OutLocalTangentZ)
+		auto GetBasis = [&CDI, ComponentSizeQuadsLOD, &HeightAndNormals, InvScaleFactor = 1.0f / CDI.GetScaleFactor()](float X, float Y, FVector& OutLocalTangentX, FVector& OutLocalTangentY, FVector& OutLocalTangentZ)
 		{
+			X = X * InvScaleFactor;
+			Y = Y * InvScaleFactor;
 
 			int32 VertexX = X;
 			int32 VertexY = Y;
@@ -4232,24 +4236,26 @@ bool ALandscapeProxy::ExportToRawMeshDataCopyNew(const FRawMeshExportParams& InE
 						// Compute all UV values that we need :
 						for (int32 UVChannel = 0; UVChannel < NumUVChannels; ++UVChannel)
 						{
+							FVector2f UV = FVector2f(NewLocalPositions[i].X, NewLocalPositions[i].Y) / CDI.GetScaleFactor();
+
 							FRawMeshExportParams::EUVMappingType UVMappingType = ComponentUVConfiguration.ExportUVMappingTypes.IsValidIndex(UVChannel) ? ComponentUVConfiguration.ExportUVMappingTypes[UVChannel] : FRawMeshExportParams::EUVMappingType::None;
 							switch (UVMappingType)
 							{
 							case FRawMeshExportParams::EUVMappingType::RelativeToProxyBoundsUV:
 							{
-								FVector2f UV = (ComponentOffsetRelativeToProxyBoundsQuadsLOD + FVector2f(static_cast<float>(VertexX), static_cast<float>(VertexY))) * ComponentUVScaleRelativeToProxyBoundsLOD;
+								UV = (ComponentOffsetRelativeToProxyBoundsQuadsLOD + UV) * ComponentUVScaleRelativeToProxyBoundsLOD;
 								UVs.Add(UV);
 								break;
 							}
 							case FRawMeshExportParams::EUVMappingType::HeightmapUV:
 							{
-								FVector2f UV = FVector2f(static_cast<float>(VertexX), static_cast<float>(VertexY)) * ComponentHeightmapUVScaleLOD + ComponentHeightmapUVPixelOffset + ComponentHeightmapUVBias;
+								UV = UV * ComponentHeightmapUVScaleLOD + ComponentHeightmapUVPixelOffset + ComponentHeightmapUVBias;
 								UVs.Add(UV);
 								break;
 							}
 							case FRawMeshExportParams::EUVMappingType::WeightmapUV:
 							{
-								FVector2f UV = FVector2f(static_cast<float>(VertexX), static_cast<float>(VertexY)) * ComponentWeightmapUVScaleLOD + ComponentWeightmapUVPixelOffset;
+								UV = UV * ComponentWeightmapUVScaleLOD + ComponentWeightmapUVPixelOffset;
 								UVs.Add(UV);
 								break;
 							}
@@ -4257,8 +4263,7 @@ bool ALandscapeProxy::ExportToRawMeshDataCopyNew(const FRawMeshExportParams& InE
 							case FRawMeshExportParams::EUVMappingType::TerrainCoordMapping_XZ:
 							case FRawMeshExportParams::EUVMappingType::TerrainCoordMapping_YZ:
 							{
-								FVector2f QuadCoords = (ComponentOffsetRelativeToProxyBoundsQuadsLOD + FVector2f(static_cast<float>(VertexX), static_cast<float>(VertexY)) / LODScale);
-								FVector2f UV;
+								FVector2f QuadCoords = (ComponentOffsetRelativeToProxyBoundsQuadsLOD + UV / LODScale);
 								if (UVMappingType == FRawMeshExportParams::EUVMappingType::TerrainCoordMapping_XY)
 								{
 									UV = QuadCoords;
@@ -4361,10 +4366,10 @@ bool ALandscapeProxy::ExportToRawMeshDataCopyOld(const FRawMeshExportParams& InE
 	}
 
 	checkf(!InExportParams.ComponentsUVConfiguration.IsSet() || InExportParams.ComponentsUVConfiguration->Num() == ComponentsToExport.Num(), TEXT("If ComponentsUVConfiguration is passed (per-component UV configuration), it must have the same number of entries as the number of components to export."))
-		checkf(!InExportParams.ComponentsMaterialSlotName.IsSet() || InExportParams.ComponentsMaterialSlotName->Num() == ComponentsToExport.Num(), TEXT("If ComponentsMaterialSlotName is passed (per-component material slot), it must have the same number of entries as the number of components to export."))
+	checkf(!InExportParams.ComponentsMaterialSlotName.IsSet() || InExportParams.ComponentsMaterialSlotName->Num() == ComponentsToExport.Num(), TEXT("If ComponentsMaterialSlotName is passed (per-component material slot), it must have the same number of entries as the number of components to export."))
 
-		// Get the tight bounds around the proxy's component (in quads, relative to the proxy's origin) :
-		const FIntRect LandscapeProxyBoundsRect = GetBoundingRect();
+	// Get the tight bounds around the proxy's component (in quads, relative to the proxy's origin) :
+	const FIntRect LandscapeProxyBoundsRect = GetBoundingRect();
 	const FVector2f LandscapeProxyBoundsRectUVScale = FVector2f(1.0f, 1.0f) / FVector2f(LandscapeProxyBoundsRect.Size());
 
 	FStaticMeshAttributes Attributes(OutRawMesh);
diff --git a/Engine/Source/Runtime/Landscape/Public/LandscapeDataAccess.h b/Engine/Source/Runtime/Landscape/Public/LandscapeDataAccess.h
--- a/Engine/Source/Runtime/Landscape/Public/LandscapeDataAccess.h
+++ b/Engine/Source/Runtime/Landscape/Public/LandscapeDataAccess.h
@@ -229,6 +229,11 @@ struct FLandscapeComponentDataInterfaceBase
 		return FVector(LocalX * ScaleFactor , LocalY * ScaleFactor, LandscapeDataAccess::GetLocalHeight(GetHeight(LocalX, LocalY, HeightAndNormals)));
 	}
 
+	float GetLocalHeight(int32 LocalX, int32 LocalY, const TArray<FColor>& HeightAndNormals) const
+	{
+		return LandscapeDataAccess::GetLocalHeight(GetHeight(LocalX, LocalY, HeightAndNormals));
+	}
+
 	const FColor* GetHeightData(int32 LocalX, int32 LocalY, const TArray<FColor>& HeightAndNormals) const
 	{
 #if LANDSCAPE_VALIDATE_DATA_ACCESS
