Moves GrassTypes from LandscapeProxy to LandscapeComponent.
Allows the use of different GrassTypes between the proxy and components by using LandscapeMaterialOverride.

diff --git a/Engine/Source/Runtime/Core/Public/UObject/FortniteMainBranchObjectVersion.h b/Engine/Source/Runtime/Core/Public/UObject/FortniteMainBranchObjectVersion.h
--- a/Engine/Source/Runtime/Core/Public/UObject/FortniteMainBranchObjectVersion.h
+++ b/Engine/Source/Runtime/Core/Public/UObject/FortniteMainBranchObjectVersion.h
@@ -360,6 +360,9 @@ struct FFortniteMainBranchObjectVersion
 		// EditableWhenInherited: Skip custom serialization on non Archetypes
 		ISMComponentEditableWhenInheritedSkipSerialization,
 
+		// GrassTypes are now per-component, rather than per-landscape proxy :
+		LandscapeSupportPerComponentGrassTypes,
+
 		// -----<new versions can be added above this line>-------------------------------------------------
 		VersionPlusOne,
 		LatestVersion = VersionPlusOne - 1
diff --git a/Engine/Source/Runtime/Landscape/Classes/LandscapeComponent.h b/Engine/Source/Runtime/Landscape/Classes/LandscapeComponent.h
--- a/Engine/Source/Runtime/Landscape/Classes/LandscapeComponent.h
+++ b/Engine/Source/Runtime/Landscape/Classes/LandscapeComponent.h
@@ -192,10 +192,13 @@ struct FLandscapeComponentGrassData
 
 	// Guid per material instance in the hierarchy between the assigned landscape material (instance) and the root UMaterial
 	// used to detect changes to material instance parameters or the root material that could affect the grass maps
-	TArray<FGuid, TInlineAllocator<2>> MaterialStateIds;
+	TArray<FGuid, TInlineAllocator<2>> MaterialStateIds_DEPRECATED;
 	// cached component rotation when material world-position-offset is used,
 	// as this will affect the direction of world-position-offset deformation (included in the HeightData below)
-	FQuat RotationForWPO;
+	FQuat RotationForWPO_DEPRECATED;
+
+	// Variable used to detect when grass data needs to be regenerated:
+	uint32 GenerationHash = 0;
 #endif
 
 #if WITH_EDITORONLY_DATA
@@ -511,6 +514,15 @@ private:
 	UPROPERTY(EditAnywhere, Category = LandscapeComponent)
 	TArray<FLandscapePerLODMaterialOverride> PerLODOverrideMaterials;
 
+	/** Cached list of grass types supported by the component's material */
+	UPROPERTY()
+	TArray<TObjectPtr<ULandscapeGrassType>> GrassTypes;
+
+	/** Cached max discard distance of all grass types supported by the component's material*/
+	UPROPERTY()
+	float GrassTypesMaxDiscardDistance = 0.0f;
+
+
 public:
 
 	/** Uniquely identifies this component's built map data. */
@@ -786,8 +798,20 @@ public:
 
 	/** Gets the landscape info object for this landscape */
 	LANDSCAPE_API ULandscapeInfo* GetLandscapeInfo() const;
+	
+	const TArray<ULandscapeGrassType*>& GetGrassTypes() const { return GrassTypes; }
+
+	/** Temporarily sets the grass type for this component. Any call to UpdateGrassTypes will override what has been set using this method */
+	void SetGrassTypes(const TArray<ULandscapeGrassType*>& InGrassTypes) { GrassTypes = InGrassTypes; }
+	
+	bool MaterialHasGrass() const { return !GetGrassTypes().IsEmpty(); }
+	
+	float GetGrassTypesMaxDiscardDistance() const { return GrassTypesMaxDiscardDistance; }
+	void SetGrassTypesMaxDiscardDistance(const float InGrassTypesMaxDiscardDistance) { GrassTypesMaxDiscardDistance = InGrassTypesMaxDiscardDistance; }
 
 #if WITH_EDITOR
+	/** Recomputes the list of grass types and other related info (in case there was a material change, for example) : */
+	void UpdateGrassTypes();
 
 	/** Deletes a material layer from the current edit layer on this component, removing all its data, adjusting other layer's weightmaps if necessary, etc. */
 	LANDSCAPE_API void DeleteLayer(ULandscapeLayerInfoObject* LayerInfo, FLandscapeEditDataInterface& LandscapeEdit);
@@ -805,9 +829,6 @@ public:
 	LANDSCAPE_API void ReplaceLayer(ULandscapeLayerInfoObject* FromLayerInfo, ULandscapeLayerInfoObject* ToLayerInfo, FLandscapeEditDataInterface& LandscapeEdit);
 	void ReplaceLayerInternal(ULandscapeLayerInfoObject* FromLayerInfo, ULandscapeLayerInfoObject* ToLayerInfo, FLandscapeEditDataInterface& LandscapeEdit, const FGuid& InEditLayerGUID);
 
-	// true if the component's landscape material supports grass
-	bool MaterialHasGrass() const;
-
 	/** Creates and destroys cooked grass data stored in the map */
 	void RenderGrassMap();
 	void RemoveGrassMap();
@@ -818,6 +839,9 @@ public:
 	/* Are the textures we need to render a grassmap currently streamed in? */
 	bool AreTexturesStreamedForGrassMapRender() const;
 
+	/** Computes a hash representing the state of the material and grasstypes used by this component. */
+	LANDSCAPE_API uint32 ComputeGrassMapGenerationHash() const;
+
 	/* Is the grassmap data outdated, eg by a material */
 	bool IsGrassMapOutdated() const;
 
diff --git a/Engine/Source/Runtime/Landscape/Classes/LandscapeProxy.h b/Engine/Source/Runtime/Landscape/Classes/LandscapeProxy.h
--- a/Engine/Source/Runtime/Landscape/Classes/LandscapeProxy.h
+++ b/Engine/Source/Runtime/Landscape/Classes/LandscapeProxy.h
@@ -514,17 +514,6 @@ public:
 	UPROPERTY(EditAnywhere, BlueprintSetter=EditorSetLandscapeMaterial, Category=Landscape, meta=(LandscapeOverridable))
 	TObjectPtr<UMaterialInterface> LandscapeMaterial;
 
-#if !WITH_EDITORONLY_DATA
-	/** Used to cache grass types from GetGrassTypes */
-	UMaterialInterface* LandscapeMaterialCached;
-
-	/** Cached grass types from GetGrassTypes */
-	TArray<ULandscapeGrassType*> LandscapeGrassTypes;
-
-	/** Cached grass max discard distance for all grass in GetGrassTypes */
-	float GrassMaxDiscardDistance;
-#endif
-
 	/** Material used to render landscape components with holes. If not set, LandscapeMaterial will be used (blend mode will be overridden to Masked if it is set to Opaque) */
 	UPROPERTY(EditAnywhere, Category=Landscape, AdvancedDisplay, meta = (LandscapeOverridable))
 	TObjectPtr<UMaterialInterface> LandscapeHoleMaterial;
@@ -636,6 +625,11 @@ public:
 	UPROPERTY()
 	bool bHasLandscapeGrass;
 
+	// Only used outside of the editor (e.g. in cooked builds)
+	// Cached grass max discard distance for all grass types in all landscape components with landscape grass configured
+	UPROPERTY()
+	float GrassTypesMaxDiscardDistance = 0.0f;
+
 	/**
 	 *	The resolution to cache lighting at, in texels/quad in one axis
 	 *  Total resolution would be changed by StaticLightingResolution*StaticLightingResolution
@@ -1028,10 +1022,6 @@ public:
 	 */
 	LANDSCAPE_API static void RemoveAllExclusionBoxes();
 
-
-	/* Get the list of grass types on this landscape */
-	static void GetGrassTypes(const UWorld* World, UMaterialInterface* LandscapeMat, TArray<ULandscapeGrassType*>& GrassTypesOut, float& OutMaxDiscardDistance);
-
 	/* Invalidate the precomputed grass and baked texture data for the specified components */
 	LANDSCAPE_API static void InvalidateGeneratedComponentData(const TSet<ULandscapeComponent*>& Components, bool bInvalidateLightingCache = false);
 	LANDSCAPE_API static void InvalidateGeneratedComponentData(const TArray<ULandscapeComponent*>& Components, bool bInvalidateLightingCache = false);
@@ -1046,7 +1036,11 @@ public:
 	void UpdateGrassData(bool bInShouldMarkDirty = false, struct FScopedSlowTask* InSlowTask = nullptr);
 
 	/** Render grass maps for the specified components */
-	void RenderGrassMaps(const TArray<ULandscapeComponent*>& LandscapeComponents, const TArray<ULandscapeGrassType*>& GrassTypes);
+	UE_DEPRECATED(5.4, "This version of RenderGrassMaps is deprecated to account for landscape components with their own grass types, use the other version instead.")
+	void RenderGrassMaps(const TArray<ULandscapeComponent*>& InLandscapeComponents, const TArray<ULandscapeGrassType*>& InGrassTypes) {}
+
+	/** Render grass maps for the specified components */
+	void RenderGrassMaps(TArrayView<ULandscapeComponent* const> InLandscapeComponents);
 
 	struct UE_DEPRECATED(5.3, "FGIBakedTextureState is officially deprecated now and nothing updates it anymore") FGIBakedTextureState
 	{
diff --git a/Engine/Source/Runtime/Landscape/Private/Landscape.cpp b/Engine/Source/Runtime/Landscape/Private/Landscape.cpp
--- a/Engine/Source/Runtime/Landscape/Private/Landscape.cpp
+++ b/Engine/Source/Runtime/Landscape/Private/Landscape.cpp
@@ -1339,6 +1339,11 @@ void ULandscapeComponent::PostLoad()
 		}
 	}
 
+	if (GIsEditor && GetLinkerCustomVersion(FFortniteMainBranchObjectVersion::GUID) < FFortniteMainBranchObjectVersion::LandscapeSupportPerComponentGrassTypes)
+	{
+		UpdateGrassTypes();
+	}
+
 #if !UE_BUILD_SHIPPING
 	if (MobileCombinationMaterialInstances.Num() == 0)
 	{
@@ -1412,11 +1417,6 @@ ALandscapeProxy::ALandscapeProxy(const FObjectInitializer& ObjectInitializer)
 #if WITH_EDITORONLY_DATA
 	, TargetDisplayOrder(ELandscapeLayerDisplayMode::Default)
 #endif // WITH_EDITORONLY_DATA
-#if !WITH_EDITORONLY_DATA
-	, LandscapeMaterialCached(nullptr)
-	, LandscapeGrassTypes()
-	, GrassMaxDiscardDistance(0.0f)
-#endif
 	, bHasLandscapeGrass(true)
 {
 	bReplicates = false;
diff --git a/Engine/Source/Runtime/Landscape/Private/LandscapeGrass.cpp b/Engine/Source/Runtime/Landscape/Private/LandscapeGrass.cpp
--- a/Engine/Source/Runtime/Landscape/Private/LandscapeGrass.cpp
+++ b/Engine/Source/Runtime/Landscape/Private/LandscapeGrass.cpp
@@ -69,6 +69,7 @@
 #include "MaterialCachedData.h"
 #include "SceneRenderTargetParameters.h"
 #include "TextureResource.h"
+#include "Serialization/ArchiveCrc32.h"
 #include "SceneRendererInterface.h"
 
 #define LOCTEXT_NAMESPACE "Landscape"
@@ -1127,23 +1128,9 @@ bool FLandscapeGrassWeightMeshProcessor::Process(
 // data also accessible by render thread
 class FLandscapeGrassWeightExporter_RenderThread
 {
-	FSceneInterface* SceneInterface = nullptr;
-
-	FLandscapeGrassWeightExporter_RenderThread(int32 InNumGrassMaps, bool InbNeedsHeightmap, TArray<int32> InHeightMips)
-		: RenderTargetResource(nullptr)
-		, NumPasses(0)
-		, HeightMips(MoveTemp(InHeightMips))
-		, FirstHeightMipsPassIndex(MAX_int32)
+	FLandscapeGrassWeightExporter_RenderThread(const TArray<int32>& InHeightMips)
+		: HeightMips(InHeightMips)
 	{
-		if (InbNeedsHeightmap || InNumGrassMaps > 0)
-		{
-			NumPasses += FMath::DivideAndRoundUp(2 /* heightmap */ + InNumGrassMaps, 4);
-		}
-		if (HeightMips.Num() > 0)
-		{
-			FirstHeightMipsPassIndex = NumPasses;
-			NumPasses += HeightMips.Num();
-		}
 	}
 
 	friend class FLandscapeGrassWeightExporter;
@@ -1154,25 +1141,40 @@ public:
 
 	struct FComponentInfo
 	{
-		TObjectPtr<ULandscapeComponent> Component;
-		FVector2D ViewOffset;
-		int32 PixelOffsetX;
-		FLandscapeComponentSceneProxy* SceneProxy;
-
-		FComponentInfo(ULandscapeComponent* InComponent, FVector2D& InViewOffset, int32 InPixelOffsetX)
+		TObjectPtr<ULandscapeComponent> Component = nullptr;
+		TArray<TObjectPtr<ULandscapeGrassType>> RequestedGrassTypes;
+		FVector2D ViewOffset = FVector2D::ZeroVector;
+		int32 PixelOffsetX = 0;
+		FLandscapeComponentSceneProxy* SceneProxy = nullptr;
+		int32 NumPasses = 0;
+		int32 FirstHeightMipsPassIndex = MAX_int32;
+
+		FComponentInfo(ULandscapeComponent* InComponent, bool bInNeedsGrassmap, bool bInNeedsHeightmap, const TArray<int32>& InHeightMips)
 			: Component(InComponent)
-			, ViewOffset(InViewOffset)
-			, PixelOffsetX(InPixelOffsetX)
 			, SceneProxy((FLandscapeComponentSceneProxy*)InComponent->SceneProxy)
-		{}
+		{
+			if (bInNeedsGrassmap)
+			{
+				RequestedGrassTypes = InComponent->GetGrassTypes();
+			}
+			int32 NumGrassMaps = RequestedGrassTypes.Num();
+			if (bInNeedsHeightmap || NumGrassMaps > 0)
+			{
+				NumPasses += FMath::DivideAndRoundUp(2 /* heightmap */ + NumGrassMaps, 4);
+			}
+			if (InHeightMips.Num() > 0)
+			{
+				FirstHeightMipsPassIndex = NumPasses;
+				NumPasses += InHeightMips.Num();
+			}
+		}
 	};
 
-	FTextureRenderTarget2DResource* RenderTargetResource;
+	FSceneInterface* SceneInterface = nullptr;
+	FTextureRenderTarget2DResource* RenderTargetResource = nullptr;
 	TArray<FComponentInfo, TInlineAllocator<1>> ComponentInfos;
 	FIntPoint TargetSize;
-	int32 NumPasses;
 	TArray<int32> HeightMips;
-	int32 FirstHeightMipsPassIndex;
 	float PassOffsetX;
 	FVector ViewOrigin;
 	FMatrix ViewRotationMatrix;
@@ -1233,7 +1235,7 @@ public:
 						const FMeshBatch& Mesh = ComponentInfo.SceneProxy->GetGrassMeshBatch();
 						Mesh.MaterialRenderProxy->UpdateUniformExpressionCacheIfNeeded(View->GetFeatureLevel());
 
-						PassMeshProcessor.AddMeshBatch(Mesh, DefaultBatchElementMask, NumPasses, ComponentInfo.ViewOffset, PassOffsetX, FirstHeightMipsPassIndex, HeightMips, ComponentInfo.SceneProxy);
+						PassMeshProcessor.AddMeshBatch(Mesh, DefaultBatchElementMask, ComponentInfo.NumPasses, ComponentInfo.ViewOffset, PassOffsetX, ComponentInfo.FirstHeightMipsPassIndex, HeightMips, ComponentInfo.SceneProxy);
 					}
 				}
 			});
@@ -1252,55 +1254,62 @@ class FLandscapeGrassWeightExporter : public FLandscapeGrassWeightExporter_Rende
 	TObjectPtr<UTextureRenderTarget2D> RenderTargetTexture;
 
 public:
-	FLandscapeGrassWeightExporter(ALandscapeProxy* InLandscapeProxy, const TArray<ULandscapeComponent*>& InLandscapeComponents, TArray<ULandscapeGrassType*> InGrassTypes, bool InbNeedsHeightmap = true, TArray<int32> InHeightMips = {})
-		: FLandscapeGrassWeightExporter_RenderThread(
-			InGrassTypes.Num(),
-			InbNeedsHeightmap,
-			MoveTemp(InHeightMips))
+	FLandscapeGrassWeightExporter(ALandscapeProxy* InLandscapeProxy, TArrayView<ULandscapeComponent* const> InLandscapeComponents, bool bInNeedsGrassmap = true, bool bInNeedsHeightmap = true, const TArray<int32>& InHeightMips = {})
+		: FLandscapeGrassWeightExporter_RenderThread(InHeightMips)
 		, LandscapeProxy(InLandscapeProxy)
 		, ComponentSizeVerts(InLandscapeProxy->ComponentSizeQuads + 1)
 		, SubsectionSizeQuads(InLandscapeProxy->SubsectionSizeQuads)
 		, NumSubsections(InLandscapeProxy->NumSubsections)
-		, GrassTypes(MoveTemp(InGrassTypes))
 		, RenderTargetTexture(nullptr)
 	{
 		check(InLandscapeComponents.Num() > 0);
 		SceneInterface = InLandscapeComponents[0]->GetScene();
 
 		// todo: use a 2d target?
-		TargetSize = FIntPoint(ComponentSizeVerts * NumPasses * InLandscapeComponents.Num(), ComponentSizeVerts);
-		FIntPoint TargetSizeMinusOne(TargetSize - FIntPoint(1, 1));
-		PassOffsetX = 2.0f * (float)ComponentSizeVerts / (float)TargetSize.X;
+		const int32 SingleTileWidth = ComponentSizeVerts;
+		TargetSize = FIntPoint(0, ComponentSizeVerts);
 
-		for (int32 Idx = 0; Idx < InLandscapeComponents.Num(); Idx++)
+		// First compute the total render target size and prepare ComponentInfos (each component has its own number of needed passes because some might have different materials, thus different grass types) :
+		ComponentInfos.Reserve(InLandscapeComponents.Num());
+		int32 CurrentPixelOffsetX = 0;
+		for (ULandscapeComponent* Component : InLandscapeComponents)
 		{
-			ULandscapeComponent* Component = InLandscapeComponents[Idx];
+			ensure(Component->SceneProxy);
+
+			FComponentInfo& ComponentInfo = ComponentInfos.Add_GetRef(FComponentInfo(Component, bInNeedsGrassmap, bInNeedsHeightmap, InHeightMips));
+			ComponentInfo.PixelOffsetX = CurrentPixelOffsetX;
 
-			FIntPoint ComponentOffset = (Component->GetSectionBase() - LandscapeProxy->LandscapeSectionOffset);
-			int32 PixelOffsetX = Idx * NumPasses * ComponentSizeVerts;
+			CurrentPixelOffsetX += ComponentInfo.NumPasses * SingleTileWidth;
+		}
+		TargetSize.X = CurrentPixelOffsetX;
+
+		FIntPoint TargetSizeMinusOne(TargetSize - FIntPoint(1, 1));
+		PassOffsetX = 2.0f * (float)SingleTileWidth / (float)TargetSize.X;
 
+		// Then compute FComponentInfo's ViewOffset with the knowledge of the total render target size : 
+		for (FComponentInfo& ComponentInfo : ComponentInfos)
+		{
+			FIntPoint ComponentOffset = (ComponentInfo.Component->GetSectionBase() - LandscapeProxy->LandscapeSectionOffset);
 			FVector2D ViewOffset(-ComponentOffset.X, ComponentOffset.Y);
-			ViewOffset.X += PixelOffsetX;
+			ViewOffset.X += ComponentInfo.PixelOffsetX;
 			ViewOffset /= (FVector2D(TargetSize) * 0.5f);
-
-			ensure(Component->SceneProxy);
-			ComponentInfos.Add(FComponentInfo(Component, ViewOffset, PixelOffsetX));
+			ComponentInfo.ViewOffset = ViewOffset;
 		}
 
 		// center of target area in world
-		FVector TargetCenter = LandscapeProxy->GetTransform().TransformPosition(FVector(TargetSizeMinusOne, 0.f)*0.5f);
+		FVector TargetCenter = LandscapeProxy->GetTransform().TransformPosition(FVector(TargetSizeMinusOne, 0.f) * 0.5f);
 
 		// extent of target in world space
-		FVector TargetExtent = FVector(TargetSize, 0.0f)*LandscapeProxy->GetActorScale()*0.5f;
+		FVector TargetExtent = FVector(TargetSize, 0.0f) * LandscapeProxy->GetActorScale() * 0.5f;
 
 		ViewOrigin = TargetCenter;
 		ViewRotationMatrix = FInverseRotationMatrix(LandscapeProxy->GetActorRotation());
-		ViewRotationMatrix *= FMatrix(FPlane(1, 0, 0, 0),
-		                              FPlane(0,-1, 0, 0),
-		                              FPlane(0, 0,-1, 0),
-		                              FPlane(0, 0, 0, 1));
+		ViewRotationMatrix *= FMatrix(FPlane(1.0f, 0.0f, 0.0f, 0.0f),
+							          FPlane(0.0f, -1.0f, 0.0f, 0.0f),
+							          FPlane(0.0f, 0.0f, -1.0f, 0.0f),
+							          FPlane(0.0f, 0.0f, 0.0f, 1.0f));
 
-		const FMatrix::FReal ZOffset = UE_OLD_WORLD_MAX;
+		const float ZOffset = WORLD_MAX;
 		ProjectionMatrix = FReversedZOrthoMatrix(
 			TargetExtent.X,
 			TargetExtent.Y,
@@ -1327,11 +1336,10 @@ public:
 			});
 	}
 
-	TMap<ULandscapeComponent*, TUniquePtr<FLandscapeComponentGrassData>, TInlineSetAllocator<1>>
-		FetchResults()
+	TMap<ULandscapeComponent*, TUniquePtr<FLandscapeComponentGrassData>, TInlineSetAllocator<1>> FetchResults()
 	{
 		TArray<FColor> Samples;
-		Samples.SetNumUninitialized(TargetSize.X*TargetSize.Y);
+		Samples.SetNumUninitialized(TargetSize.X * TargetSize.Y);
 
 		// Copy the contents of the remote texture to system memory
 		FReadSurfaceDataFlags ReadSurfaceDataFlags;
@@ -1340,7 +1348,7 @@ public:
 
 		TMap<ULandscapeComponent*, TUniquePtr<FLandscapeComponentGrassData>, TInlineSetAllocator<1>> Results;
 		Results.Reserve(ComponentInfos.Num());
-		
+
 		// Local data will be moved in contiguous array at the end of export (to minimize slack waste)
 		TArray<uint16> HeightData;
 		TMap<ULandscapeGrassType*, TArray<uint8>> WeightData;
@@ -1352,7 +1360,7 @@ public:
 
 			TUniquePtr<FLandscapeComponentGrassData> NewGrassData = MakeUnique<FLandscapeComponentGrassData>(Component);
 
-			if (FirstHeightMipsPassIndex > 0)
+			if (ComponentInfo.FirstHeightMipsPassIndex != MAX_int32)
 			{
 				HeightData.Empty(FMath::Square(ComponentSizeVerts));
 			}
@@ -1364,14 +1372,14 @@ public:
 
 			WeightData.Empty();
 			TArray<TArray<uint8>*> GrassWeightArrays;
-			GrassWeightArrays.Empty(GrassTypes.Num());
-			for (auto GrassType : GrassTypes)
+			GrassWeightArrays.Empty(ComponentInfo.RequestedGrassTypes.Num());
+			for (auto GrassType : ComponentInfo.RequestedGrassTypes)
 			{
 				WeightData.Add(GrassType);
 			}
 
 			// need a second loop because the WeightData map will reallocate its arrays as grass types are added
-			for (auto GrassType : GrassTypes)
+			for (auto GrassType : ComponentInfo.RequestedGrassTypes)
 			{
 				TArray<uint8>* DataArray = WeightData.Find(GrassType);
 				check(DataArray);
@@ -1387,14 +1395,14 @@ public:
 				FString TempPath = FPaths::ScreenShotDir();
 				TempPath += TEXT("/GrassDebug");
 				IFileManager::Get().MakeDirectory(*TempPath, true);
-				FFileHelper::CreateBitmap(*(TempPath / "Grass"), TargetSize.X, TargetSize.Y, Samples.GetData(), nullptr, &IFileManager::Get(), nullptr, GrassTypes.Num() >= 2);
+				FFileHelper::CreateBitmap(*(TempPath / "Grass"), TargetSize.X, TargetSize.Y, Samples.GetData(), nullptr, &IFileManager::Get(), nullptr, ComponentInfo.RequestedGrassTypes.Num() >= 2);
 			}
 #endif
 
-			for (int32 PassIdx = 0; PassIdx < NumPasses; PassIdx++)
+			for (int32 PassIdx = 0; PassIdx < ComponentInfo.NumPasses; PassIdx++)
 			{
-				FColor* SampleData = &Samples[ComponentInfo.PixelOffsetX + PassIdx*ComponentSizeVerts];
-				if (PassIdx < FirstHeightMipsPassIndex)
+				FColor* SampleData = &Samples[ComponentInfo.PixelOffsetX + PassIdx * ComponentSizeVerts];
+				if (PassIdx < ComponentInfo.FirstHeightMipsPassIndex)
 				{
 					if (PassIdx == 0)
 					{
@@ -1405,10 +1413,10 @@ public:
 								FColor& Sample = SampleData[x + y * TargetSize.X];
 								uint16 Height = (((uint16)Sample.R) << 8) + (uint16)(Sample.G);
 								HeightData.Add(Height);
-								if (GrassTypes.Num() > 0)
+								if (ComponentInfo.RequestedGrassTypes.Num() > 0)
 								{
 									GrassWeightArrays[0]->Add(Sample.B);
-									if (GrassTypes.Num() > 1)
+									if (ComponentInfo.RequestedGrassTypes.Num() > 1)
 									{
 										GrassWeightArrays[1]->Add(Sample.A);
 									}
@@ -1426,13 +1434,13 @@ public:
 
 								int32 TypeIdx = PassIdx * 4 - 2;
 								GrassWeightArrays[TypeIdx++]->Add(Sample.R);
-								if (TypeIdx < GrassTypes.Num())
+								if (TypeIdx < ComponentInfo.RequestedGrassTypes.Num())
 								{
 									GrassWeightArrays[TypeIdx++]->Add(Sample.G);
-									if (TypeIdx < GrassTypes.Num())
+									if (TypeIdx < ComponentInfo.RequestedGrassTypes.Num())
 									{
 										GrassWeightArrays[TypeIdx++]->Add(Sample.B);
-										if (TypeIdx < GrassTypes.Num())
+										if (TypeIdx < ComponentInfo.RequestedGrassTypes.Num())
 										{
 											GrassWeightArrays[TypeIdx++]->Add(Sample.A);
 										}
@@ -1444,7 +1452,7 @@ public:
 				}
 				else // PassIdx >= FirstHeightMipsPassIndex
 				{
-					const int32 Mip = HeightMips[PassIdx - FirstHeightMipsPassIndex];
+					const int32 Mip = HeightMips[PassIdx - ComponentInfo.FirstHeightMipsPassIndex];
 					int32 MipSizeVerts = NumSubsections * (SubsectionSizeQuads >> Mip);
 					TArray<uint16>& MipHeightData = NewGrassData->HeightMipData.Add(Mip);
 					for (int32 y = 0; y < MipSizeVerts; y++)
@@ -1518,79 +1526,73 @@ public:
 			{
 				Collector.AddReferencedObject(Info.Component);
 			}
-		}
 
-		for (auto GrassType : GrassTypes)
-		{
-			if (GrassType)
-			{
-				Collector.AddReferencedObject(GrassType);
-			}
+			Collector.AddReferencedObjects(Info.RequestedGrassTypes);
 		}
 	}
 };
 
 FLandscapeComponentGrassData::FLandscapeComponentGrassData(ULandscapeComponent* Component)
-	: RotationForWPO(Component->GetLandscapeMaterial()->GetMaterial()->IsPropertyConnected(MP_WorldPositionOffset) ? Component->GetComponentTransform().GetRotation() : FQuat(0, 0, 0, 0))
-	, bIsDirty(false)
+	: GenerationHash(Component->ComputeGrassMapGenerationHash())
 {
-	UMaterialInterface* Material = Component->GetLandscapeMaterial();
-	for (UMaterialInstanceConstant* MIC = Cast<UMaterialInstanceConstant>(Material); MIC; MIC = Cast<UMaterialInstanceConstant>(Material))
-	{
-		MaterialStateIds.Add(MIC->ParameterStateId);
-		Material = MIC->Parent;
-	}
-	MaterialStateIds.Add(CastChecked<UMaterial>(Material)->StateId);
 }
 
-bool ULandscapeComponent::MaterialHasGrass() const
+void ULandscapeComponent::UpdateGrassTypes()
 {
-	UMaterialInterface* Material = GetLandscapeMaterial();
-	TArray<const UMaterialExpressionLandscapeGrassOutput*> GrassExpressions;
-	Material->GetMaterial()->GetAllExpressionsOfType<UMaterialExpressionLandscapeGrassOutput>(GrassExpressions);
-	if (GrassExpressions.Num() > 0 &&
-		GrassExpressions[0]->GrassTypes.Num() > 0)
+	GrassTypesMaxDiscardDistance = 0.0f;
+	GrassTypes.Reset();
+	
+	if (UMaterialInterface* Material = GetLandscapeMaterial())
 	{
-		return GrassExpressions[0]->GrassTypes.ContainsByPredicate([](FGrassInput& GrassInput) { return (GrassInput.Input.IsConnected() && GrassInput.GrassType); });
-	}
+		GrassTypes = Material->GetMaterial()->GetCachedExpressionData().GrassTypes;
 
-	return false;
-}
-
-bool ULandscapeComponent::IsGrassMapOutdated() const
-{
-	// Don't mark a component as outdated if we won't be able to render it anyway : 
-	if (GrassData->HasValidData() && CanRenderGrassMap())
-	{
-		// check material / instances haven't changed
-		const auto& MaterialStateIds = GrassData->MaterialStateIds;
-		UMaterialInterface* Material = GetLandscapeMaterial();
-		int32 TestIndex = 0;
-		for (UMaterialInstanceConstant* MIC = Cast<UMaterialInstanceConstant>(Material); MIC; MIC = Cast<UMaterialInstanceConstant>(Material))
+		for (const ULandscapeGrassType* GrassType : GrassTypes)
 		{
-			if (!MaterialStateIds.IsValidIndex(TestIndex) || MaterialStateIds[TestIndex] != MIC->ParameterStateId)
+			if (GrassType != nullptr)
 			{
-				return true;
+				for (const FGrassVariety& GrassVariety : GrassType->GrassVarieties)
+				{
+					GrassTypesMaxDiscardDistance = FMath::Max((float)GrassVariety.EndCullDistance.GetValue(), GrassTypesMaxDiscardDistance);
+				}
 			}
-			Material = MIC->Parent;
-			++TestIndex;
 		}
+	}
+}
 
-		UMaterial* MaterialBase = Cast<UMaterial>(Material);
+uint32 ULandscapeComponent::ComputeGrassMapGenerationHash() const
+{
+	TRACE_CPUPROFILER_EVENT_SCOPE(ULandscapeComponent::ComputeGrassMapGenerationHash);
+	uint32 Hash = 0;
 
-		// last one should be a UMaterial
-		if (TestIndex != MaterialStateIds.Num() - 1 || (MaterialBase != nullptr && MaterialStateIds[TestIndex] != MaterialBase->StateId))
+	if (UMaterialInterface* Material = GetLandscapeMaterial())
+	{
+		// Take into account any material state change :
+		Hash = Material->ComputeAllStateCRC();
+		if (UMaterial* MaterialBase = Material->GetMaterial())
 		{
-			return true;
-		}
+			// If anything changes in the grass types, we should take that into account as well :
+			// TODO [jonathan.bard], todo(luc.eygasier): make this event-based... 
+			FArchiveCrc32 Ar(Hash);
+			
+			for (ULandscapeGrassType* GrassType : GrassTypes)
+			{
+				if (GrassType == nullptr)
+				{
+					continue;
+				}
 
-		FQuat RotationForWPO = GetLandscapeMaterial()->GetMaterial()->IsPropertyConnected(MP_WorldPositionOffset) ? GetComponentTransform().GetRotation() : FQuat(0, 0, 0, 0);
-		if (GrassData->RotationForWPO != RotationForWPO)
-		{
-			return true;
+				GrassType->Serialize(Ar);
+			}
+			Hash = Ar.GetCrc();
 		}
 	}
-	return false;
+
+	return Hash;
+}
+
+bool ULandscapeComponent::IsGrassMapOutdated() const
+{
+	return GrassData->HasValidData() ? (ComputeGrassMapGenerationHash() != GrassData->GenerationHash) : false;
 }
 
 bool ULandscapeComponent::CanRenderGrassMap() const
@@ -1655,19 +1657,6 @@ void ULandscapeComponent::RenderGrassMap()
 	UMaterialInterface* Material = GetLandscapeMaterial();
 	if (ensure(CanRenderGrassMap()))
 	{
-		TArray<ULandscapeGrassType*> GrassTypes;
-
-		TArray<const UMaterialExpressionLandscapeGrassOutput*> GrassExpressions;
-		Material->GetMaterial()->GetAllExpressionsOfType<UMaterialExpressionLandscapeGrassOutput>(GrassExpressions);
-		if (GrassExpressions.Num() > 0)
-		{
-			GrassTypes.Empty(GrassExpressions[0]->GrassTypes.Num());
-			for (auto& GrassTypeInput : GrassExpressions[0]->GrassTypes)
-			{
-				GrassTypes.Add(GrassTypeInput.GrassType);
-			}
-		}
-
 		const bool bBakeMaterialPositionOffsetIntoCollision = (GetLandscapeProxy() && GetLandscapeProxy()->bBakeMaterialPositionOffsetIntoCollision);
 
 		TArray<int32> HeightMips;
@@ -1685,10 +1674,7 @@ void ULandscapeComponent::RenderGrassMap()
 
 		if (GrassTypes.Num() > 0 || bBakeMaterialPositionOffsetIntoCollision)
 		{
-			TArray<ULandscapeComponent*> LandscapeComponents;
-			LandscapeComponents.Add(this);
-
-			FLandscapeGrassWeightExporter Exporter(GetLandscapeProxy(), MoveTemp(LandscapeComponents), MoveTemp(GrassTypes), true, MoveTemp(HeightMips));
+			FLandscapeGrassWeightExporter Exporter(GetLandscapeProxy(), { this }, /*bInNeedsGrassmap = */ true, /*bInNeedsHeightmap =*/ true, MoveTemp(HeightMips));
 			Exporter.ApplyResults();
 		}
 	}
@@ -1705,13 +1691,9 @@ TArray<uint16> ULandscapeComponent::RenderWPOHeightmap(int32 LOD)
 
 	if (ensure(SceneProxy))
 	{
-		TArray<ULandscapeGrassType*> GrassTypes;
-		TArray<ULandscapeComponent*> LandscapeComponents;
-		LandscapeComponents.Add(this);
-
 		if (LOD == 0)
 		{
-			FLandscapeGrassWeightExporter Exporter(GetLandscapeProxy(), MoveTemp(LandscapeComponents), MoveTemp(GrassTypes), true, {});
+			FLandscapeGrassWeightExporter Exporter(GetLandscapeProxy(), { this }, /*bInNeedsGrassmap = */ false, /*bInNeedsHeightmap =*/ true, {});
 			TMap<ULandscapeComponent*, TUniquePtr<FLandscapeComponentGrassData>, TInlineSetAllocator<1>> TempGrassData;
 			TempGrassData = Exporter.FetchResults();
 			Results = TArray<uint16>(TempGrassData[this]->GetHeightData());
@@ -1720,7 +1702,7 @@ TArray<uint16> ULandscapeComponent::RenderWPOHeightmap(int32 LOD)
 		{
 			TArray<int32> HeightMips;
 			HeightMips.Add(LOD);
-			FLandscapeGrassWeightExporter Exporter(GetLandscapeProxy(), MoveTemp(LandscapeComponents), MoveTemp(GrassTypes), false, MoveTemp(HeightMips));
+			FLandscapeGrassWeightExporter Exporter(GetLandscapeProxy(), { this }, /*bInNeedsGrassmap = */ false, /*bInNeedsHeightmap =*/ false, MoveTemp(HeightMips));
 			TMap<ULandscapeComponent*, TUniquePtr<FLandscapeComponentGrassData>, TInlineSetAllocator<1>> TempGrassData;
 			TempGrassData = Exporter.FetchResults();
 			Results = MoveTemp(TempGrassData[this]->HeightMipData[LOD]);
@@ -1742,7 +1724,7 @@ void ULandscapeComponent::RemoveGrassMap()
 	}
 }
 
-void ALandscapeProxy::RenderGrassMaps(const TArray<ULandscapeComponent*>& InLandscapeComponents, const TArray<ULandscapeGrassType*>& GrassTypes)
+void ALandscapeProxy::RenderGrassMaps(TArrayView<ULandscapeComponent* const> InLandscapeComponents)
 {
 	TArray<int32> HeightMips;
 	if (CollisionMipLevel > 0)
@@ -1754,7 +1736,7 @@ void ALandscapeProxy::RenderGrassMaps(const TArray<ULandscapeComponent*>& InLand
 		HeightMips.Add(SimpleCollisionMipLevel);
 	}
 
-	FLandscapeGrassWeightExporter Exporter(this, InLandscapeComponents, GrassTypes, true, MoveTemp(HeightMips));
+	FLandscapeGrassWeightExporter Exporter(this, InLandscapeComponents, /*bInNeedsGrassmap = */ true, /*bInNeedsHeightmap =*/ true, MoveTemp(HeightMips));
 	Exporter.ApplyResults();
 }
 
@@ -2190,24 +2172,32 @@ FArchive& operator<<(FArchive& Ar, FLandscapeComponentGrassData& Data)
 #if WITH_EDITORONLY_DATA
 	if (!Ar.IsFilterEditorOnly())
 	{
-		if (Ar.CustomVer(FLandscapeCustomVersion::GUID) >= FLandscapeCustomVersion::GrassMaterialInstanceFix)
-		{
-			Ar << Data.MaterialStateIds;
-		}
-		else
+		if (Ar.CustomVer(FFortniteMainBranchObjectVersion::GUID) < FFortniteMainBranchObjectVersion::LandscapeSupportPerComponentGrassTypes)
 		{
-			Data.MaterialStateIds.Empty(1);
-			if (Ar.UEVer() >= VER_UE4_SERIALIZE_LANDSCAPE_GRASS_DATA_MATERIAL_GUID)
+			if (Ar.CustomVer(FLandscapeCustomVersion::GUID) >= FLandscapeCustomVersion::GrassMaterialInstanceFix)
 			{
-				FGuid MaterialStateId;
-				Ar << MaterialStateId;
-				Data.MaterialStateIds.Add(MaterialStateId);
+				Ar << Data.MaterialStateIds_DEPRECATED;
+			}
+			else
+			{
+				Data.MaterialStateIds_DEPRECATED.Empty(1);
+				if (Ar.UEVer() >= VER_UE4_SERIALIZE_LANDSCAPE_GRASS_DATA_MATERIAL_GUID)
+				{
+					FGuid MaterialStateId;
+					Ar << MaterialStateId;
+					Data.MaterialStateIds_DEPRECATED.Add(MaterialStateId);
+				}
 			}
-		}
 
-		if (Ar.CustomVer(FLandscapeCustomVersion::GUID) >= FLandscapeCustomVersion::GrassMaterialWPO)
+			if (Ar.CustomVer(FLandscapeCustomVersion::GUID) >= FLandscapeCustomVersion::GrassMaterialWPO)
+			{
+				Ar << Data.RotationForWPO_DEPRECATED;
+			}
+		}
+		else
 		{
-			Ar << Data.RotationForWPO;
+			// MaterialStateIds and RotationForWPO have been replaced with a simple hash :
+			Ar << Data.GenerationHash;
 		}
 	}
 
@@ -3049,30 +3039,6 @@ void ALandscapeProxy::FlushGrassComponents(const TSet<ULandscapeComponent*>* Onl
 	}
 }
 
-void ALandscapeProxy::GetGrassTypes(const UWorld* World, UMaterialInterface* LandscapeMat, TArray<ULandscapeGrassType*>& GrassTypesOut, float& OutMaxDiscardDistance)
-{
-	OutMaxDiscardDistance = 0.0f;
-	if (LandscapeMat)
-	{
-		GrassTypesOut.Append(LandscapeMat->GetCachedExpressionData().GrassTypes);
-
-		for (const ULandscapeGrassType* GrassType : LandscapeMat->GetCachedExpressionData().GrassTypes)
-		{
-			if (GrassType != nullptr)
-			{
-				for (auto& GrassVariety : GrassType->GrassVarieties)
-				{
-					const float EndCullDistance = static_cast<float>(GrassVariety.GetEndCullDistance());
-					if (EndCullDistance > OutMaxDiscardDistance)
-					{
-						OutMaxDiscardDistance = EndCullDistance;
-					}
-				}
-			}
-		}
-	}
-}
-
 static uint32 GGrassExclusionChangeTag = 1;
 static uint32 GFrameNumberLastStaleCheck = 0;
 static TMap<FWeakObjectPtr, FBox, FDefaultSetAllocator, TWeakObjectPtrMapKeyFuncs<FWeakObjectPtr, FBox>> GGrassExclusionBoxes;
@@ -3153,10 +3119,21 @@ void ALandscapeProxy::BuildGrassMaps(FScopedSlowTask* InSlowTask)
 {
 	if (!HasAnyFlags(RF_ClassDefaultObject))
 	{
-		bool bNewHasLandscapeGrass = LandscapeComponents.ContainsByPredicate([](ULandscapeComponent* Component) { return Component->MaterialHasGrass(); });
-		if (bHasLandscapeGrass != bNewHasLandscapeGrass)
+		// TODO [jonathan.bard], todo(luc.eygasier): do this part any time UpdateGrassTypes is called so that it's always in sync :
+		bool bHasLandscapeGrassBefore = bHasLandscapeGrass;
+		float GrassTypesMaxDiscardDistanceBefore = GrassTypesMaxDiscardDistance;
+
+		GrassTypesMaxDiscardDistance = 0.0f;
+		bHasLandscapeGrass = false;
+		for (ULandscapeComponent* Component : LandscapeComponents)
+		{
+			bHasLandscapeGrass |= (Component->GetGrassTypes().Num() > 0);
+			GrassTypesMaxDiscardDistance = FMath::Max(Component->GetGrassTypesMaxDiscardDistance(), GrassTypesMaxDiscardDistance);
+		}
+
+		if ((bHasLandscapeGrass != bHasLandscapeGrassBefore)
+			|| (GrassTypesMaxDiscardDistance != GrassTypesMaxDiscardDistanceBefore))
 		{
-			bHasLandscapeGrass = bNewHasLandscapeGrass;
 			MarkPackageDirty();
 		}
 
@@ -3202,7 +3179,7 @@ void ALandscapeProxy::UpdateGrassDataStatus(TSet<UTexture2D*>* OutCurrentForcedS
 	{
 		*OutOutdatedGrassMaps = 0;
 	}
-		
+
 	const UWorld* World = GetWorld();
 	if (!World || World->IsGameWorld())
 	{
@@ -3212,11 +3189,6 @@ void ALandscapeProxy::UpdateGrassDataStatus(TSet<UTexture2D*>* OutCurrentForcedS
 	// In either case we want to check the Grass textures stream state
 	const bool bCheckStreamingState = OutDesiredForcedStreamedTextures || bInEnableForceResidentFlag;
 
-	TArray<ULandscapeGrassType*> GrassTypes;
-	float OutMaxDiscardDistance = 0.0f;
-	GetGrassTypes(World, GetLandscapeMaterial(), GrassTypes, OutMaxDiscardDistance);
-	const bool bHasGrassTypes = GrassTypes.Num() > 0;
-
 	const bool bIsOutermostPackageDirty = GetOutermost()->IsDirty();
 
 	int32 OutdatedGrassMaps = 0;
@@ -3244,14 +3216,26 @@ void ALandscapeProxy::UpdateGrassDataStatus(TSet<UTexture2D*>* OutCurrentForcedS
 				}
 			}
 
-			if (OutOutdatedComponents && Component->IsGrassMapOutdated())
+			bool bIsGrassMapOutdated = Component->IsGrassMapOutdated();
+			if (bIsGrassMapOutdated)
 			{
-				OutOutdatedComponents->Add(Component);
+#if WITH_EDITOR
+				// If the material has changed we need to cache the GrassTypes and related info again : 
+				Component->UpdateGrassTypes();
+#endif // WITH_EDITOR
+
+				if (OutOutdatedComponents)
+				{
+					OutOutdatedComponents->Add(Component);
+				}
 			}
 
+			// Needs to be called after UpdateGrassTypes
+			TArray<ULandscapeGrassType*> GrassTypes = Component->GetGrassTypes();
+			bool bHasGrassTypes = (GrassTypes.Num() > 0);
 			if (bHasGrassTypes || bBakeMaterialPositionOffsetIntoCollision)
 			{
-				if (Component->IsGrassMapOutdated() || !Component->GrassData->HasValidData() || GGrassUpdateAllOnRebuild != 0)
+				if (bIsGrassMapOutdated || !Component->GrassData->HasValidData() || (GGrassUpdateAllOnRebuild != 0))
 				{
 					if (OutComponentsNeedingGrassMapRender)
 					{
@@ -3431,25 +3415,10 @@ void ALandscapeProxy::UpdateGrass(const TArray<FVector>& Cameras, int32& InOutNu
 
 		if (World)
 		{
-			UMaterialInterface* LandscapeMaterialInterface = GetLandscapeMaterial();
-#if WITH_EDITOR
-			TArray<ULandscapeGrassType*> LandscapeGrassTypes;
-			float GrassMaxDiscardDistance = 0.0f;
-			GetGrassTypes(World, LandscapeMaterialInterface, LandscapeGrassTypes, GrassMaxDiscardDistance);
-#else
-			// In non editor builds, cache grass types for performance.
-			if (LandscapeMaterialInterface != LandscapeMaterialCached)
-			{
-				LandscapeMaterialCached = LandscapeMaterialInterface;
-				LandscapeGrassTypes.Reset();
-				GrassMaxDiscardDistance = 0.0f;
-				GetGrassTypes(World, LandscapeMaterialInterface, LandscapeGrassTypes, GrassMaxDiscardDistance);
-			}
-#endif
 			// Cull grass max distance based on Cull Distance scale factor and on max GuardBand factor.
-			float GrassMaxCulledDiscardDistance = GrassMaxDiscardDistance * GGrassCullDistanceScale *
-			FMath::Max(GGuardBandDiscardMultiplier, GGuardBandMultiplier);
+			float GrassMaxCulledDiscardDistance = GrassTypesMaxDiscardDistance * GGrassCullDistanceScale * FMath::Max(GGuardBandDiscardMultiplier, GGuardBandMultiplier);
 			float GrassMaxSquareDiscardDistance = GrassMaxCulledDiscardDistance * GrassMaxCulledDiscardDistance;
+
 #if WITH_EDITOR
 
 
@@ -3485,6 +3454,7 @@ void ALandscapeProxy::UpdateGrass(const TArray<FVector>& Cameras, int32& InOutNu
 
 			static TArray<SortedLandscapeElement> SortedLandscapeComponents;
 			SortedLandscapeComponents.Reset(LandscapeComponents.Num());
+			
 			for (ULandscapeComponent* Component : LandscapeComponents)
 			{
 				// skip if we have no data and no way to generate it
@@ -3493,13 +3463,15 @@ void ALandscapeProxy::UpdateGrass(const TArray<FVector>& Cameras, int32& InOutNu
 					continue;
 				}
 				FBoxSphereBounds WorldBounds = Component->CalcBounds(Component->GetComponentTransform());
+			
 				float MinSqrDistanceToComponent = Cameras.Num() ? MAX_flt : 0.0f;
+				
 				for (const FVector& CameraPos : Cameras)
 				{
 					MinSqrDistanceToComponent = FMath::Min<float>(MinSqrDistanceToComponent, static_cast<float>(WorldBounds.ComputeSquaredDistanceFromBoxToPoint(CameraPos)));
 				}
 
-				if (MinSqrDistanceToComponent > GrassMaxSquareDiscardDistance)
+				if ((GrassMaxSquareDiscardDistance > 0.f) && (MinSqrDistanceToComponent > GrassMaxSquareDiscardDistance))
 				{
 					continue;
 				}
@@ -3541,6 +3513,7 @@ void ALandscapeProxy::UpdateGrass(const TArray<FVector>& Cameras, int32& InOutNu
 					Component->ChangeTag = GGrassExclusionChangeTag;
 				}
 
+				const TArray<ULandscapeGrassType*>& LandscapeGrassTypes = Component->GetGrassTypes();
 				for (ULandscapeGrassType* GrassType : LandscapeGrassTypes)
 				{
 					if (GrassType)
@@ -3869,7 +3842,7 @@ void ALandscapeProxy::UpdateGrass(const TArray<FVector>& Cameras, int32& InOutNu
 			{
 				int32 NumComponentsRendered = 0;
 				int32 NumComponentsUnableToRender = 0;
-				if ((LandscapeGrassTypes.Num() > 0 && GPrerenderGrassmaps > 0) || bBakeMaterialPositionOffsetIntoCollision)
+				if ((GPrerenderGrassmaps > 0) || bBakeMaterialPositionOffsetIntoCollision)
 				{
 					// try to render some grassmaps
 					TArray<ULandscapeComponent*> ComponentsToRender;
@@ -3910,7 +3883,7 @@ void ALandscapeProxy::UpdateGrass(const TArray<FVector>& Cameras, int32& InOutNu
 						RenderCaptureInterface::FScopedCapture RenderCapture((GCaptureNextGrassUpdate != 0), TEXT("RenderGrassMaps"));
 						GCaptureNextGrassUpdate = FMath::Max(GCaptureNextGrassUpdate - 1, 0);
 
-						RenderGrassMaps(ComponentsToRender, LandscapeGrassTypes);
+						RenderGrassMaps(ComponentsToRender);
 					}
 				}
 
