Fix for crash when undoing the adding of a new landscape component.

UpdatedCachedBounds() was added to PostEditUndo, but cannot handle a missing heightmap, which occurs when the relevant component has been flagged for deletion (i.e. it is undoing the creation of the component)

#localization none
#tests local testing on 100% repro case to validate the crash no longer occurs

diff --git a/Engine/Source/Runtime/Landscape/Private/LandscapeEdit.cpp b/Engine/Source/Runtime/Landscape/Private/LandscapeEdit.cpp
--- a/Engine/Source/Runtime/Landscape/Private/LandscapeEdit.cpp
+++ b/Engine/Source/Runtime/Landscape/Private/LandscapeEdit.cpp
@@ -706,7 +706,10 @@ void ULandscapeComponent::PostEditUndo()
 		// Updating the bounds after undo will re-sync CachedLocalBox with the heightmap data. This is necessary because ULandscapeComponent might be part of 
 		//  the transaction buffer, but the final heightmaps aren't, so undoing might desynchronize them. If the heightmap update that is requested right after 
 		//  ends up modifying heightmap data, then CachedLocalBox will get updated properly again so this is a 'safe' way to fix the issue
-		UpdateCachedBounds();
+		if (IsValid(this))
+		{
+			UpdateCachedBounds();
+		}
 
 		const bool bUpdateAll = true;
 		RequestHeightmapUpdate(bUpdateAll);
