Adds StateHash when ULandscapeGrassType are modified.
ComputeGrassMapGenerationHash uses the cached ULandscapeGrassType StateHash.

diff --git a/Engine/Source/Runtime/Landscape/Classes/LandscapeGrassType.h b/Engine/Source/Runtime/Landscape/Classes/LandscapeGrassType.h
--- a/Engine/Source/Runtime/Landscape/Classes/LandscapeGrassType.h
+++ b/Engine/Source/Runtime/Landscape/Classes/LandscapeGrassType.h
@@ -171,6 +171,9 @@ class ULandscapeGrassType : public UObject
 	UPROPERTY(EditAnywhere, Category = Scalability)
 	uint32 bEnableDensityScaling : 1;
 
+	UPROPERTY()
+	uint32 StateHash = 0;
+
 	UPROPERTY()
 	TObjectPtr<UStaticMesh> GrassMesh_DEPRECATED;
 	UPROPERTY()
@@ -185,13 +188,15 @@ class ULandscapeGrassType : public UObject
 	bool RandomRotation_DEPRECATED;
 	UPROPERTY()
 	bool AlignToSurface_DEPRECATED;
-
+	
 	//~ Begin UObject Interface
 	virtual void PostLoad() override;
 #if WITH_EDITOR
 	virtual void PostEditChangeProperty(FPropertyChangedEvent& PropertyChangedEvent) override;
 #endif
 	//~ End UObject Interface
+
+	uint32 ComputeStateHash();
 };
 
 
diff --git a/Engine/Source/Runtime/Landscape/Private/LandscapeGrass.cpp b/Engine/Source/Runtime/Landscape/Private/LandscapeGrass.cpp
--- a/Engine/Source/Runtime/Landscape/Private/LandscapeGrass.cpp
+++ b/Engine/Source/Runtime/Landscape/Private/LandscapeGrass.cpp
@@ -1574,9 +1574,6 @@ uint32 ULandscapeComponent::ComputeGrassMapGenerationHash() const
 		if (UMaterial* MaterialBase = Material->GetMaterial())
 		{
 			// If anything changes in the grass types, we should take that into account as well :
-			// TODO [jonathan.bard], todo(luc.eygasier): make this event-based... 
-			FArchiveCrc32 Ar(Hash);
-			
 			for (ULandscapeGrassType* GrassType : GrassTypes)
 			{
 				if (GrassType == nullptr)
@@ -1584,9 +1581,8 @@ uint32 ULandscapeComponent::ComputeGrassMapGenerationHash() const
 					continue;
 				}
 
-				GrassType->Serialize(Ar);
+				Hash = FCrc::TypeCrc32(GrassType->StateHash, Hash);
 			}
-			Hash = Ar.GetCrc();
 		}
 	}
 
@@ -2017,8 +2013,9 @@ void ULandscapeGrassType::PostLoad()
 		GrassVarieties.Add(Grass);
 		GrassMesh_DEPRECATED = nullptr;
 	}
-}
 
+	StateHash = ComputeStateHash();
+}
 
 #if WITH_EDITOR
 void ULandscapeGrassType::PostEditChangeProperty(FPropertyChangedEvent& PropertyChangedEvent)
@@ -2054,9 +2051,22 @@ void ULandscapeGrassType::PostEditChangeProperty(FPropertyChangedEvent& Property
 			}
 		}
 	}
+
+	StateHash = ComputeStateHash();
 }
 #endif
 
+uint32 ULandscapeGrassType::ComputeStateHash()
+{
+	uint32 Hash = 0;
+
+	FArchiveCrc32 Ar;
+	Serialize(Ar);
+	Hash = Ar.GetCrc();
+
+	return Hash;
+}
+
 //
 // FLandscapeComponentGrassData
 //
