Landscape : Fix incorrect message reporting an update is required due to physical materials on landscapes with no physical material

diff --git a/Engine/Source/Runtime/Landscape/Private/LandscapeEdit.cpp b/Engine/Source/Runtime/Landscape/Private/LandscapeEdit.cpp
--- a/Engine/Source/Runtime/Landscape/Private/LandscapeEdit.cpp
+++ b/Engine/Source/Runtime/Landscape/Private/LandscapeEdit.cpp
@@ -1878,6 +1878,7 @@ void ULandscapeComponent::UpdatePhysicalMaterialTasks()
 		}
 		else
 		{
+			PhysicalMaterialHash = Hash;
 			// Clear the renderable physical material properties as we don't need them :
 			CollisionComponent->PhysicalMaterialRenderObjects.Reset();
 			CollisionComponent->PhysicalMaterialRenderData.RemoveBulkData();
