Fix landscape LOD not matching HLOD resolution when level instanced or PIE.

When instanced, the landscape changes it's GUID, but the HLOD does not.
This leads to the HLOD being registered with a different renderer, which means it doesn't communicate the neighboring LODs properly.

This change makes the HLOD components change their GUIDs when instanced the same way as landscape.

(Note this also affects PIE, as the PIE world is instanced off of the editor world)

diff --git a/Engine/Source/Runtime/Landscape/Classes/LandscapeMeshProxyComponent.h b/Engine/Source/Runtime/Landscape/Classes/LandscapeMeshProxyComponent.h
--- a/Engine/Source/Runtime/Landscape/Classes/LandscapeMeshProxyComponent.h
+++ b/Engine/Source/Runtime/Landscape/Classes/LandscapeMeshProxyComponent.h
@@ -40,5 +40,6 @@ public:
 	LANDSCAPE_API void InitializeForLandscape(ALandscapeProxy* Landscape, int8 InProxyLOD);
 	
 	virtual FPrimitiveSceneProxy* CreateSceneProxy() override;
+	virtual void PostLoad() override;
 };
 
diff --git a/Engine/Source/Runtime/Landscape/Private/Landscape.cpp b/Engine/Source/Runtime/Landscape/Private/Landscape.cpp
--- a/Engine/Source/Runtime/Landscape/Private/Landscape.cpp
+++ b/Engine/Source/Runtime/Landscape/Private/Landscape.cpp
@@ -3626,27 +3626,23 @@ bool ULandscapeInfo::UpdateLayerInfoMap(ALandscapeProxy* Proxy /*= nullptr*/, bo
  * this avoids guid collisions when you instance a world (and its landscapes) multiple times,
  * while maintaining the same GUID between landscape proxy objects within an instance
  */ 
-void ALandscapeProxy::PostLoadFixupLandscapeGuidsIfInstanced()
-{
-	// record the original value before modification
-	check(!OriginalLandscapeGuid.IsValid() || (OriginalLandscapeGuid == LandscapeGuid));
-	OriginalLandscapeGuid = this->LandscapeGuid;
-
+ void ChangeLandscapeGuidIfObjectIsInstanced(FGuid& InOutGuid, UObject* InObject)
+ {
 	// we shouldn't be dealing with any instanced landscapes in these cases, early out
-	if (this->IsTemplate() || IsRunningCookCommandlet())
+	if (InObject->IsTemplate() || IsRunningCookCommandlet())
 	{
 		return;
 	}
 
-	UWorldPartition* WorldPartition = FWorldPartitionHelpers::GetWorldPartition(this);
-	UWorld* OuterWorld = WorldPartition ? WorldPartition->GetTypedOuter<UWorld>() : this->GetTypedOuter<UWorld>();
+	UWorldPartition* WorldPartition = FWorldPartitionHelpers::GetWorldPartition(InObject);
+	UWorld* OuterWorld = WorldPartition ? WorldPartition->GetTypedOuter<UWorld>() : InObject->GetTypedOuter<UWorld>();
 
 	// TODO [chris.tchou] : Note this is not 100% correct, IsInstanced() returns TRUE when using PIE on non-instanced landscapes.
 	// That is generally ok however, as the GUID remaps are still deterministic and landscape still works.
 	if (OuterWorld && OuterWorld->IsInstanced())
 	{
 		FArchiveMD5 Ar;
-		FGuid OldLandscapeGuid = this->LandscapeGuid;
+		FGuid OldLandscapeGuid = InOutGuid;
 		Ar << OldLandscapeGuid;
 
 		UPackage* OuterWorldPackage = OuterWorld->GetPackage();
@@ -3656,9 +3652,18 @@ void ALandscapeProxy::PostLoadFixupLandscapeGuidsIfInstanced()
 			Ar << PackageName;
 		}
 
-		this->LandscapeGuid = Ar.GetGuidFromHash();
+		InOutGuid = Ar.GetGuidFromHash();
 	}
 }
+ 
+void ALandscapeProxy::PostLoadFixupLandscapeGuidsIfInstanced()
+{
+	// record the original value before modification
+	check(!OriginalLandscapeGuid.IsValid() || (OriginalLandscapeGuid == LandscapeGuid));
+	OriginalLandscapeGuid = this->LandscapeGuid;
+
+	ChangeLandscapeGuidIfObjectIsInstanced(this->LandscapeGuid, this);
+}
 
 void ALandscapeProxy::PostLoad()
 {
@@ -5659,6 +5664,13 @@ ULandscapeMeshProxyComponent::ULandscapeMeshProxyComponent(const FObjectInitiali
 {
 }
 
+void ULandscapeMeshProxyComponent::PostLoad()
+{
+	Super::PostLoad();
+
+	ChangeLandscapeGuidIfObjectIsInstanced(this->LandscapeGuid, this);
+}
+
 void ULandscapeMeshProxyComponent::InitializeForLandscape(ALandscapeProxy* Landscape, int8 InProxyLOD)
 {
 	LandscapeGuid = Landscape->GetLandscapeGuid();
