Expose landscape force layer update to BP

diff --git a/Engine/Source/Runtime/Landscape/Classes/Landscape.h b/Engine/Source/Runtime/Landscape/Classes/Landscape.h
--- a/Engine/Source/Runtime/Landscape/Classes/Landscape.h
+++ b/Engine/Source/Runtime/Landscape/Classes/Landscape.h
@@ -345,6 +345,7 @@ public:
 	
 	LANDSCAPE_API void ToggleCanHaveLayersContent();
 	LANDSCAPE_API void ForceUpdateLayersContent(bool bIntermediateRender = false);
+	UFUNCTION(BlueprintCallable, Category = "Landscape")
 	LANDSCAPE_API void ForceLayersFullUpdate();
 	LANDSCAPE_API void InitializeLandscapeLayersWeightmapUsage();
 
