Fix for landscape grass type state hash changing every time it is calculated

diff --git a/Engine/Source/Runtime/Landscape/Private/LandscapeGrass.cpp b/Engine/Source/Runtime/Landscape/Private/LandscapeGrass.cpp
--- a/Engine/Source/Runtime/Landscape/Private/LandscapeGrass.cpp
+++ b/Engine/Source/Runtime/Landscape/Private/LandscapeGrass.cpp
@@ -1569,20 +1569,18 @@ uint32 ULandscapeComponent::ComputeGrassMapGenerationHash() const
 
 	if (UMaterialInterface* Material = GetLandscapeMaterial())
 	{
-		// Take into account any material state change :
+		// Take into account any material state change : (excluding texture state)
 		Hash = Material->ComputeAllStateCRC();
-		if (UMaterial* MaterialBase = Material->GetMaterial())
+
+		// If anything changes in the grass types, we should take that into account as well :
+		for (ULandscapeGrassType* GrassType : GrassTypes)
 		{
-			// If anything changes in the grass types, we should take that into account as well :
-			for (ULandscapeGrassType* GrassType : GrassTypes)
+			if (GrassType == nullptr)
 			{
-				if (GrassType == nullptr)
-				{
-					continue;
-				}
-
-				Hash = FCrc::TypeCrc32(GrassType->StateHash, Hash);
+				continue;
 			}
+
+			Hash = FCrc::TypeCrc32(GrassType->StateHash, Hash);
 		}
 	}
 
@@ -2061,9 +2059,17 @@ uint32 ULandscapeGrassType::ComputeStateHash()
 	uint32 Hash = 0;
 
 	FArchiveCrc32 Ar;
+
+    // Save and nullify previous Hash value (don't want to include it in the new calculated hash)
+    uint32 PreviousHash = StateHash;
+    StateHash = 0;
+
 	Serialize(Ar);
 	Hash = Ar.GetCrc();
 
+    // Restore previous hash value
+    StateHash = PreviousHash;
+
 	return Hash;
 }
 
