Expose ULandscapeGrassType properties to BP

diff --git a/Engine/Source/Runtime/Engine/Public/PerPlatformProperties.h b/Engine/Source/Runtime/Engine/Public/PerPlatformProperties.h
--- a/Engine/Source/Runtime/Engine/Public/PerPlatformProperties.h
+++ b/Engine/Source/Runtime/Engine/Public/PerPlatformProperties.h
@@ -207,7 +207,7 @@ ENGINE_API void operator<<(FStructuredArchive::FSlot Slot, TPerPlatformProperty<
 struct FFreezablePerPlatformInt;
 
 /** FPerPlatformInt - int32 property with per-platform overrides */
-USTRUCT()
+USTRUCT(BlueprintType)
 struct FPerPlatformInt
 #if CPP
 :	public TPerPlatformProperty<FPerPlatformInt, int32, NAME_IntProperty>
@@ -215,7 +215,7 @@ struct FPerPlatformInt
 {
 	GENERATED_USTRUCT_BODY()
 
-	UPROPERTY(EditAnywhere, Category = PerPlatform)
+	UPROPERTY(BlueprintReadOnly, EditAnywhere, Category = PerPlatform)
 	int32 Default;
 
 #if WITH_EDITORONLY_DATA
@@ -301,7 +301,7 @@ struct TStructOpsTypeTraits<FPerPlatformInt>
 struct FFreezablePerPlatformFloat;
 
 /** FPerPlatformFloat - float property with per-platform overrides */
-USTRUCT(meta = (CanFlattenStruct))
+USTRUCT(BlueprintType, meta = (CanFlattenStruct))
 struct FPerPlatformFloat
 #if CPP
 :	public TPerPlatformProperty<FPerPlatformFloat, float, NAME_FloatProperty>
@@ -309,7 +309,7 @@ struct FPerPlatformFloat
 {
 	GENERATED_USTRUCT_BODY()
 
-	UPROPERTY(EditAnywhere, Category = PerPlatform)
+	UPROPERTY(BlueprintReadOnly, EditAnywhere, Category = PerPlatform)
 	float Default;
 
 #if WITH_EDITORONLY_DATA
diff --git a/Engine/Source/Runtime/Engine/Public/PerQualityLevelProperties.h b/Engine/Source/Runtime/Engine/Public/PerQualityLevelProperties.h
--- a/Engine/Source/Runtime/Engine/Public/PerQualityLevelProperties.h
+++ b/Engine/Source/Runtime/Engine/Public/PerQualityLevelProperties.h
@@ -173,7 +173,7 @@ ENGINE_API FArchive& operator<<(FArchive& Ar, FPerQualityLevelProperty<_StructTy
 template<typename _StructType, typename _ValueType, EName _BasePropertyName>
 ENGINE_API void operator<<(FStructuredArchive::FSlot Slot, FPerQualityLevelProperty<_StructType, _ValueType, _BasePropertyName>& Property);
 
-USTRUCT()
+USTRUCT(BlueprintType)
 struct FPerQualityLevelInt 
 #if CPP
 	:	public FPerQualityLevelProperty<FPerQualityLevelInt, int32, NAME_IntProperty>
@@ -181,7 +181,7 @@ struct FPerQualityLevelInt
 {
 	GENERATED_USTRUCT_BODY()
 
-	UPROPERTY(EditAnywhere, Category = PerQualityLevel)
+	UPROPERTY(BlueprintReadOnly, EditAnywhere, Category = PerQualityLevel)
 	int32 Default;
 
 	UPROPERTY(EditAnywhere, Category = PerQualityLevel)
@@ -216,7 +216,7 @@ struct TStructOpsTypeTraits<FPerQualityLevelInt>
 	static constexpr EPropertyObjectReferenceType WithSerializerObjectReferences = EPropertyObjectReferenceType::None;
 };
 
-USTRUCT()
+USTRUCT(BlueprintType)
 struct FPerQualityLevelFloat
 #if CPP
 	:	public FPerQualityLevelProperty<FPerQualityLevelFloat, float, NAME_FloatProperty>
@@ -224,7 +224,7 @@ struct FPerQualityLevelFloat
 {
 	GENERATED_USTRUCT_BODY()
 
-	UPROPERTY(EditAnywhere, Category = PerQualityLevel)
+	UPROPERTY(BlueprintReadOnly, EditAnywhere, Category = PerQualityLevel)
 	float Default;
 
 	UPROPERTY(EditAnywhere, Category = PerQualityLevel)
diff --git a/Engine/Source/Runtime/Landscape/Classes/LandscapeGrassType.h b/Engine/Source/Runtime/Landscape/Classes/LandscapeGrassType.h
--- a/Engine/Source/Runtime/Landscape/Classes/LandscapeGrassType.h
+++ b/Engine/Source/Runtime/Landscape/Classes/LandscapeGrassType.h
@@ -25,117 +25,117 @@ enum class EGrassScaling : uint8
 	LockXY,
 };
 
-USTRUCT()
-struct FGrassVariety
+USTRUCT(BlueprintType)
+struct LANDSCAPE_API FGrassVariety
 {
 	GENERATED_USTRUCT_BODY()
 
 	FGrassVariety();
 
-	UPROPERTY(EditAnywhere, Category=Grass)
+	UPROPERTY(BlueprintReadOnly, EditAnywhere, Category=Grass)
 	TObjectPtr<UStaticMesh> GrassMesh;
 
-	UPROPERTY(EditAnywhere, Category=Grass, meta = (ToolTip = "Material Overrides."))
+	UPROPERTY(BlueprintReadOnly, EditAnywhere, Category=Grass, meta = (ToolTip = "Material Overrides."))
 	TArray<TObjectPtr<class UMaterialInterface>> OverrideMaterials;
 
 	/* Instances per 10 square meters. */
-	UPROPERTY(EditAnywhere, Category=Grass, meta = (UIMin = 0, ClampMin = 0, UIMax = 1000, ClampMax = 1000))
+	UPROPERTY(BlueprintReadOnly, EditAnywhere, Category=Grass, meta = (UIMin = 0, ClampMin = 0, UIMax = 1000, ClampMax = 1000))
 	FPerPlatformFloat GrassDensity;
 
-	UPROPERTY(EditAnywhere, Category=Grass, meta = (UIMin = 0, ClampMin = 0, UIMax = 1000, ClampMax = 1000))
+	UPROPERTY(BlueprintReadOnly, EditAnywhere, Category=Grass, meta = (UIMin = 0, ClampMin = 0, UIMax = 1000, ClampMax = 1000))
 	FPerQualityLevelFloat GrassDensityQuality;
 
 	/* If true, use a jittered grid sequence for placement, otherwise use a halton sequence. */
-	UPROPERTY(EditAnywhere, Category=Grass)
+	UPROPERTY(BlueprintReadOnly, EditAnywhere, Category=Grass)
 	bool bUseGrid;
 
-	UPROPERTY(EditAnywhere, Category=Grass, meta = (EditCondition = "bUseGrid", UIMin = 0, ClampMin = 0, UIMax = 1, ClampMax = 1))
+	UPROPERTY(BlueprintReadOnly, EditAnywhere, Category=Grass, meta = (EditCondition = "bUseGrid", UIMin = 0, ClampMin = 0, UIMax = 1, ClampMax = 1))
 	float PlacementJitter;
 
 	/* The distance where instances will begin to fade out if using a PerInstanceFadeAmount material node. 0 disables. */
-	UPROPERTY(EditAnywhere, Category=Grass, meta = (UIMin = 0, ClampMin = 0, UIMax = 1000000, ClampMax = 1000000))
+	UPROPERTY(BlueprintReadOnly, EditAnywhere, Category=Grass, meta = (UIMin = 0, ClampMin = 0, UIMax = 1000000, ClampMax = 1000000))
 	FPerPlatformInt StartCullDistance;
 
-	UPROPERTY(EditAnywhere, Category = Grass, meta = (UIMin = 0, ClampMin = 0, UIMax = 1000000, ClampMax = 1000000))
+	UPROPERTY(BlueprintReadOnly, EditAnywhere, Category = Grass, meta = (UIMin = 0, ClampMin = 0, UIMax = 1000000, ClampMax = 1000000))
 	FPerQualityLevelInt StartCullDistanceQuality;
 
 	/**
 	 * The distance where instances will have completely faded out when using a PerInstanceFadeAmount material node. 0 disables. 
 	 * When the entire cluster is beyond this distance, the cluster is completely culled and not rendered at all.
 	 */
-	UPROPERTY(EditAnywhere, Category = Grass, meta = (UIMin = 0, ClampMin = 0, UIMax = 1000000, ClampMax = 1000000))
+	UPROPERTY(BlueprintReadOnly, EditAnywhere, Category = Grass, meta = (UIMin = 0, ClampMin = 0, UIMax = 1000000, ClampMax = 1000000))
 	FPerPlatformInt EndCullDistance;
 
-	UPROPERTY(EditAnywhere, Category = Grass, meta = (UIMin = 0, ClampMin = 0, UIMax = 1000000, ClampMax = 1000000))
+	UPROPERTY(BlueprintReadOnly, EditAnywhere, Category = Grass, meta = (UIMin = 0, ClampMin = 0, UIMax = 1000000, ClampMax = 1000000))
 	FPerQualityLevelInt EndCullDistanceQuality;
 	/** 
 	 * Specifies the smallest LOD that will be used for this component.
 	 * If -1 (default), the MinLOD of the static mesh asset will be used instead.
 	 */
-	UPROPERTY(EditAnywhere, Category = Grass, meta = (UIMin = -1, ClampMin = -1, UIMax = 8, ClampMax = 8))
+	UPROPERTY(BlueprintReadOnly, EditAnywhere, Category = Grass, meta = (UIMin = -1, ClampMin = -1, UIMax = 8, ClampMax = 8))
 	int32 MinLOD;
 
 	/** Specifies grass instance scaling type */
-	UPROPERTY(EditAnywhere, Category=Grass)
+	UPROPERTY(BlueprintReadOnly, EditAnywhere, Category=Grass)
 	EGrassScaling Scaling;
 
 	/** Specifies the range of scale, from minimum to maximum, to apply to a grass instance's X Scale property */
-	UPROPERTY(EditAnywhere, Category=Grass)
+	UPROPERTY(BlueprintReadOnly, EditAnywhere, Category=Grass)
 	FFloatInterval ScaleX;
 
 	/** Specifies the range of scale, from minimum to maximum, to apply to a grass instance's Y Scale property */
-	UPROPERTY(EditAnywhere, Category=Grass, meta = (EditCondition = "Scaling == EGrassScaling::Free"))
+	UPROPERTY(BlueprintReadOnly, EditAnywhere, Category=Grass, meta = (EditCondition = "Scaling == EGrassScaling::Free"))
 	FFloatInterval ScaleY;
 
 	/** Specifies the range of scale, from minimum to maximum, to apply to a grass instance's Z Scale property */
-	UPROPERTY(EditAnywhere, Category=Grass, meta = (EditCondition = "Scaling == EGrassScaling::Free || Scaling == EGrassScaling::LockXY"))
+	UPROPERTY(BlueprintReadOnly, EditAnywhere, Category=Grass, meta = (EditCondition = "Scaling == EGrassScaling::Free || Scaling == EGrassScaling::LockXY"))
 	FFloatInterval ScaleZ;
 
 	/** If enabled the the scale of instances is reduced as the weight (density) decreases */
-	UPROPERTY(EditAnywhere, AdvancedDisplay, Category = Grass)
+	UPROPERTY(BlueprintReadOnly, EditAnywhere, AdvancedDisplay, Category = Grass)
 	bool bWeightAttenuatesMaxScale = false;
 
 	/** Modulate the scale of the instances based on weight (normalized density). The weight range (ScaleWeightAttenuation, 1.0) maps to (scaleMin, scaleMax), weight values less than ScaleWeightAttenuation are set to minScale */
-	UPROPERTY(EditAnywhere, AdvancedDisplay, Category = Grass, meta = (UIMin = 0, ClampMin = 0, UIMax = 1, ClampMax = 1, EditCondition = "bWeightAttenuatesMaxScale == true"))
+	UPROPERTY(BlueprintReadOnly, EditAnywhere, AdvancedDisplay, Category = Grass, meta = (UIMin = 0, ClampMin = 0, UIMax = 1, ClampMax = 1, EditCondition = "bWeightAttenuatesMaxScale == true"))
 	float MaxScaleWeightAttenuation = 0.5f;
 
 	/** Whether the grass instances should be placed at random rotation (true) or all at the same rotation (false) */
-	UPROPERTY(EditAnywhere, Category = Grass)
+	UPROPERTY(BlueprintReadOnly, EditAnywhere, Category = Grass)
 	bool RandomRotation;
 
 	/** Whether the grass instances should be tilted to the normal of the landscape (true), or always vertical (false) */
-	UPROPERTY(EditAnywhere, Category = Grass)
+	UPROPERTY(BlueprintReadOnly, EditAnywhere, Category = Grass)
 	bool AlignToSurface;
 
 	/* Whether to use the landscape's lightmap when rendering the grass. */
-	UPROPERTY(EditAnywhere, Category = Grass)
+	UPROPERTY(BlueprintReadOnly, EditAnywhere, Category = Grass)
 	bool bUseLandscapeLightmap;
 
 	/**
 	 * Lighting channels that the grass will be assigned. Lights with matching channels will affect the grass.
 	 * These channels only apply to opaque materials, direct lighting, and dynamic lighting and shadowing.
 	 */
-	UPROPERTY(EditAnywhere, AdvancedDisplay, Category = Grass)
+	UPROPERTY(BlueprintReadOnly, EditAnywhere, AdvancedDisplay, Category = Grass)
 	FLightingChannels LightingChannels;
 
 	/** Whether the grass instances should receive decals. */
-	UPROPERTY(EditAnywhere, Category = Grass)
+	UPROPERTY(BlueprintReadOnly, EditAnywhere, Category = Grass)
 	bool bReceivesDecals;
 
 	/** Controls whether the primitive should affect dynamic distance field lighting methods. */
-	UPROPERTY(EditAnywhere, Category = Grass)
+	UPROPERTY(BlueprintReadOnly, EditAnywhere, Category = Grass)
 	bool bAffectDistanceFieldLighting;
 
 	/** Whether the grass should cast shadows when using non-precomputed shadowing. **/
-	UPROPERTY(EditAnywhere, Category = Grass)
+	UPROPERTY(BlueprintReadOnly, EditAnywhere, Category = Grass)
 	bool bCastDynamicShadow;
 
 	/** Whether the grass should cast contact shadows. **/
-	UPROPERTY(EditAnywhere, Category = Grass)
+	UPROPERTY(BlueprintReadOnly, EditAnywhere, Category = Grass)
 	bool bCastContactShadow;
 
 	/** Whether we should keep a cpu copy of the instance buffer. This should be set to true if you plan on using GetOverlappingXXXXCount functions of the component otherwise it won't return any data.**/
-	UPROPERTY(EditAnywhere, Category = Grass)
+	UPROPERTY(BlueprintReadOnly, EditAnywhere, Category = Grass)
 	bool bKeepInstanceBufferCPUCopy;
 
 	/** Distance at which to grass instances should disable WPO for performance reasons */
@@ -143,7 +143,7 @@ struct FGrassVariety
 	uint32 InstanceWorldPositionOffsetDisableDistance;
 
 	/** Control shadow invalidation behavior, in particular with respect to Virtual Shadow Maps and material effects like World Position Offset. */
-	UPROPERTY(EditAnywhere, Category=Grass, AdvancedDisplay)
+	UPROPERTY(BlueprintReadOnly, EditAnywhere, Category=Grass, AdvancedDisplay)
 	EShadowCacheInvalidationBehavior ShadowCacheInvalidationBehavior;
 
 	bool IsGrassQualityLevelEnable() const;
@@ -156,19 +156,19 @@ struct FGrassVariety
 
 };
 
-UCLASS(MinimalAPI)
+UCLASS(BlueprintType, MinimalAPI)
 class ULandscapeGrassType : public UObject
 {
 	GENERATED_UCLASS_BODY()
 
-	UPROPERTY(EditAnywhere, Category = Grass)
+	UPROPERTY(BlueprintReadOnly, EditAnywhere, Category = Grass)
 	TArray<FGrassVariety> GrassVarieties;
 
 	/**
 	* Whether this grass type should be affected by the Engine Scalability system's Foliage grass.DensityScale setting. 
 	* This is enabled by default but can be disabled should this grass type be important for gameplay reasons.
 	*/
-	UPROPERTY(EditAnywhere, Category = Scalability)
+	UPROPERTY(BlueprintReadOnly, EditAnywhere, Category = Scalability)
 	uint32 bEnableDensityScaling : 1;
 
 	UPROPERTY()
