Fixes invalid ZOffset overflowing the grass exporter projection matrix.

diff --git a/Engine/Source/Runtime/Landscape/Private/LandscapeGrass.cpp b/Engine/Source/Runtime/Landscape/Private/LandscapeGrass.cpp
--- a/Engine/Source/Runtime/Landscape/Private/LandscapeGrass.cpp
+++ b/Engine/Source/Runtime/Landscape/Private/LandscapeGrass.cpp
@@ -1312,7 +1312,7 @@ public:
 							          FPlane(0.0f, 0.0f, -1.0f, 0.0f),
 							          FPlane(0.0f, 0.0f, 0.0f, 1.0f));
 
-		const float ZOffset = WORLD_MAX;
+		const float ZOffset = UE_OLD_WORLD_MAX;
 		ProjectionMatrix = FReversedZOrthoMatrix(
 			TargetExtent.X,
 			TargetExtent.Y,
