Do not fill FAsyncGrassBuilder::InstanceData on platforms that have no Nanite HW Support

diff --git a/Engine/Source/Runtime/Foliage/Public/GrassInstancedStaticMeshComponent.h b/Engine/Source/Runtime/Foliage/Public/GrassInstancedStaticMeshComponent.h
--- a/Engine/Source/Runtime/Foliage/Public/GrassInstancedStaticMeshComponent.h
+++ b/Engine/Source/Runtime/Foliage/Public/GrassInstancedStaticMeshComponent.h
@@ -17,5 +17,10 @@ public:
 	FOLIAGE_API void AcceptPrebuiltTree(TArray<FInstancedStaticMeshInstanceData>& InInstanceData, TArray<FClusterNode>& InClusterTree, int32 InOcclusionLayerNumNodes, int32 InNumBuiltRenderInstances);
 
 	FOLIAGE_API bool SupportsWorldPositionOffsetVelocity() const override;
+
+	inline bool RequiresInstanceDataForTree() const
+	{
+		return ShouldCreateNaniteProxy();
+	}
 };
 
diff --git a/Engine/Source/Runtime/Landscape/Private/LandscapeGrass.cpp b/Engine/Source/Runtime/Landscape/Private/LandscapeGrass.cpp
--- a/Engine/Source/Runtime/Landscape/Private/LandscapeGrass.cpp
+++ b/Engine/Source/Runtime/Landscape/Private/LandscapeGrass.cpp
@@ -2484,6 +2484,7 @@ struct FAsyncGrassBuilder : public FGrassBuilderBase
 	FVector2D LightMapComponentBias;
 	FVector2D LightMapComponentScale;
 	bool RequireCPUAccess;
+	bool RequiresInstanceDataForTree;
 
 	TArray<FBox> ExcludedBoxes;
 
@@ -2523,7 +2524,7 @@ struct FAsyncGrassBuilder : public FGrassBuilderBase
 		, LightMapComponentBias(FVector2D::ZeroVector)
 		, LightMapComponentScale(FVector2D::UnitVector)
 		, RequireCPUAccess(GrassVariety.bKeepInstanceBufferCPUCopy)
-
+		, RequiresInstanceDataForTree(GrassInstancedStaticMeshComponent->RequiresInstanceDataForTree())
 		// output
 		, InstanceBuffer(/*bSupportsVertexHalfFloat*/ true)
 		, ClusterTree()
@@ -2873,10 +2874,13 @@ struct FAsyncGrassBuilder : public FGrassBuilderBase
 			TArray<float> InstanceCustomDataDummy;
 			UGrassInstancedStaticMeshComponent::BuildTreeAnyThread(InstanceTransforms, InstanceCustomDataDummy, 0, MeshBox, ClusterTree, SortedInstances, InstanceReorderTable, OutOcclusionLayerNum, DesiredInstancesPerLeaf, false);
 
-			InstanceData.Reset(NumInstances);
-			for (const FMatrix& Transform : InstanceTransforms)
+			if (RequiresInstanceDataForTree)
 			{
-				InstanceData.Emplace(Transform);
+				InstanceData.Reset(NumInstances);
+				for (const FMatrix& Transform : InstanceTransforms)
+				{
+					InstanceData.Emplace(Transform);
+				}
 			}
 			
 			// in-place sort the instances and generate the sorted instance data
@@ -2888,7 +2892,10 @@ struct FAsyncGrassBuilder : public FGrassBuilderBase
 				{
 					check(LoadFrom > FirstUnfixedIndex);
 					InstanceBuffer.SwapInstance(FirstUnfixedIndex, LoadFrom);
-					InstanceData.Swap(FirstUnfixedIndex, LoadFrom);
+					if (RequiresInstanceDataForTree)
+					{
+						InstanceData.Swap(FirstUnfixedIndex, LoadFrom);
+					}
 
 					int32 SwapGoesTo = InstanceReorderTable[FirstUnfixedIndex];
 					check(SwapGoesTo > FirstUnfixedIndex);
