Adds GrassGenerationHashKey to force invalidate grass data in Landscape Components.

diff --git a/Engine/Source/Runtime/Landscape/Private/LandscapeGrass.cpp b/Engine/Source/Runtime/Landscape/Private/LandscapeGrass.cpp
--- a/Engine/Source/Runtime/Landscape/Private/LandscapeGrass.cpp
+++ b/Engine/Source/Runtime/Landscape/Private/LandscapeGrass.cpp
@@ -1569,8 +1569,12 @@ uint32 ULandscapeComponent::ComputeGrassMapGenerationHash() const
 
 	if (UMaterialInterface* Material = GetLandscapeMaterial())
 	{
+		// Bump the generation hash key to invalidate all cached grass density maps.
+		static FGuid GrassGenerationHashKey("216D95C7651D4095ADC6A8459B4F181D");
+
+		Hash = GetTypeHash(GrassGenerationHashKey);
 		// Take into account any material state change : (excluding texture state)
-		Hash = Material->ComputeAllStateCRC();
+		Hash = FCrc::TypeCrc32(Material->ComputeAllStateCRC(), Hash);
 
 		// If anything changes in the grass types, we should take that into account as well :
 		for (ULandscapeGrassType* GrassType : GrassTypes)
