Fix terrain collision disappearance by keeping the chaos collision heightfield alive after Unregister(), so that the data is available if Register() is called again.  Chaos collision heightfields are ref-counted, and deallocate automatically when all references are destroyed.  This change also removes the old workaround that disabled the deleting of the cooked collision data.

diff --git a/Engine/Source/Runtime/Landscape/Classes/LandscapeHeightfieldCollisionComponent.h b/Engine/Source/Runtime/Landscape/Classes/LandscapeHeightfieldCollisionComponent.h
--- a/Engine/Source/Runtime/Landscape/Classes/LandscapeHeightfieldCollisionComponent.h
+++ b/Engine/Source/Runtime/Landscape/Classes/LandscapeHeightfieldCollisionComponent.h
@@ -165,7 +165,7 @@ public:
 	 *	Cooked HeightField data. Serialized only with cooked content 
 	 *	Stored as array instead of BulkData to take advantage of precaching during async loading
 	 */
-	bool CookedCollisionDataWasDeleted = false;
+	bool bCookedCollisionDataWasDeleted = false;
 	TArray<uint8>								CookedCollisionData;
 	
 	/** This is a list of physical materials that is actually used by a cooked HeightField */
@@ -175,6 +175,11 @@ public:
 	/** Physics engine version of heightfield data. */
 	TRefCountPtr<FHeightfieldGeometryRef>	HeightfieldRef;
 	
+	// local non-serialized ref counted pointers to keep the chaos heightfields alive between Unregister() and actual destruction of the component.
+	// this allows us to re-use them if the component gets a call to Register() again
+	Chaos::FHeightFieldPtr LocalHeightfieldGeometryRef;
+	Chaos::FHeightFieldPtr LocalHeightfieldSimpleGeometryRef;
+
 	/** Cached PxHeightFieldSamples values for navmesh generation. Note that it's being used only if navigation octree is set up for lazy geometry exporting */
 	int32 HeightfieldRowsCount;
 	int32 HeightfieldColumnsCount;
diff --git a/Engine/Source/Runtime/Landscape/Private/LandscapeCollision.cpp b/Engine/Source/Runtime/Landscape/Private/LandscapeCollision.cpp
--- a/Engine/Source/Runtime/Landscape/Private/LandscapeCollision.cpp
+++ b/Engine/Source/Runtime/Landscape/Private/LandscapeCollision.cpp
@@ -101,12 +101,6 @@ static TAutoConsoleVariable<bool> CVarLandscapeCollisionMeshShowPhysicalMaterial
 	TEXT("When enabled, vertex colors of the collision mesh are chosen based on the physical material"),
 	ECVF_RenderThreadSafe);
 
-static TAutoConsoleVariable<bool> CVarLandscapeCollisionDeleteCookedDataAfterUse(
-	TEXT("landscape.Collision.DeleteCookedDataAfterUse"),
-	true,
-	TEXT("When enabled, landscape cooked collision data is deleted after it is used to construct a physics heightfield.  Enabling this will reclaim memory, but may cause problems if the components are ever re-registered and collision must be recreated."));
-
-
 #if ENABLE_COOK_STATS
 namespace LandscapeCollisionCookStats
 {
@@ -263,6 +257,14 @@ void ULandscapeHeightfieldCollisionComponent::OnUnregister()
 {
 	Super::OnUnregister();
 
+	// Save off the Heightfields for potential re-use later, because the original cooked data was deleted in OnRegister.
+	// These heightfields are only used if this component gets re-registered before being destroyed.
+	if (HeightfieldRef)
+	{
+		LocalHeightfieldGeometryRef = HeightfieldRef->HeightfieldGeometry;
+		LocalHeightfieldSimpleGeometryRef = HeightfieldRef->HeightfieldSimpleGeometry;
+	}
+
 	// The physics object was destroyed in Super::OnUnregister. However we must
 	// extend the lifetime of the collision until the enqueued destroy command
 	// if processed on the physics thread, otherwise we may get Destroyed before
@@ -867,11 +869,28 @@ void ULandscapeHeightfieldCollisionComponent::CreateCollisionObject()
 
 			if (CookedCollisionData.IsEmpty())
 			{
-				if (CookedCollisionDataWasDeleted)
+				if (LocalHeightfieldGeometryRef.IsValid())
 				{
-					// only complain if we actually deleted the data.. otherwise it may have been intentional?
-					UE_LOG(LogLandscape, Warning, TEXT("Tried to create collision for component '%s', but the collision data was deleted!  To avoid set: landscape.Collision.DeleteCookedDataAfterUse 0"), *GetName());
+					// create heightfield ref from the local heightfield cached copy
+					HeightfieldRef = GSharedHeightfieldRefs.Add(HeightfieldGuid, new FHeightfieldGeometryRef(HeightfieldGuid));
+
+					HeightfieldRef->HeightfieldGeometry = MoveTemp(LocalHeightfieldGeometryRef);
+					if (LocalHeightfieldSimpleGeometryRef.IsValid())
+					{
+						HeightfieldRef->HeightfieldSimpleGeometry = MoveTemp(LocalHeightfieldSimpleGeometryRef);
+					}
+				}
+				else
+				{
+					if (bCookedCollisionDataWasDeleted)
+					{
+						// only complain if we actually deleted the data.. otherwise it may have been intentional
+						UE_LOG(LogLandscape, Warning, TEXT("Tried to create heightfield collision for component '%s', but the collision data was deleted!"), *GetName());
+					}
+					return;
 				}
+				
+				// Fallthrough to the shared register materials code below
 			}
 			else
 			{
@@ -890,49 +909,46 @@ void ULandscapeHeightfieldCollisionComponent::CreateCollisionObject()
 						Ar << HeightfieldRef->HeightfieldSimpleGeometry;
 					}
 				}
+			}
 
-				// Register materials
-				for(UPhysicalMaterial* PhysicalMaterial : CookedPhysicalMaterials)
+			// Register materials
+			for(UPhysicalMaterial* PhysicalMaterial : CookedPhysicalMaterials)
+			{
+				//TODO: Figure out why we are getting into a state like this (PhysicalMaterial == nullptr) in the first place. Potentially a loading issue
+				if (PhysicalMaterial)
 				{
-					//TODO: Figure out why we are getting into a state like this (PhysicalMaterial == nullptr) in the first place. Potentially a loading issue
-					if (PhysicalMaterial)
-					{
-						//todo: total hack until we get landscape fully converted to chaos
-						HeightfieldRef->UsedChaosMaterials.Add(PhysicalMaterial->GetPhysicsMaterial());
-					}
+					//todo: total hack until we get landscape fully converted to chaos
+					HeightfieldRef->UsedChaosMaterials.Add(PhysicalMaterial->GetPhysicsMaterial());
 				}
+			}
 
-				// Release cooked collison data
-				// In cooked builds created collision object will never be deleted while component is alive, so we don't need this data anymore
-				if(FPlatformProperties::RequiresCookedData() || World->IsGameWorld())
-				{
-					if (CVarLandscapeCollisionDeleteCookedDataAfterUse.GetValueOnGameThread())
-					{
-						CookedCollisionData.Empty();
-						CookedCollisionDataWasDeleted = true;
-					}
-				}
+			// Release cooked collison data
+			// In cooked builds created collision object will never be deleted while component is alive, so we don't need this data anymore
+			if(FPlatformProperties::RequiresCookedData() || World->IsGameWorld())
+			{
+				CookedCollisionData.Empty();
+				bCookedCollisionDataWasDeleted = true;
+			}
 
 #if WITH_EDITOR
-				// Create heightfield for the landscape editor (no holes in it)
-				if(bNeedsEditorHeightField)
+			// Create heightfield for the landscape editor (no holes in it)
+			if(bNeedsEditorHeightField)
+			{
+				TArray<UPhysicalMaterial*> CookedMaterialsEd;
+				if(CookCollisionData(PhysicsFormatName, true, bCheckDDC, CookedCollisionDataEd, CookedMaterialsEd))
 				{
-					TArray<UPhysicalMaterial*> CookedMaterialsEd;
-					if(CookCollisionData(PhysicsFormatName, true, bCheckDDC, CookedCollisionDataEd, CookedMaterialsEd))
-					{
-						FMemoryReader Reader(CookedCollisionDataEd);
-						Chaos::FChaosArchive Ar(Reader);
+					FMemoryReader Reader(CookedCollisionDataEd);
+					Chaos::FChaosArchive Ar(Reader);
 
-						// Don't actually care about this but need to strip it out of the data
-						bool bContainsSimple = false;
-						Ar << bContainsSimple;
-						Ar << HeightfieldRef->EditorHeightfieldGeometry;
+					// Don't actually care about this but need to strip it out of the data
+					bool bContainsSimple = false;
+					Ar << bContainsSimple;
+					Ar << HeightfieldRef->EditorHeightfieldGeometry;
 
-						CookedCollisionDataEd.Empty();
-					}
+					CookedCollisionDataEd.Empty();
 				}
-#endif //WITH_EDITOR
 			}
+#endif //WITH_EDITOR
 		}
 	}
 }
@@ -1560,10 +1576,10 @@ void ULandscapeMeshCollisionComponent::CreateCollisionObject()
 
 			if (CookedCollisionData.IsEmpty())
 			{
-				if (CookedCollisionDataWasDeleted)
+				if (bCookedCollisionDataWasDeleted)
 				{
 					// only complain if we actually deleted the data.. otherwise it may have been intentional?
-					UE_LOG(LogLandscape, Warning, TEXT("Tried to create collision for component '%s', but the collision data was deleted!  To avoid set: landscape.Collision.DeleteCookedDataAfterUse 0"), *GetName());
+					UE_LOG(LogLandscape, Warning, TEXT("Tried to create mesh collision for component '%s', but the collision data was deleted!"), *GetName());
 				}
 			}
 			else
@@ -1584,11 +1600,8 @@ void ULandscapeMeshCollisionComponent::CreateCollisionObject()
 				// In cooked builds created collision object will never be deleted while component is alive, so we don't need this data anymore
 				if (FPlatformProperties::RequiresCookedData() || GetWorld()->IsGameWorld())
 				{
-					if (CVarLandscapeCollisionDeleteCookedDataAfterUse.GetValueOnGameThread())
-					{
-						CookedCollisionData.Empty();
-						CookedCollisionDataWasDeleted = true;
-					}
+					CookedCollisionData.Empty();
+					bCookedCollisionDataWasDeleted = true;
 				}
 
 #if WITH_EDITOR
