Added accessor to physical material on landscape proxy and landscale heightfield component similar to the GetHeight call

diff --git a/Engine/Source/Runtime/Experimental/Chaos/Private/Chaos/HeightField.cpp b/Engine/Source/Runtime/Experimental/Chaos/Private/Chaos/HeightField.cpp
--- a/Engine/Source/Runtime/Experimental/Chaos/Private/Chaos/HeightField.cpp
+++ b/Engine/Source/Runtime/Experimental/Chaos/Private/Chaos/HeightField.cpp
@@ -876,6 +876,15 @@ namespace Chaos
 		return TNumericLimits<uint8>::Max();
 	}
 
+	uint8 Chaos::FHeightField::GetMaterialIndexAt(const FVec2& InGridLocationLocal) const
+	{
+		const FVec2 ClampledGridLocationLocal = FlatGrid.Clamp(InGridLocationLocal);
+		TVec2<int32> CellCoord = FlatGrid.Cell(ClampledGridLocationLocal);
+
+		const int32 Index = CellCoord[1] * (GeomData.NumCols - 1) + CellCoord[0];
+		return GetMaterialIndex(Index);
+	}
+
 	uint8 Chaos::FHeightField::GetMaterialIndex(int32 InX, int32 InY) const
 	{
 		const int32 Index = InY * (GeomData.NumCols - 1) + InX;
diff --git a/Engine/Source/Runtime/Experimental/Chaos/Public/Chaos/HeightField.h b/Engine/Source/Runtime/Experimental/Chaos/Public/Chaos/HeightField.h
--- a/Engine/Source/Runtime/Experimental/Chaos/Public/Chaos/HeightField.h
+++ b/Engine/Source/Runtime/Experimental/Chaos/Public/Chaos/HeightField.h
@@ -54,6 +54,7 @@ namespace Chaos
 		CHAOS_API bool IsHole(int32 InCellX, int32 InCellY) const;
 		CHAOS_API FVec3 GetNormalAt(const FVec2& InGridLocationLocal) const;
 		CHAOS_API FReal GetHeightAt(const FVec2& InGridLocationLocal) const;
+		CHAOS_API uint8 GetMaterialIndexAt(const FVec2& InGridLocationLocal) const;
 
 		int32 GetNumRows() const { return GeomData.NumRows; }
 		int32 GetNumCols() const { return GeomData.NumCols; }
diff --git a/Engine/Source/Runtime/Landscape/Classes/LandscapeHeightfieldCollisionComponent.h b/Engine/Source/Runtime/Landscape/Classes/LandscapeHeightfieldCollisionComponent.h
--- a/Engine/Source/Runtime/Landscape/Classes/LandscapeHeightfieldCollisionComponent.h
+++ b/Engine/Source/Runtime/Landscape/Classes/LandscapeHeightfieldCollisionComponent.h
@@ -328,7 +328,8 @@ public:
 	void SetRenderComponent(ULandscapeComponent* InRenderComponent) { RenderComponentRef = InRenderComponent; }
 
 public:
-	TOptional<float> GetHeight(float X, float Y, EHeightfieldSource HeightFieldSource);
+	LANDSCAPE_API TOptional<float> GetHeight(float X, float Y, EHeightfieldSource HeightFieldSource);
+	LANDSCAPE_API UPhysicalMaterial* GetPhysicalMaterial(float X, float Y, EHeightfieldSource HeightFieldSource);
 
 	/**
 	 * Populates a supplied array with the heights from the heightfield.  Samples are placed
diff --git a/Engine/Source/Runtime/Landscape/Classes/LandscapeProxy.h b/Engine/Source/Runtime/Landscape/Classes/LandscapeProxy.h
--- a/Engine/Source/Runtime/Landscape/Classes/LandscapeProxy.h
+++ b/Engine/Source/Runtime/Landscape/Classes/LandscapeProxy.h
@@ -935,6 +935,7 @@ public:
 	void EnableNaniteComponents(bool bInNaniteActive);
 
 	LANDSCAPE_API TOptional<float> GetHeightAtLocation(FVector Location, EHeightfieldSource HeightFieldSource = EHeightfieldSource::Complex) const;
+	LANDSCAPE_API UPhysicalMaterial* GetPhysicalMaterialAtLocation(FVector Location, EHeightfieldSource HeightFieldSOurce = EHeightfieldSource::Complex) const;
 
 	/** Fills an array with height values **/
 	LANDSCAPE_API void GetHeightValues(int32& SizeX, int32& SizeY, TArray<float>& ArrayValue) const;
diff --git a/Engine/Source/Runtime/Landscape/Private/LandscapeCollision.cpp b/Engine/Source/Runtime/Landscape/Private/LandscapeCollision.cpp
--- a/Engine/Source/Runtime/Landscape/Private/LandscapeCollision.cpp
+++ b/Engine/Source/Runtime/Landscape/Private/LandscapeCollision.cpp
@@ -3020,6 +3020,50 @@ TOptional<float> ULandscapeHeightfieldCollisionComponent::GetHeight(float X, flo
 	return Height;
 }
 
+UPhysicalMaterial* ULandscapeHeightfieldCollisionComponent::GetPhysicalMaterial(float X, float Y, EHeightfieldSource HeightFieldSource)
+{
+	UPhysicalMaterial* PhysicalMaterial = nullptr;
+
+	if (!IsValidRef(HeightfieldRef))
+	{
+		return PhysicalMaterial;
+	}
+
+	Chaos::FHeightField* HeightField = nullptr;
+
+	switch (HeightFieldSource)
+	{
+	case EHeightfieldSource::None:
+		break;
+	case EHeightfieldSource::Simple:
+		HeightField = HeightfieldRef->HeightfieldSimpleGeometry.GetReference();
+		break;
+	case EHeightfieldSource::Complex:
+		HeightField = HeightfieldRef->HeightfieldGeometry.GetReference();
+		break;
+#if WITH_EDITORONLY_DATA		
+	case EHeightfieldSource::Editor:
+		HeightField = HeightfieldRef->EditorHeightfieldGeometry.GetReference();
+		break;
+#endif 
+	}
+
+	if (HeightField)
+	{
+		const uint8 MaterialIndex = HeightField->GetMaterialIndexAt({ X, Y });
+		if (MaterialIndex != TNumericLimits<uint8>::Max() && HeightfieldRef->UsedChaosMaterials.IsValidIndex(MaterialIndex))
+		{
+			Chaos::FMaterialHandle MaterialHandle = HeightfieldRef->UsedChaosMaterials[MaterialIndex];
+			if (Chaos::FChaosPhysicsMaterial* ChaosMaterial = MaterialHandle.Get())
+			{
+				PhysicalMaterial = FChaosUserData::Get<UPhysicalMaterial>(ChaosMaterial->UserData);
+			}
+		}
+	}
+
+	return PhysicalMaterial;
+}
+
 struct FHeightFieldAccessor
 {
 	FHeightFieldAccessor(const ULandscapeHeightfieldCollisionComponent::FHeightfieldGeometryRef& InGeometryRef)
@@ -3154,6 +3198,24 @@ TOptional<float> ALandscapeProxy::GetHeightAtLocation(FVector Location, EHeightf
 	return Height;
 }
 
+UPhysicalMaterial* ALandscapeProxy::GetPhysicalMaterialAtLocation(FVector Location, EHeightfieldSource HeightFieldSource) const
+{
+	UPhysicalMaterial* PhysicalMaterial = nullptr;
+	if (ULandscapeInfo* Info = GetLandscapeInfo())
+	{
+		const FVector ActorSpaceLocation = LandscapeActorToWorld().InverseTransformPosition(Location);
+		const FIntPoint Key = FIntPoint(FMath::FloorToInt32(ActorSpaceLocation.X / ComponentSizeQuads), FMath::FloorToInt32(ActorSpaceLocation.Y / ComponentSizeQuads));
+		ULandscapeHeightfieldCollisionComponent* Component = Info->XYtoCollisionComponentMap.FindRef(Key);
+		if (Component)
+		{
+			const FVector ComponentSpaceLocation = Component->GetComponentToWorld().InverseTransformPosition(Location);
+			PhysicalMaterial = Component->GetPhysicalMaterial(static_cast<float>(ComponentSpaceLocation.X), static_cast<float>(ComponentSpaceLocation.Y), HeightFieldSource);
+		}
+	}
+
+	return PhysicalMaterial;
+}
+
 void ALandscapeProxy::GetHeightValues(int32& SizeX, int32& SizeY, TArray<float> &ArrayValues) const
 {			
 	SizeX = 0;
