Fixed logic around Landscape fixed grid VF

diff --git a/Engine/Source/Runtime/Landscape/Private/LandscapeRender.cpp b/Engine/Source/Runtime/Landscape/Private/LandscapeRender.cpp
--- a/Engine/Source/Runtime/Landscape/Private/LandscapeRender.cpp
+++ b/Engine/Source/Runtime/Landscape/Private/LandscapeRender.cpp
@@ -3571,9 +3571,16 @@ public:
 				static const FName RayTracingDynamicGeometryConverterCS = FName(TEXT("FRayTracingDynamicGeometryConverterCS"));
 				const bool bIsRayTracingShaderType = ShaderType->GetFName() == RayTracingDynamicGeometryConverterCS;
 
-				if ((bIsRayTracingShaderType || !bIsShaderTypeUsingFixedGrid) && VertexFactoryType->SupportsLandscape())
+				if (VertexFactoryType->SupportsLandscape())
 				{
-					return FMaterialResource::ShouldCache(Platform, ShaderType, VertexFactoryType);
+					if (&FLandscapeFixedGridVertexFactory::StaticType == VertexFactoryType)
+					{
+						return (bIsRayTracingShaderType || bIsShaderTypeUsingFixedGrid) && FMaterialResource::ShouldCache(Platform, ShaderType, VertexFactoryType);
+					}
+					else
+					{
+						return (bIsRayTracingShaderType || !bIsShaderTypeUsingFixedGrid) && FMaterialResource::ShouldCache(Platform, ShaderType, VertexFactoryType);
+					}
 				}
 			}
 		}
