Fix for crash in landscape post edit undo.

diff --git a/Engine/Source/Runtime/Landscape/Private/LandscapeEdit.cpp b/Engine/Source/Runtime/Landscape/Private/LandscapeEdit.cpp
--- a/Engine/Source/Runtime/Landscape/Private/LandscapeEdit.cpp
+++ b/Engine/Source/Runtime/Landscape/Private/LandscapeEdit.cpp
@@ -680,7 +680,8 @@ void ULandscapeComponent::PostEditUndo()
 	{
 		// PreEditUndo of the LandscapeStreamingProxy resets it's pointer to the top level landscape so this ensures it's fixed up prior
 		// to requesting HeightMap or WeightMap updates
-		if (ULandscapeInfo* LandscapeInfo = Proxy->GetLandscapeInfo(); !LandscapeInfo->IsRegistered(Proxy))
+		ULandscapeInfo* LandscapeInfo = Proxy->GetLandscapeInfo();
+		if (LandscapeInfo && !LandscapeInfo->IsRegistered(Proxy))
 		{
 			LandscapeInfo->RegisterActor(Proxy, true);
 		}
