Fixes for landscape render section registration crashes.

1) Added range check to Set/GetSectionInfo -- array bounds checks already exist but in shipping builds the checks go away, possibly returning a garbage pointer.  This is the likely immediate cause of the crash.  Why it is access out of bounds though?

2) Move SectionInfo array resize to Register. The array was originally resized in Create/DestroyResources, however, when compacting the array in Destroy, it can discard any region that is not Registered (even if it is still Created).  This means it is possible you can get a Register call where the array location has been compacted away, causing you to access out-of-bounds.  By moving the resize to Register we can guarantee the location always exists in bounds when we try to access it.

3) Fix HLOD registration when using LOD Groups - previously it would always use the Base coordinates and be mismatched with landscape.  Now it records positions for the proxies and uses the projection based registration, so it calculates the same coordinates as regular Landscape sections.  I tested this with HLODs built on a rotated landscape to verify that the coordinates match up even under transforms.

diff --git a/Engine/Source/Runtime/Landscape/Classes/LandscapeMeshProxyComponent.h b/Engine/Source/Runtime/Landscape/Classes/LandscapeMeshProxyComponent.h
--- a/Engine/Source/Runtime/Landscape/Classes/LandscapeMeshProxyComponent.h
+++ b/Engine/Source/Runtime/Landscape/Classes/LandscapeMeshProxyComponent.h
@@ -25,11 +25,24 @@ private:
 	UPROPERTY()
 	FGuid LandscapeGuid;
 
-	/* The components this proxy was generated for */
+	/* The section coordinates of the landscape components that this proxy was generated for.  Used to register with LandscapeRender when LODGroupKey == 0 */
 	UPROPERTY()
 	TArray<FIntPoint> ProxyComponentBases;
 
-	/* LOD level proxy was generated for */
+	/* The center of the landscape components that this proxy was generated for, in local component space.  Used to register with LandscapeRender when LODGroupKey != 0 */
+	UPROPERTY()
+	TArray<FVector> ProxyComponentCentersObjectSpace;
+
+	UPROPERTY()
+	FVector ComponentXVectorObjectSpace;
+
+	UPROPERTY()
+	FVector ComponentYVectorObjectSpace;
+
+	UPROPERTY()
+	int32 ComponentResolution;
+
+	/* LOD level this proxy was generated for */
 	UPROPERTY()
 	int8 ProxyLOD;
 
diff --git a/Engine/Source/Runtime/Landscape/Private/Landscape.cpp b/Engine/Source/Runtime/Landscape/Private/Landscape.cpp
--- a/Engine/Source/Runtime/Landscape/Private/Landscape.cpp
+++ b/Engine/Source/Runtime/Landscape/Private/Landscape.cpp
@@ -5679,11 +5679,37 @@ void ULandscapeMeshProxyComponent::InitializeForLandscape(ALandscapeProxy* Lands
 	LandscapeGuid = Landscape->GetLandscapeGuid();
 	LODGroupKey = Landscape->LODGroupKey;
 
+	FTransform WorldToLocal = GetComponentTransform().Inverse();
+
+	bool bFirst = true;
 	for (ULandscapeComponent* Component : Landscape->LandscapeComponents)
 	{
 		if (Component)
 		{
+			const FTransform& ComponentLocalToWorld = Component->GetComponentTransform();
+
+			if (bFirst)
+			{
+				bFirst = false;
+				ComponentResolution = Component->ComponentSizeQuads + 1;
+				FVector ComponentXVectorWorldSpace = ComponentLocalToWorld.TransformVector(FVector::XAxisVector) * ComponentResolution;
+				FVector ComponentYVectorWorldSpace = ComponentLocalToWorld.TransformVector(FVector::YAxisVector) * ComponentResolution;
+				ComponentXVectorObjectSpace = WorldToLocal.TransformVector(ComponentXVectorWorldSpace);
+				ComponentYVectorObjectSpace = WorldToLocal.TransformVector(ComponentYVectorWorldSpace);
+			}
+			else
+			{
+				// assume it's the same resolution and orientation as the first component... (we only record one resolution and orientation)
+			}
+
+			// record the component coordinate
 			ProxyComponentBases.Add(Component->GetSectionBase() / Component->ComponentSizeQuads);
+			
+			// record the component center position (in the space of the ULandscapeMeshProxyComponent)
+			FBoxSphereBounds ComponentLocalBounds = Component->CalcBounds(FTransform::Identity);
+			FVector ComponentOriginWorld = ComponentLocalToWorld.TransformPosition(ComponentLocalBounds.Origin);
+			FVector LocalOrigin = WorldToLocal.TransformPosition(ComponentOriginWorld);
+			ProxyComponentCentersObjectSpace.Add(LocalOrigin);
 		}
 	}
 
diff --git a/Engine/Source/Runtime/Landscape/Private/LandscapeRender.cpp b/Engine/Source/Runtime/Landscape/Private/LandscapeRender.cpp
--- a/Engine/Source/Runtime/Landscape/Private/LandscapeRender.cpp
+++ b/Engine/Source/Runtime/Landscape/Private/LandscapeRender.cpp
@@ -454,6 +454,7 @@ FLandscapeRenderSystem::FLandscapeRenderSystem()
 	: Min(MAX_int32, MAX_int32)
 	, Size(EForceInit::ForceInitToZero)
 	, ReferenceCount(0)
+	, RegisteredCount(0)
 	, ForcedLODOverride(-1)
 	, SectionsRemovedSinceLastCompact(0)
 {
@@ -500,6 +501,7 @@ void FLandscapeRenderSystem::DestroyResources(FLandscapeSectionInfo* SectionInfo
 
 	if (LandscapeRenderSystem->ReferenceCount == 0)
 	{
+		check(LandscapeRenderSystem->RegisteredCount == 0);
 		delete LandscapeRenderSystem;
 		LandscapeRenderSystems.Remove(SectionInfo->LandscapeKey);
 	}
@@ -517,54 +519,43 @@ void FLandscapeRenderSystem::CreateResources_Internal(FRHICommandListBase& RHICm
 	if (SectionInfo->LODGroupKey != 0)
 	{
 		// Record and check settings (resolution, scale and orientation) that need to match across landscapes in an LOD Group
+		FVector SectionCenterWorldSpace, SectionXVector, SectionYVector;
+		SectionInfo->GetSectionCenterAndVectors(SectionCenterWorldSpace, SectionXVector, SectionYVector);
 		int32 SectionComponentResolution = SectionInfo->GetComponentResolution();
-		if (SectionComponentResolution > 0)
-		{
-			FBoxSphereBounds LocalBounds;
-			FMatrix LocalToWorld;
-			SectionInfo->GetSectionBoundsAndLocalToWorld(LocalBounds, LocalToWorld);
-
-			FVector SectionCenterWorldSpace = LocalToWorld.TransformPosition(LocalBounds.Origin);
-			FVector SectionXVector = LocalToWorld.TransformVector(FVector::XAxisVector) * SectionComponentResolution;
-			FVector SectionYVector = LocalToWorld.TransformVector(FVector::YAxisVector) * SectionComponentResolution;
 
-			if (ComponentResolution < 0)
+		if (ComponentResolution < 0)
+		{
+			// the first time we register a section, set up our RenderCoord grid so the component is located at the origin
+			ComponentResolution = SectionComponentResolution;
+			ComponentOrigin = SectionCenterWorldSpace;
+			ComponentXVector = SectionXVector;
+			ComponentYVector = SectionYVector;
+		}
+		else
+		{
+			// validate each section has a matching resolution, scale and orientation
+			bool bResolutionMatches = (ComponentResolution == SectionComponentResolution);
+			bool bXVectorMatches = (SectionXVector - ComponentXVector).IsNearlyZero();
+			bool bYVectorMatches = (SectionYVector - ComponentYVector).IsNearlyZero();
+			if (!(bResolutionMatches && bXVectorMatches && bYVectorMatches))
 			{
-				// the first component with a resolution -- record its information
-				ComponentResolution = SectionComponentResolution;
-				ComponentOrigin = SectionCenterWorldSpace;
-				ComponentXVector = SectionXVector;
-				ComponentYVector = SectionYVector;
+				UE_LOG(LogLandscape, Warning, TEXT("Landscapes in LOD Group %d do not have matching resolution (%d == %d), scale (%f == %f, %f == %f) and/or rotation; geometry seam artifacts may appear. If using HLOD, it may need to be updated to reflect changed resolutions or transforms."),
+					SectionInfo->LODGroupKey,
+					ComponentResolution, SectionComponentResolution,
+					ComponentXVector.Length(), SectionXVector.Length(),
+					ComponentYVector.Length(), SectionYVector.Length());
 			}
-			else
-			{
-				// validate matching resolution
-				bool bResolutionMatches = (ComponentResolution == SectionComponentResolution);
-				bool bXVectorMatches = (SectionXVector - ComponentXVector).IsNearlyZero();
-				bool bYVectorMatches = (SectionYVector - ComponentYVector).IsNearlyZero();
-				if (!(bResolutionMatches && bXVectorMatches && bYVectorMatches))
-				{
-					UE_LOG(LogLandscape, Warning, TEXT("Landscapes in LOD Group %d do not have matching resolution (%d == %d), scale (%f == %f, %f == %f) and/or rotation; geometry seam artifacts may appear."),
-						SectionInfo->LODGroupKey,
-						ComponentResolution, SectionComponentResolution,
-						ComponentXVector.Length(), SectionXVector.Length(),
-						ComponentYVector.Length(), SectionYVector.Length());
-				}
-			}
-			
-			// project onto the Component X/Y plane to calculate the render coordinates
-			FVector Delta = SectionCenterWorldSpace - ComponentOrigin;
-			SectionInfo->RenderCoord.X = FMath::RoundToInt32(Delta.Dot(ComponentXVector) / ComponentXVector.SquaredLength());
-			SectionInfo->RenderCoord.Y = FMath::RoundToInt32(Delta.Dot(ComponentYVector) / ComponentXVector.SquaredLength());
 		}
+			
+		// project onto the Component X/Y plane to calculate the render coordinates
+		FVector Delta = SectionCenterWorldSpace - ComponentOrigin;
+		SectionInfo->RenderCoord.X = FMath::RoundToInt32(Delta.Dot(ComponentXVector) / ComponentXVector.SquaredLength());
+		SectionInfo->RenderCoord.Y = FMath::RoundToInt32(Delta.Dot(ComponentYVector) / ComponentXVector.SquaredLength());
 	}
 
 	// we changed the RenderCoord, need to update the uniform buffer
 	SectionInfo->OnRenderCoordsChanged(RHICmdList);
 
-	check(SectionInfo->RenderCoord.X > INT32_MIN);
-	ResizeToInclude(SectionInfo->RenderCoord);
-
 	ReferenceCount++;
 }
 
@@ -594,6 +585,7 @@ void FLandscapeRenderSystem::DestroyResources_Internal(FLandscapeSectionInfo* Se
 void FLandscapeRenderSystem::RegisterSection(FLandscapeSectionInfo* SectionInfo)
 {
 	check(SectionInfo != nullptr);
+	check(SectionInfo->bResourcesCreated);
 	check(!SectionInfo->bRegistered);
 
 	// With HLODs, it's possible to have multiple loaded sections representing the same
@@ -604,9 +596,11 @@ void FLandscapeRenderSystem::RegisterSection(FLandscapeSectionInfo* SectionInfo)
 
 	FLandscapeRenderSystem*& LandscapeRenderSystem = LandscapeRenderSystems.FindChecked(SectionInfo->LandscapeKey);
 
+	check(SectionInfo->RenderCoord.X > INT32_MIN);
 	FLandscapeSectionInfo* ExistingSection = LandscapeRenderSystem->GetSectionInfo(SectionInfo->RenderCoord);
 	if (ExistingSection == nullptr)
 	{
+		LandscapeRenderSystem->ResizeToInclude(SectionInfo->RenderCoord);
 		LandscapeRenderSystem->SetSectionInfo(SectionInfo->RenderCoord, SectionInfo);
 	}
 	else
@@ -632,18 +626,21 @@ void FLandscapeRenderSystem::RegisterSection(FLandscapeSectionInfo* SectionInfo)
 		}
 		else if (CurrentSection == ExistingSection)
 		{
-			// Set as head
+			// Set as head (SectionInfo was inserted before the previous head in the loop above)
 			LandscapeRenderSystem->SetSectionInfo(SectionInfo->RenderCoord, SectionInfo);
 		}
 	}
 
 	SectionInfo->bRegistered = true;
+	LandscapeRenderSystem->RegisteredCount++;
 }
 
 void FLandscapeRenderSystem::UnregisterSection(FLandscapeSectionInfo* SectionInfo)
 {
 	check(SectionInfo != nullptr);
+	check(SectionInfo->bResourcesCreated);
 
+	// Sections may be unregistered multiple times
 	if (SectionInfo->bRegistered)
 	{
 		FLandscapeRenderSystem* LandscapeRenderSystem = LandscapeRenderSystems.FindChecked(SectionInfo->LandscapeKey);
@@ -655,8 +652,8 @@ void FLandscapeRenderSystem::UnregisterSection(FLandscapeSectionInfo* SectionInf
 		}
 
 		SectionInfo->Unlink();
-
 		SectionInfo->bRegistered = false;
+		LandscapeRenderSystem->RegisteredCount--;
 	}
 }
 
@@ -705,11 +702,17 @@ void FLandscapeRenderSystem::ResizeAndMoveTo(FIntPoint NewMin, FIntPoint NewMax)
 						int32 NewLinearIndex = NewYBase + NewX;
 						NewSectionLODBiases[NewLinearIndex] = SectionLODBiases[OldLinearIndex];
 						NewSectionInfos[NewLinearIndex] = SectionInfos[OldLinearIndex];
+
+						// we null this out in order to check below that we have moved everything.
+						SectionInfos[OldLinearIndex] = nullptr;
 					}
 				}
 			}
 		}
 
+		// check that we have moved everything out of the old section infos, to ensure that this resize doesn't lose data
+		check(!AnySectionsInRangeInclusive(Min, Min + Size - 1));
+
 		Min = NewMin;
 		Size = NewSize;
 		SectionLODBiases = MoveTemp(NewSectionLODBiases);
@@ -4161,6 +4164,17 @@ void FLandscapeComponentSceneProxy::GetSectionBoundsAndLocalToWorld(FBoxSphereBo
 	OutLocalToWorld = GetLocalToWorld();
 }
 
+void FLandscapeComponentSceneProxy::GetSectionCenterAndVectors(FVector& OutSectionCenterWorldSpace, FVector& OutSectionXVectorWorldSpace, FVector& OutSectionYVectorWorldSpace) const
+{
+	FBoxSphereBounds ComponentLocalBounds = GetLocalBounds();
+	FMatrix ComponentLocalToWorld = GetLocalToWorld();
+	int32 ComponentResolution = GetComponentResolution();
+
+	OutSectionCenterWorldSpace = ComponentLocalToWorld.TransformPosition(ComponentLocalBounds.Origin);
+	OutSectionXVectorWorldSpace = ComponentLocalToWorld.TransformVector(FVector::XAxisVector) * ComponentResolution;
+	OutSectionYVectorWorldSpace = ComponentLocalToWorld.TransformVector(FVector::YAxisVector) * ComponentResolution;
+}
+
 //
 // FLandscapeSectionInfo
 //
@@ -4180,9 +4194,13 @@ FLandscapeSectionInfo::FLandscapeSectionInfo(const UWorld* InWorld, const FGuid&
 class FLandscapeProxySectionInfo : public FLandscapeSectionInfo
 {
 public:
-	FLandscapeProxySectionInfo(const UWorld* InWorld, const FGuid& InLandscapeGuid, const FIntPoint& InComponentBase, int8 InProxyLOD, uint32 LODGroupKey)
+	FLandscapeProxySectionInfo(const UWorld* InWorld, const FGuid& InLandscapeGuid, const FIntPoint& InComponentBase, const FVector& InComponentCenterLocalSpace, const FVector& InComponentXVectorLocalSpace, const FVector& InComponentYVectorLocalSpace, const FTransform& LocalToWorld, int32 SectionComponentResolution, int8 InProxyLOD, uint32 LODGroupKey)
 		: FLandscapeSectionInfo(InWorld, InLandscapeGuid, InComponentBase, LODGroupKey)
 		, ProxyLOD(InProxyLOD)
+		, ComponentResolution(SectionComponentResolution)
+		, CenterWorldSpace(LocalToWorld.TransformPosition(InComponentCenterLocalSpace))
+		, XVectorWorldSpace(LocalToWorld.TransformVector(InComponentXVectorLocalSpace))
+		, YVectorWorldSpace(LocalToWorld.TransformVector(InComponentXVectorLocalSpace))
 	{
 	}
 
@@ -4201,26 +4219,44 @@ public:
 		return ProxyLOD;
 	}
 	
-	virtual void GetSectionBoundsAndLocalToWorld(FBoxSphereBounds& LocalBounds, FMatrix& LocalToWorld) const
+	virtual void GetSectionBoundsAndLocalToWorld(FBoxSphereBounds& LocalBounds, FMatrix& LocalToWorld) const override
 	{
 		LocalBounds = FBoxSphereBounds(ForceInit);
 		LocalToWorld = FMatrix::Identity;
 	}
 
+	virtual void GetSectionCenterAndVectors(FVector& OutSectionCenterWorldSpace, FVector& OutSectionXVectorWorldSpace, FVector& OutSectionYVectorWorldSpace) const override
+	{
+		OutSectionCenterWorldSpace = CenterWorldSpace;
+		OutSectionXVectorWorldSpace = XVectorWorldSpace;
+		OutSectionYVectorWorldSpace = YVectorWorldSpace;
+	}
+
+	virtual int32 GetComponentResolution() const override
+	{
+		return ComponentResolution;
+	}
+
 	virtual void OnRenderCoordsChanged(FRHICommandListBase& RHICmdList)
 	{
 	}
 
 private:
 	int8 ProxyLOD;
+	int32 ComponentResolution;
+	FVector CenterWorldSpace;
+	FVector XVectorWorldSpace;
+	FVector YVectorWorldSpace;
 };
 
 //
 // FLandscapeMeshProxySceneProxy
 //
-FLandscapeMeshProxySceneProxy::FLandscapeMeshProxySceneProxy(UStaticMeshComponent* InComponent, const FGuid& InLandscapeGuid, const TArray<FIntPoint>& InProxySectionsBases, int8 InProxyLOD, uint32 InLODGroupKey)
+FLandscapeMeshProxySceneProxy::FLandscapeMeshProxySceneProxy(UStaticMeshComponent* InComponent, const FGuid& InLandscapeGuid, const TArray<FIntPoint>& InProxySectionsBases, const TArray<FVector>& InProxySectionsCentersLocalSpace, const FVector& InComponentXVector, const FVector& InComponentYVector, const FTransform& LocalToWorld, int32 ComponentResolution, int8 InProxyLOD, uint32 InLODGroupKey)
 	: FStaticMeshSceneProxy(InComponent, false)
 {
+	check(InLODGroupKey == 0 || ComponentResolution != 0);	// if using LODGroupKey, we require HLODs to be rebuilt to capture the component resolution and section transform data (position / orientation)
+
 	VisibilityHelper.Init(InComponent, this);
 
 	if (VisibilityHelper.RequiresVisibleLevelToRender())
@@ -4229,9 +4265,14 @@ FLandscapeMeshProxySceneProxy::FLandscapeMeshProxySceneProxy(UStaticMeshComponen
 	}
 
 	ProxySectionsInfos.Empty(InProxySectionsBases.Num());
-	for (FIntPoint SectionBase : InProxySectionsBases)
+	
+	check(InProxySectionsBases.Num() == InProxySectionsCentersLocalSpace.Num() ||
+		  InProxySectionsCentersLocalSpace.Num() == 0);
+	for (int32 Index = 0; Index < InProxySectionsBases.Num(); Index++)
 	{
-		ProxySectionsInfos.Emplace(MakeUnique<FLandscapeProxySectionInfo>(InComponent->GetWorld(), InLandscapeGuid, SectionBase, InProxyLOD, InLODGroupKey));
+		FIntPoint SectionBase = InProxySectionsBases[Index];
+		FVector ProxySectionCenterLocalSpace = InProxySectionsCentersLocalSpace.IsValidIndex(Index) ? InProxySectionsCentersLocalSpace[Index] : FVector::Zero();
+		ProxySectionsInfos.Emplace(MakeUnique<FLandscapeProxySectionInfo>(InComponent->GetWorld(), InLandscapeGuid, SectionBase, ProxySectionCenterLocalSpace, InComponentXVector, InComponentYVector, LocalToWorld, ComponentResolution, InProxyLOD, InLODGroupKey));
 	}
 
 	// Force OnTransformChanged and {Create, Destroy}RenderThreadResources onto the render thread.
@@ -4319,7 +4360,15 @@ FPrimitiveSceneProxy* ULandscapeMeshProxyComponent::CreateSceneProxy()
 		return nullptr;
 	}
 
-	return new FLandscapeMeshProxySceneProxy(this, LandscapeGuid, ProxyComponentBases, ProxyLOD, LODGroupKey);
+	if ((LODGroupKey != 0) && (ComponentResolution == 0))
+	{
+		// this is an OLD HLOD not built with the new resolution/x/y/center information, which we need if LODGroups are used.
+		UE_LOG(LogLandscape, Warning, TEXT("HLOD for Landscape needs to be rebuilt!  Landscape Mesh Proxy Component '%s' is using LODGroups but does not have the position and orientation information necessary to register it with the Landscape renderer.  This HLOD will not render until it is fixed."),
+			*GetName());
+		return nullptr;
+	}
+
+	return new FLandscapeMeshProxySceneProxy(this, LandscapeGuid, ProxyComponentBases, ProxyComponentCentersObjectSpace, ComponentXVectorObjectSpace, ComponentYVectorObjectSpace, GetComponentTransform(), ComponentResolution, ProxyLOD, LODGroupKey);
 }
 
 class FLandscapeNaniteSceneProxy : public ::Nanite::FSceneProxy
diff --git a/Engine/Source/Runtime/Landscape/Public/LandscapeRender.h b/Engine/Source/Runtime/Landscape/Public/LandscapeRender.h
--- a/Engine/Source/Runtime/Landscape/Public/LandscapeRender.h
+++ b/Engine/Source/Runtime/Landscape/Public/LandscapeRender.h
@@ -406,6 +406,8 @@ public:
 	
 	virtual void GetSectionBoundsAndLocalToWorld(FBoxSphereBounds& LocalBounds, FMatrix& LocalToWorld) const = 0;
 
+	virtual void GetSectionCenterAndVectors(FVector& OutSectionCenterWorldSpace, FVector& OutSectionXVectorWorldSpace, FVector& OutSectionYVectorWorldSpace) const = 0;
+
 	/* return the resolution of a component, in vertices (-1 for any sections that are not grid based, i.e. mesh sections) */
 	virtual int32 GetComponentResolution() const { return -1; }
 
@@ -469,7 +471,8 @@ struct FLandscapeRenderSystem
 
 	TResourceArray<float> SectionLODBiases;
 	TArray<FLandscapeSectionInfo*> SectionInfos;
-	int32 ReferenceCount;
+	int32 ReferenceCount;			// number of sections with resources created
+	int32 RegisteredCount;			// number of sections registered
 
 	FBufferRHIRef SectionLODBiasBuffer;
 	FShaderResourceViewRHIRef SectionLODBiasSRV;
@@ -519,12 +522,19 @@ struct FLandscapeRenderSystem
 
 	void SetSectionInfo(FIntPoint InRenderCoord, FLandscapeSectionInfo* InSectionInfo)
 	{
-		SectionInfos[GetSectionLinearIndex(InRenderCoord)] = InSectionInfo;
+		if (IsValidCoord(InRenderCoord))
+		{
+			SectionInfos[GetSectionLinearIndex(InRenderCoord)] = InSectionInfo;
+		}
 	}
 
 	FLandscapeSectionInfo* GetSectionInfo(FIntPoint InRenderCoord)
 	{
-		return SectionInfos[GetSectionLinearIndex(InRenderCoord)];
+		if (IsValidCoord(InRenderCoord))
+		{
+			return SectionInfos[GetSectionLinearIndex(InRenderCoord)];
+		}
+		return nullptr;
 	}
 
 	float GetSectionLODValue(const FSceneView& SceneView, FIntPoint InRenderCoord) const
@@ -647,7 +657,7 @@ class FLandscapeMeshProxySceneProxy final : public FStaticMeshSceneProxy
 public:
 	SIZE_T GetTypeHash() const override;
 
-	FLandscapeMeshProxySceneProxy(UStaticMeshComponent* InComponent, const FGuid& InLandscapeGuid, const TArray<FIntPoint>& InProxySectionsBases, int8 InProxyLOD, uint32 InLODGroupKey);
+	FLandscapeMeshProxySceneProxy(UStaticMeshComponent* InComponent, const FGuid& InLandscapeGuid, const TArray<FIntPoint>& InProxySectionsBases, const TArray<FVector>& InProxySectionsCentersLocalSpace, const FVector& InComponentXVector, const FVector& InComponentYVector, const FTransform& LocalToWorld, int32 ComponentResolution, int8 InProxyLOD, uint32 InLODGroupKey);
 	virtual void CreateRenderThreadResources(FRHICommandListBase& RHICmdList) override;
 	virtual void DestroyRenderThreadResources() override;
 	virtual bool OnLevelAddedToWorld_RenderThread() override;
@@ -902,6 +912,8 @@ public:
 
 	LANDSCAPE_API virtual double ComputeSectionResolution() const override;
 	LANDSCAPE_API virtual void GetSectionBoundsAndLocalToWorld(FBoxSphereBounds& LocalBounds, FMatrix& LocalToWorld) const override;
+
+	virtual void GetSectionCenterAndVectors(FVector& OutSectionCenterWorldSpace, FVector& OutSectionXVectorWorldSpace, FVector& OutSectionYVectorWorldSpace) const override;
 };
 
 class FLandscapeDebugMaterialRenderProxy : public FMaterialRenderProxy
