Adds ALandscape override to ULandscapeInfo MarkDirty methods to circumvent codepath when ULandscapeInfo is not fully ready.

diff --git a/Engine/Source/Runtime/Landscape/Classes/LandscapeInfo.h b/Engine/Source/Runtime/Landscape/Classes/LandscapeInfo.h
--- a/Engine/Source/Runtime/Landscape/Classes/LandscapeInfo.h
+++ b/Engine/Source/Runtime/Landscape/Classes/LandscapeInfo.h
@@ -284,9 +284,10 @@ public:
 	 * @param InObject Object to mark dirty. Must be either a ALandscapeProxy or one of its inner objects
 	 * @param bInForceResave passed to true in order to bypass the bDirtyOnlyInMode and force the object into the list of modified packages. This ensures the object is marked as needing resaving in case we want to 
 	 *  help the user re-save its landscape actors (e.g. following a data fixup)
+	 * @param InLandscapeOverride if set, MarkObjectDirty will use this ALandscape instead of the local one in the method.
 	 * @return true if the object could not be marked dirty but has been added to the list of ModifiedPackages
 	 */
-	LANDSCAPE_API bool MarkObjectDirty(UObject* InObject, bool bInForceResave = false);
+	LANDSCAPE_API bool MarkObjectDirty(UObject* InObject, bool bInForceResave = false, const ALandscape* InLandscapeOverride = nullptr);
 #endif //WITH_EDITOR
 
 	/**
@@ -460,6 +461,6 @@ private:
 	void MoveSegment(ULandscapeSplineSegment* InSegment, TScriptInterface<ILandscapeSplineInterface> From, TScriptInterface<ILandscapeSplineInterface> To);
 	void MoveControlPoint(ULandscapeSplineControlPoint* InControlPoint, TScriptInterface<ILandscapeSplineInterface> From, TScriptInterface<ILandscapeSplineInterface> To);
 	bool UpdateLayerInfoMapInternal(ALandscapeProxy* Proxy, bool bInvalidate);
-	bool TryAddToModifiedPackages(UPackage* InPackage);
+	bool TryAddToModifiedPackages(UPackage* InPackage, const ALandscape* InLandscapeOverride = nullptr);
 #endif
 };
diff --git a/Engine/Source/Runtime/Landscape/Private/Landscape.cpp b/Engine/Source/Runtime/Landscape/Private/Landscape.cpp
--- a/Engine/Source/Runtime/Landscape/Private/Landscape.cpp
+++ b/Engine/Source/Runtime/Landscape/Private/Landscape.cpp
@@ -4380,8 +4380,11 @@ void ALandscapeProxy::FixupSharedData(ALandscape* Landscape, const bool bMapChec
 
 		if (bUpdated)
 		{
+			// In cases where LandscapeInfo is not fully ready yet, we forward the provided ALandscape. If ALandscape is available in LandscapeInfo, we let the object function naturally.
+			const ALandscape* LandscapeActor = (LandscapeInfo->LandscapeActor == nullptr) ? Landscape : nullptr;
+
 			// Force resave the proxy through the modified landscape system, so that the user can then use the Save Modified Landscapes menu and therefore manually trigger the re-save of all modified proxies. * /
-			bool bNeedsManualResave = LandscapeInfo->MarkObjectDirty(/*InObject = */this, /*bInForceResave = */true);
+			bool bNeedsManualResave = LandscapeInfo->MarkObjectDirty(/*InObject = */this, /*bInForceResave = */true, LandscapeActor);
 
 			if (bMapCheck && bNeedsManualResave)
 			{
@@ -4519,9 +4522,9 @@ int32 ULandscapeInfo::GetModifiedPackageCount() const
 	return IntCastChecked<int32>(Algo::CountIf(ModifiedPackages, [](const TWeakObjectPtr<UPackage>& InWeakPackagePtr) { return InWeakPackagePtr.IsValid(); }));
 }
 
-bool ULandscapeInfo::TryAddToModifiedPackages(UPackage* InPackage)
+bool ULandscapeInfo::TryAddToModifiedPackages(UPackage* InPackage, const ALandscape* InLandscapeOverride)
 {
-	ALandscape* LocalLandscapeActor = LandscapeActor.Get();
+	const ALandscape* LocalLandscapeActor = (InLandscapeOverride != nullptr) ? InLandscapeOverride : LandscapeActor.Get();
 	check(LocalLandscapeActor);
 
 	// We don't want to bother with packages being marked dirty for anything else than the Editor world 
@@ -4546,7 +4549,7 @@ bool ULandscapeInfo::TryAddToModifiedPackages(UPackage* InPackage)
 	return true;
 }
 
-bool ULandscapeInfo::MarkObjectDirty(UObject* InObject, bool bInForceResave)
+bool ULandscapeInfo::MarkObjectDirty(UObject* InObject, bool bInForceResave, const ALandscape* InLandscapeOverride)
 {
 	check(InObject && (InObject->IsA<ALandscapeProxy>() || InObject->GetTypedOuter<ALandscapeProxy>() != nullptr));
 
@@ -4557,12 +4560,12 @@ bool ULandscapeInfo::MarkObjectDirty(UObject* InObject, bool bInForceResave)
 		{
 			// When force-resaving (e.g. when syncing must-sync properties on load), unconditionally add the package to the list of packages to save if we couldn't mark it dirty already, so that 
 			//  the user can manually resave all that needs to be saved with the Save Modified Landscapes button :
-			bWasAddedToModifiedPackages = TryAddToModifiedPackages(InObject->GetPackage());
+			bWasAddedToModifiedPackages = TryAddToModifiedPackages(InObject->GetPackage(), InLandscapeOverride);
 		}
 	}
 	else if (bDirtyOnlyInMode)
 	{
-		ALandscape* LocalLandscapeActor = LandscapeActor.Get();
+		const ALandscape* LocalLandscapeActor = (InLandscapeOverride != nullptr) ? InLandscapeOverride : LandscapeActor.Get();
 		check(LocalLandscapeActor);
 		if (LocalLandscapeActor->HasLandscapeEdMode())
 		{
@@ -4570,7 +4573,7 @@ bool ULandscapeInfo::MarkObjectDirty(UObject* InObject, bool bInForceResave)
 		}
 		else
 		{
-			bWasAddedToModifiedPackages = TryAddToModifiedPackages(InObject->GetPackage());
+			bWasAddedToModifiedPackages = TryAddToModifiedPackages(InObject->GetPackage(), InLandscapeOverride);
 		}
 	}
 	else
