Fix for crash running PIE with some landscapes

This reverts the previous change that moved the render map allocation from CreateResources to Register -- here we move it back to CreateResources.  This fixes the bug because it is allocating other arrays based on the total allocation size, after CreateResources and before Register, so if you don't allocate in CreateResources those arrays can be incorrect size.

To fix the original bug that the allocation move was supposed to fix, I instead require that the number of Sections Created and Sections Registered are equal before running compaction, which guarantees that compaction will never remove entries that could still be in use.

diff --git a/Engine/Source/Runtime/Landscape/Private/LandscapeRender.cpp b/Engine/Source/Runtime/Landscape/Private/LandscapeRender.cpp
--- a/Engine/Source/Runtime/Landscape/Private/LandscapeRender.cpp
+++ b/Engine/Source/Runtime/Landscape/Private/LandscapeRender.cpp
@@ -556,6 +556,9 @@ void FLandscapeRenderSystem::CreateResources_Internal(FRHICommandListBase& RHICm
 	// we changed the RenderCoord, need to update the uniform buffer
 	SectionInfo->OnRenderCoordsChanged(RHICmdList);
 
+	check(SectionInfo->RenderCoord.X > INT32_MIN);
+	ResizeToInclude(SectionInfo->RenderCoord);
+
 	ReferenceCount++;
 }
 
@@ -571,7 +574,12 @@ void FLandscapeRenderSystem::DestroyResources_Internal(FLandscapeSectionInfo* Se
 
 	// try to compact the map every once in a while
 	SectionsRemovedSinceLastCompact++;
-	if (SectionsRemovedSinceLastCompact >= 128)
+
+	// When ReferenceCount == RegisteredCount then we know all resource-created sections are registered.
+	// So any empty sections of the map can be compacted without issues.  When they are not equal
+	// then compacting may erroneously remove areas that are still allocated to resource-created sections,
+	// but are not yet registered.
+	if ((SectionsRemovedSinceLastCompact >= 128) && ReferenceCount == RegisteredCount)
 	{
 		// if there are at least 128 free entries, run a compact step
 		if (Size.X * Size.Y - ReferenceCount >= 128)
@@ -600,7 +608,6 @@ void FLandscapeRenderSystem::RegisterSection(FLandscapeSectionInfo* SectionInfo)
 	FLandscapeSectionInfo* ExistingSection = LandscapeRenderSystem->GetSectionInfo(SectionInfo->RenderCoord);
 	if (ExistingSection == nullptr)
 	{
-		LandscapeRenderSystem->ResizeToInclude(SectionInfo->RenderCoord);
 		LandscapeRenderSystem->SetSectionInfo(SectionInfo->RenderCoord, SectionInfo);
 	}
 	else
