Removed unused CVar in landscapegrass.cpp

#rnx

diff --git a/Engine/Source/Runtime/Landscape/Private/LandscapeGrass.cpp b/Engine/Source/Runtime/Landscape/Private/LandscapeGrass.cpp
--- a/Engine/Source/Runtime/Landscape/Private/LandscapeGrass.cpp
+++ b/Engine/Source/Runtime/Landscape/Private/LandscapeGrass.cpp
@@ -118,12 +118,6 @@ static FAutoConsoleVariableRef CVarMaxAsyncTasks(
 	GMaxAsyncTasks,
 	TEXT("Used to control the number of grass components created at a time."));
 
-static int32 GUseHaltonDistribution = 0;
-static FAutoConsoleVariableRef CVarUseHaltonDistribution(
-	TEXT("grass.UseHaltonDistribution"),
-	GUseHaltonDistribution,
-	TEXT("Used to control the distribution of grass instances. If non-zero, use a halton sequence."));
-
 static float GGrassDensityScale = 1;
 static FAutoConsoleVariableRef CVarGrassDensityScale(
 	TEXT("grass.densityScale"),
