Fix an async loading crash due to renaming landscape brush (add REN_ForceNoResetLoaders flag). (via Patrick Enfedaque)

diff --git a/Engine/Source/Runtime/Landscape/Private/LandscapeEditLayers.cpp b/Engine/Source/Runtime/Landscape/Private/LandscapeEditLayers.cpp
--- a/Engine/Source/Runtime/Landscape/Private/LandscapeEditLayers.cpp
+++ b/Engine/Source/Runtime/Landscape/Private/LandscapeEditLayers.cpp
@@ -9999,7 +9999,7 @@ void FLandscapeLayerBrush::SetOwner(ALandscape* InOwner)
 	{
 		if (BlueprintBrush->GetTypedOuter<ULevel>() != InOwner->GetTypedOuter<ULevel>())
 		{
-			BlueprintBrush->Rename(nullptr, InOwner->GetTypedOuter<ULevel>());
+			BlueprintBrush->Rename(nullptr, InOwner->GetTypedOuter<ULevel>(), REN_ForceNoResetLoaders);
 		}
 		BlueprintBrush->SetOwningLandscape(InOwner);
 	}
