* Add ability to global toggle stripping of the landscape physics and grass data
* Fix a couple of instances of nullptr access.

diff --git a/Engine/Source/Runtime/Landscape/Private/Landscape.cpp b/Engine/Source/Runtime/Landscape/Private/Landscape.cpp
--- a/Engine/Source/Runtime/Landscape/Private/Landscape.cpp
+++ b/Engine/Source/Runtime/Landscape/Private/Landscape.cpp
@@ -185,6 +185,16 @@ static FAutoConsoleVariable CVarStripLayerTextureMipsOnLoad(
 	TEXT("landscape.StripLayerMipsOnLoad"),
 	false,
 	TEXT("Remove (on load) the mip chain from textures used in layers which don't require them"));
+
+static FAutoConsoleVariable CVarAllowPhysicsStripping(
+	TEXT("landscape.AllowPhysicsStripping"),
+	true,
+	TEXT("Enables the conditional stripping of physics data during cook.  Disabling this means the bStripPhysicsWhenCooked* will be ignored."));
+
+static FAutoConsoleVariable CVarAllowGrassStripping(
+	TEXT("landscape.AllowGrassStripping"),
+	true,
+	TEXT("Enables the conditional stripping of grass data during cook.  Disabling this means the bStripGrassWhenCooked* will be ignored."));
 #endif // WITH_EDITOR
 
 int32 GRenderNaniteLandscape = 1;
@@ -2156,14 +2166,22 @@ TArray<FWeightmapLayerAllocationInfo>& ULandscapeComponent::GetCurrentRuntimeWei
 
 FLandscapeLayerComponentData* ULandscapeComponent::GetEditingLayer()
 {
-	const FGuid& EditingLayerGuid = GetLandscapeActor()->GetEditingLayer();
-	return EditingLayerGuid.IsValid() ? LayersData.Find(EditingLayerGuid) : nullptr;
+	if (ALandscape* LandscapeActor = GetLandscapeActor())
+	{
+		const FGuid& EditingLayerGuid = LandscapeActor->GetEditingLayer();
+		return EditingLayerGuid.IsValid() ? LayersData.Find(EditingLayerGuid) : nullptr;
+	}
+	return nullptr;
 }
 
 const FLandscapeLayerComponentData* ULandscapeComponent::GetEditingLayer() const
 {
-	const FGuid& EditingLayerGuid = GetLandscapeActor()->GetEditingLayer();
-	return EditingLayerGuid.IsValid() ? LayersData.Find(EditingLayerGuid) : nullptr;
+	if (ALandscape* LandscapeActor = GetLandscapeActor())
+	{
+		const FGuid& EditingLayerGuid = LandscapeActor->GetEditingLayer();
+		return EditingLayerGuid.IsValid() ? LayersData.Find(EditingLayerGuid) : nullptr;
+	}
+	return nullptr;
 }
 
 void ULandscapeComponent::CopyFinalLayerIntoEditingLayer(FLandscapeEditDataInterface& DataInterface, TSet<UTexture2D*>& ProcessedHeightmaps)
@@ -3228,9 +3246,10 @@ void ALandscapeProxy::PreSave(FObjectPreSaveContext ObjectSaveContext)
 	// Strip data according to flags when cooked (and not cooking for editor)
 	if (ObjectSaveContext.IsCooking() && !ObjectSaveContext.GetTargetPlatform()->AllowsEditorObjects())
 	{
-		if ((bStripPhysicsWhenCookedClient && bStripPhysicsWhenCookedServer) ||
+		if (CVarAllowPhysicsStripping->GetBool() &&
+			((bStripPhysicsWhenCookedClient && bStripPhysicsWhenCookedServer) ||
 			(bStripPhysicsWhenCookedClient && ObjectSaveContext.GetTargetPlatform()->IsClientOnly()) ||
-			(bStripPhysicsWhenCookedServer && ObjectSaveContext.GetTargetPlatform()->IsServerOnly()))
+			(bStripPhysicsWhenCookedServer && ObjectSaveContext.GetTargetPlatform()->IsServerOnly())))
 		{
 			// Clear old CollisionComponent containers
 			CollisionComponents.Empty();
@@ -3263,9 +3282,10 @@ void ALandscapeProxy::PreSave(FObjectPreSaveContext ObjectSaveContext)
 			}
 		}
 		
-		if ((bStripGrassWhenCookedClient && bStripGrassWhenCookedServer) ||
+		if (CVarAllowGrassStripping->GetBool() &&
+			((bStripGrassWhenCookedClient && bStripGrassWhenCookedServer) ||
 			(bStripGrassWhenCookedClient && ObjectSaveContext.GetTargetPlatform()->IsClientOnly()) ||
-			(bStripGrassWhenCookedServer && ObjectSaveContext.GetTargetPlatform()->IsServerOnly()))
+			(bStripGrassWhenCookedServer && ObjectSaveContext.GetTargetPlatform()->IsServerOnly())))
 		{
 			for (ULandscapeComponent* Component : LandscapeComponents)
 			{
