Added RenderWeightmap/RenderWeightmaps BP functions to landscape in order to let users pack weightmaps onto the channels of a render target 2D (or render target 2D array), much like they could with heightmaps
These can be called at runtime as they operate on the final weightmap data and they will render in a given world area, regardless of landscape component limits

diff --git a/Engine/Shaders/Private/Landscape/LandscapeMergeHeightmapsPS.usf b/Engine/Shaders/Private/Landscape/LandscapeMergeHeightmapsPS.usf
deleted file mode 100644
diff --git a/Engine/Shaders/Private/Landscape/LandscapeMergeTexturesPS.usf b/Engine/Shaders/Private/Landscape/LandscapeMergeTexturesPS.usf
new file mode 100644
--- /dev/null
+++ b/Engine/Shaders/Private/Landscape/LandscapeMergeTexturesPS.usf
@@ -0,0 +1,128 @@
+// Copyright Epic Games, Inc. All Rights Reserved.
+
+#include "../Common.ush"
+#include "LandscapeCommon.ush"
+
+// ----------------------------------------------------------------------------------
+
+#if defined (__INTELLISENSE__)
+// Uncomment the appropriate define for enabling syntax highlighting with HLSL Tools for Visual Studio : 
+//#define IS_HEIGHTMAP 1
+//#define MERGE_TEXTURE 1
+//#define RESAMPLE_MERGED_TEXTURE 1
+//#define COMPRESS_HEIGHT 1
+//#define PACK_RGBA_CHANNELS 1
+#endif // defined (__INTELLISENSE__)
+
+// ----------------------------------------------------------------------------------
+
+#if MERGE_TEXTURE
+// MergeTexture inputs/outputs :
+
+uint4 InAtlasSubregion; // Subregion of the atlas texture that is being rendered to
+uint4 InSourceTextureSubregion; // Subregion of the input texture that we copy from
+int InSourceTextureChannel; // Channel of the input texture that we copy from (for weightmaps only)
+Texture2D<float4> InSourceTexture; // Source texture (packed height in the .rg channels for heightmap, weight in one the 4 .rgba channels for the weightmap)
+#endif // MERGE_TEXTURE
+
+#if RESAMPLE_MERGED_TEXTURE
+// ResampleTexture inputs/outputs :
+
+float4x4 InOutputUVToMergedTextureUV; // Transform to apply to the texture coordinates in order to resample the atlas in the requested render area
+Texture2D<float> InMergedTexture; // Source texture (unpacked floating point height for heightmap)
+SamplerState InMergedTextureSampler;
+uint2 InRenderAreaSize; // Size of the output render target
+#endif // RESAMPLE_MERGED_TEXTURE
+
+#if PACK_RGBA_CHANNELS
+int InNumChannels; // Number of channels to pack
+Texture2D<float> InSourceTextures_0; // Source, single-channel, texture to pack
+Texture2D<float> InSourceTextures_1; // Source, single-channel, texture to pack
+Texture2D<float> InSourceTextures_2; // Source, single-channel, texture to pack
+Texture2D<float> InSourceTextures_3; // Source, single-channel, texture to pack
+#endif // PACK_RGBA_CHANNELS
+
+static const float LANDSCAPE_MAX_VALUE = 65535.0f;
+
+// ----------------------------------------------------------------------------------
+// Util functions : 
+
+#if MERGE_TEXTURE
+
+void MergeTexture(in float4 SVPos : SV_POSITION, out float OutValue : SV_Target0)
+{
+	// Input texture and output texture atlas may not be at the same resolution (different mips), so we need to renormalize the sample position before sampling the input texture :
+	uint2 AtlasSubregionSize = InAtlasSubregion.zw - InAtlasSubregion.xy;
+	float2 SampleCoordinatesNormalized = (SVPos.xy - InAtlasSubregion.xy) / (float2) AtlasSubregionSize;
+	uint2 SourceTextureSubregionSize = InSourceTextureSubregion.zw - InSourceTextureSubregion.xy;
+	uint2 SampleCoordinates = floor(InSourceTextureSubregion.xy + SampleCoordinatesNormalized * SourceTextureSubregionSize);
+		
+	float4 Sample = InSourceTexture.Load(int3(SampleCoordinates, 0));
+
+#if IS_HEIGHTMAP
+	OutValue = UnpackHeight(Sample.xy) / LANDSCAPE_MAX_VALUE;
+#else // IS_HEIGHTMAP
+	OutValue = Sample[InSourceTextureChannel];
+#endif // !IS_HEIGHTMAP
+}
+
+#endif // MERGE_TEXTURE
+
+
+// ----------------------------------------------------------------------------------
+
+#if RESAMPLE_MERGED_TEXTURE
+
+void ResampleMergedTexture(in float4 SVPos : SV_POSITION, out float4 OutColor : SV_Target0)
+{
+	OutColor = 0;
+
+	// Pixel coordinates to UV (SvPos is already offset by 0.5 to be at pixel center) :
+	float2 UV = SVPos.xy / (float2)InRenderAreaSize;
+	float4 SampleUV = float4(UV, 0.0f, 1.0f);
+	SampleUV = mul(SampleUV, InOutputUVToMergedTextureUV);
+
+	float Sample = 0.0f;
+	if (all(SampleUV.xy >= 0.0f) && all(SampleUV.xy <= 1.0f))
+	{
+		Sample = InMergedTexture.Sample(InMergedTextureSampler, SampleUV.xy);
+	}
+
+	OutColor.x = Sample;
+
+#if IS_HEIGHTMAP
+	Sample *= LANDSCAPE_MAX_VALUE;
+
+#if COMPRESS_HEIGHT
+	OutColor.xy = PackHeight(Sample);
+#endif // !COMPRESS_HEIGHT
+
+#endif // IS_HEIGHTMAP
+}
+
+#endif // RESAMPLE_MERGED_TEXTURE
+
+
+// ----------------------------------------------------------------------------------
+
+#if PACK_RGBA_CHANNELS
+
+void PackRGBAChannels(in float4 SVPos : SV_POSITION, out float4 OutColor : SV_Target0)
+{
+	OutColor = 0;
+	OutColor.x = InSourceTextures_0.Load(int3(SVPos.xy, 0)).x;
+	if (InNumChannels > 0)
+	{
+		OutColor.y = InSourceTextures_1.Load(int3(SVPos.xy, 0)).x;
+	}
+	if (InNumChannels > 1)
+	{
+		OutColor.z = InSourceTextures_2.Load(int3(SVPos.xy, 0)).x;
+	}
+	if (InNumChannels > 2)
+	{
+		OutColor.w = InSourceTextures_3.Load(int3(SVPos.xy, 0)).x;
+	}
+}
+
+#endif // PACK_RGBA_CHANNELS
\ No newline at end of file
diff --git a/Engine/Source/Runtime/Landscape/Classes/Landscape.h b/Engine/Source/Runtime/Landscape/Classes/Landscape.h
--- a/Engine/Source/Runtime/Landscape/Classes/Landscape.h
+++ b/Engine/Source/Runtime/Landscape/Classes/Landscape.h
@@ -17,6 +17,7 @@ class ULandscapeComponent;
 class ILandscapeEdModeInterface;
 class SNotificationItem;
 class UStreamableRenderAsset;
+class UTextureRenderTarget;
 class FMaterialResource;
 struct FLandscapeEditLayerComponentReadbackResult;
 struct FLandscapeNotification;
@@ -264,10 +265,45 @@ public:
 	LANDSCAPE_API void SetLODGroupKey(uint32 InLODGroupKey);
 	LANDSCAPE_API uint32 GetLODGroupKey();
 
+	/**
+	* Render the final heightmap in the requested top-down window as one -atlased- texture in the provided render target 2D
+	*  Can be called at runtime.
+	* @param InWorldTransform World transform of the area where the texture should be rendered
+	* @param InExtents Extents of the area where the texture should be rendered (local to InWorldTransform). If size is zero, then the entire loaded landscape will be exported.
+	* @param OutRenderTarget Render target in which the texture will be rendered. The size/format of the render target will be respected.
+	* @return false in case of failure (e.g. invalid inputs, incompatible render target format...)
+	*/
 	UFUNCTION(BlueprintCallable, Category = "Landscape|Runtime")
-	LANDSCAPE_API void RenderHeightmap(const FTransform& InWorldTransform, const FBox2D& InExtents, UTextureRenderTarget2D* OutRenderTarget);
+	LANDSCAPE_API bool RenderHeightmap(const FTransform& InWorldTransform, const FBox2D& InExtents, UTextureRenderTarget2D* OutRenderTarget);
+
+	/**
+	* Render the final weightmap for the requested layer, in the requested top-down window, as one -atlased- texture in the provided render target 2D
+	*  Can be called at runtime.
+	* @param InWorldTransform World transform of the area where the texture should be rendered
+	* @param InExtents Extents of the area where the texture should be rendered (local to InWorldTransform). If size is zero, then the entire loaded landscape will be exported.
+	* @param InWeightmapLayerName Weightmap layer that is being requested to render
+	* @param OutRenderTarget Render target in which the texture will be rendered. The size/format of the render target will be respected.
+	* @return false in case of failure (e.g. invalid inputs, incompatible render target format...)
+	*/
+	UFUNCTION(BlueprintCallable, Category = "Landscape|Runtime")
+	LANDSCAPE_API bool RenderWeightmap(const FTransform& InWorldTransform, const FBox2D& InExtents, FName InWeightmapLayerName, UTextureRenderTarget2D* OutRenderTarget);
+
+	/**
+	* Render the final weightmaps for the requested layers, in the requested top-down window, as one -atlased- texture in the provided render target (2D or 2DArray) 
+	*  Can be called at runtime.
+	* @param InWorldTransform World transform of the area where the texture should be rendered
+	* @param InExtents Extents of the area where the texture should be rendered (local to InWorldTransform). If size is zero, then the entire loaded landscape will be exported.
+	* @param InWeightmapLayerNames List of weightmap layers that are being requested to render
+	* @param OutRenderTarget Render target in which the texture will be rendered. The size/format of the render target will be respected.
+	*  - If a UTextureRenderTarget2D is passed, the requested layers will be packed in the RGBA channels in order (up to the number of channels available with the render target's format).
+	*  - If a UTextureRenderTarget2DArray is passed, the requested layers will be packed in the RGBA channels of each slice (up to the number of channels * slices available with the render target's format and number of slices).
+	* @return false in case of failure (e.g. invalid inputs, incompatible render target format...)
+	*/
+	UFUNCTION(BlueprintCallable, Category = "Landscape|Runtime")
+	LANDSCAPE_API bool RenderWeightmaps(const FTransform& InWorldTransform, const FBox2D& InExtents, const TArray<FName>& InWeightmapLayerNames, UTextureRenderTarget* OutRenderTarget);
 
 	bool IsValidRenderTargetFormatHeightmap(EPixelFormat InRenderTargetFormat, bool& bOutCompressHeight);
+	bool IsValidRenderTargetFormatWeightmap(EPixelFormat InRenderTargetFormat, int32& OutNumChannels);
 
 #if WITH_EDITOR
 	/** Computes & returns bounds containing all landscape proxies (if any) or this landscape's bounds otherwise. Note that in non-WP worlds this will call GetLoadedBounds(). */
@@ -447,6 +483,7 @@ private:
 
 private:
 	void MarkAllLandscapeRenderStateDirty();
+	bool RenderMergedTextureInternal(const FTransform& InRenderAreaWorldTransform, const FBox2D& InRenderAreaExtents, const TArray<FName>& InWeightmapLayerNames, UTextureRenderTarget* OutRenderTarget);
 
 public:
 
diff --git a/Engine/Source/Runtime/Landscape/Private/Landscape.cpp b/Engine/Source/Runtime/Landscape/Private/Landscape.cpp
--- a/Engine/Source/Runtime/Landscape/Private/Landscape.cpp
+++ b/Engine/Source/Runtime/Landscape/Private/Landscape.cpp
@@ -58,6 +58,7 @@ Landscape.cpp: Terrain rendering
 #include "EngineGlobals.h"
 #include "Engine/Engine.h"
 #include "Engine/TextureRenderTarget2D.h"
+#include "Engine/TextureRenderTarget2DArray.h"
 #include "EngineUtils.h"
 #include "ComponentRecreateRenderStateContext.h"
 #include "LandscapeWeightmapUsage.h"
@@ -83,6 +84,7 @@ Landscape.cpp: Terrain rendering
 #include "LandscapeDataAccess.h"
 #include "LandscapeNotification.h"
 #include "LandscapeHeightfieldCollisionComponent.h"
+#include "SystemTextures.h"
 #include "Algo/BinarySearch.h"
 #include "Algo/Count.h"
 #include "Algo/Transform.h"
@@ -151,11 +153,14 @@ FAutoConsoleCommand CmdPrintNumLandscapeShadows(
 	FConsoleCommandDelegate::CreateStatic(PrintNumLandscapeShadows)
 	);
 
-int32 RenderCaptureNextHeightmapRenders = 0;
-static FAutoConsoleVariableRef CVarRenderCaptureNextHeightmapRenders(
-	TEXT("landscape.RenderCaptureNextHeightmapRenders"),
-	RenderCaptureNextHeightmapRenders,
-	TEXT("Trigger a render capture during the next N RenderHeightmap draws"));
+namespace UE::Landscape
+{
+int32 RenderCaptureNextMergeRenders = 0;
+static FAutoConsoleVariableRef CVarRenderCaptureNextMergeRenders(
+	TEXT("landscape.RenderCaptureNextMergeRenders"),
+	RenderCaptureNextMergeRenders,
+	TEXT("Trigger a render capture during the next N RenderHeightmap/RenderWeightmap(s) draws"));
+} // namespace UE::Landscape
 
 #if WITH_EDITOR
 
@@ -2750,128 +2755,194 @@ FBox ALandscape::GetLoadedBounds() const
 
 // ----------------------------------------------------------------------------------
 
-// This shader allows to render parts of the heightmaps (all pixels except the redundant ones on the right/bottom edges) in an atlas render target (uncompressed height)
-class FLandscapeMergeHeightmapsPS : public FGlobalShader
+// This shader allows to render parts of the heightmaps/weightmaps (all pixels except the redundant ones on the right/bottom edges) in an atlas render target (uncompressed height for heightmaps)
+class FLandscapeMergeTexturesPS : public FGlobalShader
 {
-	DECLARE_GLOBAL_SHADER(FLandscapeMergeHeightmapsPS);
-	SHADER_USE_PARAMETER_STRUCT(FLandscapeMergeHeightmapsPS, FGlobalShader);
+	DECLARE_GLOBAL_SHADER(FLandscapeMergeTexturesPS);
+	SHADER_USE_PARAMETER_STRUCT(FLandscapeMergeTexturesPS, FGlobalShader);
 
 public:
 	BEGIN_SHADER_PARAMETER_STRUCT(FParameters, )
-		SHADER_PARAMETER(FUintVector4, InHeightmapAtlasSubregion)
-		SHADER_PARAMETER(FUintVector4, InHeightmapSubregion)
-		SHADER_PARAMETER_RDG_TEXTURE_SRV(Texture2D<float4>, InHeightmap)
+		SHADER_PARAMETER(FUintVector4, InAtlasSubregion)
+		SHADER_PARAMETER(FUintVector4, InSourceTextureSubregion)
+		SHADER_PARAMETER(int32, InSourceTextureChannel)
+		SHADER_PARAMETER_RDG_TEXTURE_SRV(Texture2D<float4>, InSourceTexture)
 		RENDER_TARGET_BINDING_SLOTS()
-		END_SHADER_PARAMETER_STRUCT()
+	END_SHADER_PARAMETER_STRUCT()
+
+	class FIsHeightmap : SHADER_PERMUTATION_BOOL("IS_HEIGHTMAP");
+
+	using FPermutationDomain = TShaderPermutationDomain<FIsHeightmap>;
 
-		static bool ShouldCompilePermutation(const FGlobalShaderPermutationParameters& InParameters)
+	static FPermutationDomain GetPermutationVector(bool bInIsHeighmap)
+	{
+		FPermutationDomain PermutationVector;
+		PermutationVector.Set<FIsHeightmap>(bInIsHeighmap);
+		return PermutationVector;
+	}
+
+	
+	static bool ShouldCompilePermutation(const FGlobalShaderPermutationParameters& InParameters)
 	{
 		return true;
 	}
 
 	static void ModifyCompilationEnvironment(const FGlobalShaderPermutationParameters& InParameters, FShaderCompilerEnvironment& OutEnvironment)
 	{
-		OutEnvironment.SetDefine(TEXT("MERGE_HEIGHTMAP"), 1);
+		OutEnvironment.SetDefine(TEXT("MERGE_TEXTURE"), 1);
 	}
 
-	static void MergeHeightmap(FRDGBuilder& GraphBuilder, FParameters* InParameters, const FIntRect& InRenderTargetArea)
+	static void MergeTexture(FRDGBuilder& GraphBuilder, FParameters* InParameters, const FIntRect& InRenderTargetArea, bool bInIsHeightmap)
 	{
 		FGlobalShaderMap* ShaderMap = GetGlobalShaderMap(GMaxRHIFeatureLevel);
-		TShaderMapRef<FLandscapeMergeHeightmapsPS> PixelShader(ShaderMap);
+
+		const FLandscapeMergeTexturesPS::FPermutationDomain PixelPermutationVector = FLandscapeMergeTexturesPS::GetPermutationVector(bInIsHeightmap);
+		TShaderMapRef<FLandscapeMergeTexturesPS> PixelShader(ShaderMap, PixelPermutationVector);
 
 		FPixelShaderUtils::AddFullscreenPass(
 			GraphBuilder,
 			ShaderMap,
-			RDG_EVENT_NAME("LandscapeLayers_MergeHeightmap"),
+			RDG_EVENT_NAME("LandscapeMergeTexture"),
 			PixelShader,
 			InParameters,
 			InRenderTargetArea);
 	}
 };
 
-IMPLEMENT_GLOBAL_SHADER(FLandscapeMergeHeightmapsPS, "/Engine/Private/Landscape/LandscapeMergeHeightmapsPS.usf", "MergeHeightmap", SF_Pixel);
+IMPLEMENT_GLOBAL_SHADER(FLandscapeMergeTexturesPS, "/Engine/Private/Landscape/LandscapeMergeTexturesPS.usf", "MergeTexture", SF_Pixel);
 
 
 // ----------------------------------------------------------------------------------
 
-// This shader allows to resample the heightmap (bilinear interpolation) from a given atlas usually heightmap produced by FLandscapeMergeHeightmapsPS :
-//  The output heightmap's heights can be either compressed or uncompressed depending on the render target format (8 bits/channel for the former, 16/32 bits/channel for the latter)
-class FLandscapeResampleHeightmapsPS : public FGlobalShader
+// This shader allows to resample the heightmap/weightmap (bilinear interpolation) from a given atlas usually produced by FLandscapeMergeTexturesPS :
+//  For heightmap, the output can be either compressed or uncompressed depending on the render target format (8 bits/channel for the former, 16/32 bits/channel for the latter)
+class FLandscapeResampleMergedTexturePS : public FGlobalShader
 {
-	DECLARE_GLOBAL_SHADER(FLandscapeResampleHeightmapsPS);
-	SHADER_USE_PARAMETER_STRUCT(FLandscapeResampleHeightmapsPS, FGlobalShader);
+	DECLARE_GLOBAL_SHADER(FLandscapeResampleMergedTexturePS);
+	SHADER_USE_PARAMETER_STRUCT(FLandscapeResampleMergedTexturePS, FGlobalShader);
 
 public:
 	BEGIN_SHADER_PARAMETER_STRUCT(FParameters, )
-		SHADER_PARAMETER(FMatrix44f, InOutputUVToMergedHeightmapUV)
-		SHADER_PARAMETER_RDG_TEXTURE_SRV(Texture2D<float4>, InMergedHeightmap)
-		SHADER_PARAMETER_SAMPLER(SamplerState, InMergedHeightmapSampler)
+		SHADER_PARAMETER(FMatrix44f, InOutputUVToMergedTextureUV)
+		SHADER_PARAMETER_RDG_TEXTURE_SRV(Texture2D<float>, InMergedTexture)
+		SHADER_PARAMETER_SAMPLER(SamplerState, InMergedTextureSampler)
 		SHADER_PARAMETER(FUintVector2, InRenderAreaSize)
 		RENDER_TARGET_BINDING_SLOTS()
-		END_SHADER_PARAMETER_STRUCT()
+	END_SHADER_PARAMETER_STRUCT()
 
-		class FCompressHeight : SHADER_PERMUTATION_BOOL("COMPRESS_HEIGHT");
-	using FPermutationDomain = TShaderPermutationDomain<FCompressHeight>;
+	class FIsHeightmap : SHADER_PERMUTATION_BOOL("IS_HEIGHTMAP");
+	class FCompressHeight : SHADER_PERMUTATION_BOOL("COMPRESS_HEIGHT");
 
-	static FPermutationDomain GetPermutationVector(bool bCompressHeight)
+	using FPermutationDomain = TShaderPermutationDomain<FIsHeightmap, FCompressHeight>;
+
+	static FPermutationDomain GetPermutationVector(bool bInIsHeighmap, bool bInCompressHeight)
 	{
 		FPermutationDomain PermutationVector;
-		PermutationVector.Set<FCompressHeight>(bCompressHeight);
+		PermutationVector.Set<FIsHeightmap>(bInIsHeighmap);
+		PermutationVector.Set<FCompressHeight>(bInCompressHeight);
 		return PermutationVector;
 	}
 
 	static bool ShouldCompilePermutation(const FGlobalShaderPermutationParameters& InParameters)
 	{
-		return true;
+		FPermutationDomain PermutationVector(InParameters.PermutationId);
+		bool bIsHeightmap = PermutationVector.Get<FIsHeightmap>();
+		bool bCompressHeight = PermutationVector.Get<FCompressHeight>();
+		// No need for heightmap compression for weightmaps
+		return (bIsHeightmap || !bCompressHeight);
 	}
 
 	static void ModifyCompilationEnvironment(const FGlobalShaderPermutationParameters& InParameters, FShaderCompilerEnvironment& OutEnvironment)
 	{
-		OutEnvironment.SetDefine(TEXT("RESAMPLE_HEIGHTMAP"), 1);
+		OutEnvironment.SetDefine(TEXT("RESAMPLE_MERGED_TEXTURE"), 1);
 	}
 
-	static void ResampleHeightmap(FRDGBuilder& GraphBuilder, FParameters* InParameters, bool bCompressHeight)
+	static void ResampleMergedTexture(FRDGBuilder& GraphBuilder, FParameters* InParameters, bool bInIsHeightmap, bool bInCompressHeight)
 	{
 		FGlobalShaderMap* ShaderMap = GetGlobalShaderMap(GMaxRHIFeatureLevel);
 
-		const FLandscapeResampleHeightmapsPS::FPermutationDomain PixelPermutationVector = FLandscapeResampleHeightmapsPS::GetPermutationVector(bCompressHeight);
-
-		TShaderMapRef<FLandscapeResampleHeightmapsPS> PixelShader(ShaderMap, PixelPermutationVector);
+		const FLandscapeResampleMergedTexturePS::FPermutationDomain PixelPermutationVector = FLandscapeResampleMergedTexturePS::GetPermutationVector(bInIsHeightmap, bInCompressHeight);
+		TShaderMapRef<FLandscapeResampleMergedTexturePS> PixelShader(ShaderMap, PixelPermutationVector);
 
 		FPixelShaderUtils::AddFullscreenPass(
 			GraphBuilder,
 			ShaderMap,
-			RDG_EVENT_NAME("LandscapeLayers_ResampleHeightmap"),
+			RDG_EVENT_NAME("ResampleMergedTexture"),
 			PixelShader,
 			InParameters,
 			FIntRect(0, 0, InParameters->InRenderAreaSize.X, InParameters->InRenderAreaSize.Y));
 	}
 };
 
-IMPLEMENT_GLOBAL_SHADER(FLandscapeResampleHeightmapsPS, "/Engine/Private/Landscape/LandscapeMergeHeightmapsPS.usf", "ResampleHeightmap", SF_Pixel);
+IMPLEMENT_GLOBAL_SHADER(FLandscapeResampleMergedTexturePS, "/Engine/Private/Landscape/LandscapeMergeTexturesPS.usf", "ResampleMergedTexture", SF_Pixel);
+
+
+// ----------------------------------------------------------------------------------
+
+// This shader allows to pack up to 4 single-channel textures onto a single rgba one
+class FLandscapePackRGBAChannelsPS : public FGlobalShader
+{
+	DECLARE_GLOBAL_SHADER(FLandscapePackRGBAChannelsPS);
+	SHADER_USE_PARAMETER_STRUCT(FLandscapePackRGBAChannelsPS, FGlobalShader);
+
+public:
+	BEGIN_SHADER_PARAMETER_STRUCT(FParameters, )
+		SHADER_PARAMETER(int32, InNumChannels)
+		SHADER_PARAMETER_RDG_TEXTURE_SRV_ARRAY(Texture2D<float>, InSourceTextures, [4])
+		RENDER_TARGET_BINDING_SLOTS()
+	END_SHADER_PARAMETER_STRUCT()
+
+	static bool ShouldCompilePermutation(const FGlobalShaderPermutationParameters& InParameters)
+	{
+		return true;
+	}
 
+	static void ModifyCompilationEnvironment(const FGlobalShaderPermutationParameters& InParameters, FShaderCompilerEnvironment& OutEnvironment)
+	{
+		OutEnvironment.SetDefine(TEXT("PACK_RGBA_CHANNELS"), 1);
+	}
+
+	static void PackRGBAChannels(FRDGBuilder& GraphBuilder, FParameters* InParameters, const FIntRect& InRenderTargetArea)
+	{
+		FGlobalShaderMap* ShaderMap = GetGlobalShaderMap(GMaxRHIFeatureLevel);
+
+		TShaderMapRef<FLandscapePackRGBAChannelsPS> PixelShader(ShaderMap);
+
+		FPixelShaderUtils::AddFullscreenPass(
+			GraphBuilder,
+			ShaderMap,
+			RDG_EVENT_NAME("PackRGBAChannels"),
+			PixelShader,
+			InParameters,
+			InRenderTargetArea);
+	}
+};
+
+IMPLEMENT_GLOBAL_SHADER(FLandscapePackRGBAChannelsPS, "/Engine/Private/Landscape/LandscapeMergeTexturesPS.usf", "PackRGBAChannels", SF_Pixel);
 
 // Render-thread version of the data / functions we need for the local merge of edit layers : 
-namespace RenderMergedLandscape_RenderThread
+namespace UE::Landscape::Private::RenderMergedTexture_RenderThread
 {
-	struct FHeightmapRenderInfo
+	struct FRenderInfo
 	{
 		// Transform to go from the output render area space ((0,0) in the lower left corner, (1,1) in the upper-right) to the temporary render target space
-		FMatrix OutputUVToMergedHeightmapUV;
-		FBox RenderAreaExtents;
-		// TODO [jonathan.bard] : remove ? This is only used to compute the temporary render target size
+		FMatrix OutputUVToMergedTextureUV;
 		FIntPoint SubsectionSizeQuads;
 		int32 NumSubsections = 1;
+		bool bIsHeightmap = true;
 		bool bCompressHeight = false;
+		FName TargetLayerName;
 
-		TMap<FIntPoint, FTexture2DResourceSubregion> ComponentHeightmapsToRender;
+		TMap<FIntPoint, FTexture2DResourceSubregion> ComponentTexturesToRender;
 	};
 
-	void RenderMergedHeightmap(const FHeightmapRenderInfo& InRenderInfo, FRDGBuilder& GraphBuilder, FRDGTextureRef OutputTexture)
+	void RenderMergedTexture(const FRenderInfo& InRenderInfo, FRDGBuilder& GraphBuilder, const FRenderTargetBinding& InOutputRenderTargetBinding)
 	{
+		RDG_EVENT_SCOPE(GraphBuilder, "RenderMergedTexture %s", *InRenderInfo.TargetLayerName.ToString());
+
 		// Find the total area that those components need to be rendered to :
 		FIntRect ComponentKeyRect;
-		for (auto Iter = InRenderInfo.ComponentHeightmapsToRender.CreateConstIterator(); Iter; ++Iter)
+		for (auto Iter = InRenderInfo.ComponentTexturesToRender.CreateConstIterator(); Iter; ++Iter)
 		{
 			ComponentKeyRect.Include(Iter.Key());
 		}
@@ -2882,13 +2953,18 @@ namespace RenderMergedLandscape_RenderThread
 		FIntPoint RenderTargetSize = NumSubsectionsToRender * InRenderInfo.SubsectionSizeQuads + 1; // add one for the end vertex
 		FIntPoint ComponentSizeQuads = InRenderInfo.SubsectionSizeQuads * InRenderInfo.NumSubsections;
 
-		// We need a temporary render target that can contain all heightmaps. Use PF_G16 (decoded height) as this will be resampled using bilinear sampling :
-		FRDGTextureDesc Desc = FRDGTextureDesc::Create2D(RenderTargetSize, PF_G16, FClearValueBinding::Black, TexCreate_RenderTargetable | TexCreate_ShaderResource);
-		FRDGTextureRef AtlasTexture = GraphBuilder.CreateTexture(Desc, TEXT("LandscapeHeightmapAtlas"));
+		// We need a temporary render target that can contain all textures. 
+		// For heightmaps, use PF_G16 (decoded height) as this will be resampled using bilinear sampling :
+		EPixelFormat AtlasTextureFormat = InRenderInfo.bIsHeightmap ? PF_G16 : PF_G8;
+		FRDGTextureDesc Desc = FRDGTextureDesc::Create2D(RenderTargetSize, AtlasTextureFormat, FClearValueBinding::Black, TexCreate_RenderTargetable | TexCreate_ShaderResource);
+		FRDGTextureRef AtlasTexture = GraphBuilder.CreateTexture(Desc, TEXT("LandscapeMergedTextureAtlas"));
 		FRDGTextureSRVRef AtlasTextureSRV = GraphBuilder.CreateSRV(FRDGTextureSRVDesc::Create(AtlasTexture));
-		FRenderTargetBinding AtlasTextureRT(AtlasTexture, ERenderTargetLoadAction::ELoad);
+		// Start with a cleared atlas : 
+		FRenderTargetBinding AtlasTextureRT(AtlasTexture, ERenderTargetLoadAction::ENoAction);
+		FRDGTextureClearInfo ClearInfo;
+		AddClearRenderTargetPass(GraphBuilder, AtlasTexture, ClearInfo);
 
-		TMap<FTexture2DResource*, FRDGTextureSRVRef> HeightmapTextureSRVs;
+		TMap<FTexture2DResource*, FRDGTextureSRVRef> SourceTextureSRVs;
 
 		// Fill that render target subsection by subsection, in order to bypass the redundant columns/lines on the subsection edges:
 		for (int32 ComponentY = ComponentKeyRect.Min.Y; ComponentY < ComponentKeyRect.Max.Y; ++ComponentY)
@@ -2896,15 +2972,15 @@ namespace RenderMergedLandscape_RenderThread
 			for (int32 ComponentX = ComponentKeyRect.Min.X; ComponentX < ComponentKeyRect.Max.X; ++ComponentX)
 			{
 				FIntPoint LandscapeComponentKey(ComponentX, ComponentY);
-				if (const FTexture2DResourceSubregion* HeightmapResourceSubregion = InRenderInfo.ComponentHeightmapsToRender.Find(LandscapeComponentKey))
+				if (const FTexture2DResourceSubregion* SourceTextureResourceSubregion = InRenderInfo.ComponentTexturesToRender.Find(LandscapeComponentKey))
 				{
-					FIntPoint SubsectionSubregionSize = HeightmapResourceSubregion->Subregion.Size() / InRenderInfo.NumSubsections;
-					FRDGTextureSRVRef* Heightmap = HeightmapTextureSRVs.Find(HeightmapResourceSubregion->Texture);
-					if (Heightmap == nullptr)
+					FIntPoint SubsectionSubregionSize = SourceTextureResourceSubregion->Subregion.Size() / InRenderInfo.NumSubsections;
+					FRDGTextureSRVRef* SourceTextureSRV = SourceTextureSRVs.Find(SourceTextureResourceSubregion->Texture);
+					if (SourceTextureSRV == nullptr)
 					{
-						FString* DebugString = GraphBuilder.AllocObject<FString>(HeightmapResourceSubregion->Texture->GetTextureName().ToString());
-						FRDGTextureRef TextureRef = GraphBuilder.RegisterExternalTexture(CreateRenderTarget(HeightmapResourceSubregion->Texture->TextureRHI, **DebugString));
-						Heightmap = &HeightmapTextureSRVs.Add(HeightmapResourceSubregion->Texture, GraphBuilder.CreateSRV(FRDGTextureSRVDesc::Create(TextureRef)));
+						FString* DebugString = GraphBuilder.AllocObject<FString>(SourceTextureResourceSubregion->Texture->GetTextureName().ToString());
+						FRDGTextureRef TextureRef = GraphBuilder.RegisterExternalTexture(CreateRenderTarget(SourceTextureResourceSubregion->Texture->TextureRHI, **DebugString));
+						SourceTextureSRV = &SourceTextureSRVs.Add(SourceTextureResourceSubregion->Texture, GraphBuilder.CreateSRV(FRDGTextureSRVDesc::Create(TextureRef)));
 					}
 
 					for (int32 SubsectionY = 0; SubsectionY < InRenderInfo.NumSubsections; ++SubsectionY)
@@ -2914,55 +2990,49 @@ namespace RenderMergedLandscape_RenderThread
 							FIntPoint SubsectionLocalKey(SubsectionX, SubsectionY);
 							FIntPoint SubsectionKey = LandscapeComponentKey * InRenderInfo.NumSubsections + SubsectionLocalKey;
 
-							FIntRect HeightmapAtlasSubregion;
-							HeightmapAtlasSubregion.Min = SubsectionKey * InRenderInfo.SubsectionSizeQuads;
+							FIntRect AtlasTextureSubregion;
+							AtlasTextureSubregion.Min = SubsectionKey * InRenderInfo.SubsectionSizeQuads;
 							// We only really need the +1 on the very last subsection to get the last row/column, since we end up overwriting the other end
 							// rows/columns when we proceed to the next tile. However it's much easier to add the +1 here and do a small amount of duplicate
-							// writes, because otherwise we would have to adjust HeightmapSubregion to align with the region we're writing, which would get
+							// writes, because otherwise we would have to adjust SubsectionSubregion to align with the region we're writing, which would get
 							// messy in cases of different mip levels.
-							HeightmapAtlasSubregion.Max = HeightmapAtlasSubregion.Min + InRenderInfo.SubsectionSizeQuads + 1;
+							AtlasTextureSubregion.Max = AtlasTextureSubregion.Min + InRenderInfo.SubsectionSizeQuads + 1;
 
-							FIntRect HeightmapSubregion;
-							HeightmapSubregion.Min = HeightmapResourceSubregion->Subregion.Min + SubsectionLocalKey * SubsectionSubregionSize;
-							HeightmapSubregion.Max = HeightmapSubregion.Min + SubsectionSubregionSize;
+							FIntRect SubsectionSubregion;
+							SubsectionSubregion.Min = SourceTextureResourceSubregion->Subregion.Min + SubsectionLocalKey * SubsectionSubregionSize;
+							SubsectionSubregion.Max = SubsectionSubregion.Min + SubsectionSubregionSize;
 
-							FLandscapeMergeHeightmapsPS::FParameters* MergeHeightmapsPSParams = GraphBuilder.AllocParameters<FLandscapeMergeHeightmapsPS::FParameters>();
-							MergeHeightmapsPSParams->InHeightmapAtlasSubregion = FUintVector4(HeightmapAtlasSubregion.Min.X, HeightmapAtlasSubregion.Min.Y, HeightmapAtlasSubregion.Max.X, HeightmapAtlasSubregion.Max.Y);
-							MergeHeightmapsPSParams->InHeightmapSubregion = FUintVector4(HeightmapSubregion.Min.X, HeightmapSubregion.Min.Y, HeightmapSubregion.Max.X, HeightmapSubregion.Max.Y);
-							MergeHeightmapsPSParams->InHeightmap = *Heightmap;
-							MergeHeightmapsPSParams->RenderTargets[0] = AtlasTextureRT;
+							FLandscapeMergeTexturesPS::FParameters* MergeTexturesPSParams = GraphBuilder.AllocParameters<FLandscapeMergeTexturesPS::FParameters>();
+							MergeTexturesPSParams->InAtlasSubregion = FUintVector4(AtlasTextureSubregion.Min.X, AtlasTextureSubregion.Min.Y, AtlasTextureSubregion.Max.X, AtlasTextureSubregion.Max.Y);
+							MergeTexturesPSParams->InSourceTexture = *SourceTextureSRV;
+							MergeTexturesPSParams->InSourceTextureSubregion = FUintVector4(SubsectionSubregion.Min.X, SubsectionSubregion.Min.Y, SubsectionSubregion.Max.X, SubsectionSubregion.Max.Y);
+							check(InRenderInfo.bIsHeightmap || ((SourceTextureResourceSubregion->ChannelIndex >= 0) && (SourceTextureResourceSubregion->ChannelIndex < 4)));
+							MergeTexturesPSParams->InSourceTextureChannel = SourceTextureResourceSubregion->ChannelIndex;
+							MergeTexturesPSParams->RenderTargets[0] = AtlasTextureRT;
 
-							FLandscapeMergeHeightmapsPS::MergeHeightmap(GraphBuilder, MergeHeightmapsPSParams, HeightmapAtlasSubregion);
+							FLandscapeMergeTexturesPS::MergeTexture(GraphBuilder, MergeTexturesPSParams, AtlasTextureSubregion, InRenderInfo.bIsHeightmap);
 						}
 					}
 				}
-				else
-				{
-					// We can clear all subsections of a given component at once : 
-					FIntRect HeightmapAtlasSubregion;
-					HeightmapAtlasSubregion.Min = LandscapeComponentKey * ComponentSizeQuads;
-					HeightmapAtlasSubregion.Max = HeightmapAtlasSubregion.Min + ComponentSizeQuads;
-					FRDGTextureClearInfo ClearInfo;
-					ClearInfo.Viewport = HeightmapAtlasSubregion;
-					AddClearRenderTargetPass(GraphBuilder, AtlasTexture, ClearInfo);
-				}
 			}
 		}
 
 		{
+			FRDGTexture* OutputTexture = InOutputRenderTargetBinding.GetTexture();
+			check(OutputTexture != nullptr);
 			FIntVector RenderAreaSize = OutputTexture->Desc.GetSize();
-			FLandscapeResampleHeightmapsPS::FParameters* ResampleHeightmapsPSParams = GraphBuilder.AllocParameters<FLandscapeResampleHeightmapsPS::FParameters>();
-			ResampleHeightmapsPSParams->InOutputUVToMergedHeightmapUV = FMatrix44f(InRenderInfo.OutputUVToMergedHeightmapUV);
-			ResampleHeightmapsPSParams->InMergedHeightmap = AtlasTextureSRV;
-			ResampleHeightmapsPSParams->InMergedHeightmapSampler = TStaticSamplerState<SF_Bilinear, AM_Clamp, AM_Clamp, AM_Clamp>::GetRHI();
-			ResampleHeightmapsPSParams->InRenderAreaSize = FUintVector2((uint32)RenderAreaSize.X, (uint32)RenderAreaSize.Y);
-			ResampleHeightmapsPSParams->RenderTargets[0] = FRenderTargetBinding(OutputTexture, ERenderTargetLoadAction::ENoAction);;
+			FLandscapeResampleMergedTexturePS::FParameters* ResampleMergedTexturePSParams = GraphBuilder.AllocParameters<FLandscapeResampleMergedTexturePS::FParameters>();
+			ResampleMergedTexturePSParams->InOutputUVToMergedTextureUV = FMatrix44f(InRenderInfo.OutputUVToMergedTextureUV);
+			ResampleMergedTexturePSParams->InMergedTexture = AtlasTextureSRV;
+			ResampleMergedTexturePSParams->InMergedTextureSampler = TStaticSamplerState<SF_Bilinear, AM_Clamp, AM_Clamp, AM_Clamp>::GetRHI();
+			ResampleMergedTexturePSParams->InRenderAreaSize = FUintVector2((uint32)RenderAreaSize.X, (uint32)RenderAreaSize.Y);
+			ResampleMergedTexturePSParams->RenderTargets[0] = InOutputRenderTargetBinding;
 
 			// We now need to resample the atlas texture where the render area is : 
-			FLandscapeResampleHeightmapsPS::ResampleHeightmap(GraphBuilder, ResampleHeightmapsPSParams, InRenderInfo.bCompressHeight);
+			FLandscapeResampleMergedTexturePS::ResampleMergedTexture(GraphBuilder, ResampleMergedTexturePSParams, InRenderInfo.bIsHeightmap, InRenderInfo.bCompressHeight);
 		}
 	}
-}
+} // namespace UE::Landscape::Private::RenderMergedTexture_RenderThread
 
 bool ALandscape::IsValidRenderTargetFormatHeightmap(EPixelFormat InRenderTargetFormat, bool& bOutCompressHeight)
 {
@@ -3000,126 +3070,341 @@ bool ALandscape::IsValidRenderTargetFormatHeightmap(EPixelFormat InRenderTargetF
 	return false;
 }
 
-void ALandscape::RenderHeightmap(const FTransform& InRenderAreaWorldTransform, const FBox2D& InRenderAreaExtents, UTextureRenderTarget2D* OutRenderTarget)
+bool ALandscape::IsValidRenderTargetFormatWeightmap(EPixelFormat InRenderTargetFormat, int32& OutNumChannels)
+{
+	OutNumChannels = 0;
+	switch (InRenderTargetFormat)
+	{
+		// TODO [jonathan.bard] : for now, we only support 8 bits formats as they're the weightmap format but possibly we could handle the conversion to other formats
+	case PF_G8:
+	case PF_A8:
+	case PF_R8G8:
+	case PF_A8R8G8B8:
+	case PF_R8G8B8A8:
+	case PF_B8G8R8A8:
+	{
+		OutNumChannels = GPixelFormats[InRenderTargetFormat].NumComponents;
+		return true;
+	}
+	default:
+		break;
+	}
+
+	return false;
+}
+
+bool ALandscape::RenderMergedTextureInternal(const FTransform& InRenderAreaWorldTransform, const FBox2D& InRenderAreaExtents, const TArray<FName>& InWeightmapLayerNames, UTextureRenderTarget* OutRenderTarget)
 {
 	// TODO: We may want a version of this function that returns a lambda that can be passed to the render thread and run
 	// there to add the pass to an existing FRDGBuilder, in case the user wants this to be a part of a render graph with
-	// other passes. In that case RenderHeightmap would just use that function.
+	// other passes. In that case RenderMergedTextureInternal would just use that function.
 
-	TRACE_CPUPROFILER_EVENT_SCOPE(Landscape_RenderMergedHeightmap);
+	using namespace UE::Landscape;
+	using namespace UE::Landscape::Private;
 
-	RenderCaptureInterface::FScopedCapture RenderCapture((RenderCaptureNextHeightmapRenders != 0), TEXT("RenderHeightmapCapture"));
-	RenderCaptureNextHeightmapRenders = FMath::Max(RenderCaptureNextHeightmapRenders - 1, 0);
+	TRACE_CPUPROFILER_EVENT_SCOPE(Landscape_RenderMergedTextureInternal);
 
 	ULandscapeInfo* Info = GetLandscapeInfo();
 	if (Info == nullptr)
 	{
-		UE_LOG(LogLandscape, Error, TEXT("RenderHeightmap : Cannot render anything if there's no associated landscape info with this landscape (%s)"), *GetFullName());
-		return;
+		UE_LOG(LogLandscape, Error, TEXT("RenderMergedTexture : Cannot render anything if there's no associated landscape info with this landscape (%s)"), *GetFullName());
+		return false;
 	}
 
 	// Check render target validity :
 	if (OutRenderTarget == nullptr)
 	{
-		UE_LOG(LogLandscape, Error, TEXT("RenderHeightmap : Missing render target"));
-		return;
+		UE_LOG(LogLandscape, Error, TEXT("RenderMergedTexture : Missing render target"));
+		return false;
 	}
 
+	// Check Render target format :
+	const bool bIsHeightmap = InWeightmapLayerNames.IsEmpty();
 	bool bCompressHeight = false;
-	if (!IsValidRenderTargetFormatHeightmap(OutRenderTarget->GetFormat(), bCompressHeight))
+	UTextureRenderTarget2D* RenderTarget2D = Cast<UTextureRenderTarget2D>(OutRenderTarget);
+	UTextureRenderTarget2DArray* RenderTarget2DArray = Cast<UTextureRenderTarget2DArray>(OutRenderTarget);
+	EPixelFormat RenderTargetFormat = (RenderTarget2DArray != nullptr) ? RenderTarget2DArray->GetFormat() : (RenderTarget2D != nullptr) ? RenderTarget2D->GetFormat() : PF_Unknown;
+	if (bIsHeightmap)
 	{
-		UE_LOG(LogLandscape, Warning, TEXT("RenderHeightmap : invalid render target format for rendering heightmap (%s)"), GetPixelFormatString(OutRenderTarget->GetFormat()));
-		return;
+		if (RenderTarget2D == nullptr)
+		{
+			UE_LOG(LogLandscape, Error, TEXT("RenderMergedTexture : Heightmap capture requires a UTextureRenderTarget2D"));
+			return false;
+		}
+
+		if (!IsValidRenderTargetFormatHeightmap(RenderTargetFormat, bCompressHeight))
+		{
+			UE_LOG(LogLandscape, Warning, TEXT("RenderMergedTexture : invalid render target format for rendering heightmap (%s)"), GetPixelFormatString(RenderTargetFormat));
+			return false;
+		}
+	}
+	else
+	{
+		// If more than 1 weightmaps are requested, we expected a texture array or at the very least a texture 2D with enough channels to fit all weightmaps :
+		int32 NumChannels = 0;
+		if (!IsValidRenderTargetFormatWeightmap(RenderTargetFormat, NumChannels))
+		{
+			UE_LOG(LogLandscape, Warning, TEXT("RenderMergedTexture : invalid render target format for rendering weightmap (%s)"), GetPixelFormatString(RenderTargetFormat));
+			return false;
+		}
+
+		if (InWeightmapLayerNames.Num() > 1)
+		{
+			if ((RenderTarget2D != nullptr) && (NumChannels < InWeightmapLayerNames.Num()))
+			{
+				UE_LOG(LogLandscape, Warning, TEXT("RenderMergedTexture : Not enough channels available (%i) in render target to accomodate for all requested weightmaps (%i)"), NumChannels, InWeightmapLayerNames.Num());
+				return false;
+			}
+			else if ((RenderTarget2DArray != nullptr) && ((NumChannels * RenderTarget2DArray->Slices) < InWeightmapLayerNames.Num()))
+			{
+				UE_LOG(LogLandscape, Warning, TEXT("RenderMergedTexture : Not enough channels available (%i) in render target array to accomodate for all requested weightmaps (%i)"), NumChannels * RenderTarget2DArray->Slices, InWeightmapLayerNames.Num());
+				return false;
+			}
+		}
+	}
+
+	// If the requested extents are invalid, use the entire loaded landscape are as extents and transform : 
+	const FTransform& LandscapeTransform = GetTransform();
+	FTransform FinalRenderAreaWorldTransform = InRenderAreaWorldTransform;
+	FBox2D FinalRenderAreaExtents = InRenderAreaExtents;
+	if (!InRenderAreaExtents.bIsValid || InRenderAreaExtents.GetExtent().IsZero())
+	{
+		FinalRenderAreaWorldTransform = LandscapeTransform;
+		FBox LoadedBounds = Info->GetLoadedBounds();
+		FinalRenderAreaExtents = FBox2D(FVector2D(LandscapeTransform.InverseTransformPosition(LoadedBounds.Min)), FVector2D(LandscapeTransform.InverseTransformPosition(LoadedBounds.Max)));
 	}
 
 	// It can be helpful to visualize where the render happened so leave a visual log for that: 
-	UE_VLOG_OBOX(this, LogLandscape, Log, FBox(FVector(InRenderAreaExtents.Min, 0.0), FVector(InRenderAreaExtents.Max, 0.0)), InRenderAreaWorldTransform.ToMatrixWithScale(), FColor::Blue, TEXT(""));
+	UE_VLOG_OBOX(this, LogLandscape, Log, FBox(FVector(FinalRenderAreaExtents.Min, 0.0), FVector(FinalRenderAreaExtents.Max, 0.0)), FinalRenderAreaWorldTransform.ToMatrixWithScale(), FColor::Blue, TEXT("LandscapeRenderMergedTexture"));
 
 	// Don't do anything if this render area overlaps with no landscape component :
 	TMap<FIntPoint, ULandscapeComponent*> OverlappedComponents;
-	FIntRect ComponentIndicesBoundingRect;
-	if (!Info->GetOverlappedComponents(InRenderAreaWorldTransform, InRenderAreaExtents, OverlappedComponents, ComponentIndicesBoundingRect))
+	FIntRect OverlappedComponentIndicesBoundingRect;
+	if (!Info->GetOverlappedComponents(FinalRenderAreaWorldTransform, FinalRenderAreaExtents, OverlappedComponents, OverlappedComponentIndicesBoundingRect))
 	{
-		UE_LOG(LogLandscape, Log, TEXT("RenderHeightmap : no heightmap to render"));
-		return;
+		UE_LOG(LogLandscape, Log, TEXT("RenderMergedTexture : nothing to render"));
+		return true;
 	}
 
-	const FTransform& LandscapeTransform = GetTransform();
+	RenderCaptureInterface::FScopedCapture RenderCapture((RenderCaptureNextMergeRenders != 0), TEXT("LandscapeRenderMergedTextureCapture"));
+	RenderCaptureNextMergeRenders = FMath::Max(RenderCaptureNextMergeRenders - 1, 0);
+
+
+	// We'll want to perform one merge per target layer (i.e. as many as there are weightmaps, or just 1 in the case of heightmap) : 
+	const int32 NumTargetLayers = bIsHeightmap ? 1 : InWeightmapLayerNames.Num();
 
-	RenderMergedLandscape_RenderThread::FHeightmapRenderInfo HeightmapMergeRenderInfo;
-	// For now, merge the heightmap at max resolution :
-	HeightmapMergeRenderInfo.SubsectionSizeQuads = SubsectionSizeQuads;
-	HeightmapMergeRenderInfo.NumSubsections = NumSubsections;
-	HeightmapMergeRenderInfo.bCompressHeight = bCompressHeight;
+	TArray<RenderMergedTexture_RenderThread::FRenderInfo> MergeTextureRenderInfos;
+	MergeTextureRenderInfos.Reserve(NumTargetLayers);
 
-	for (auto It : OverlappedComponents)
+	for (int32 TargetLayerIndex = 0; TargetLayerIndex < NumTargetLayers; ++TargetLayerIndex)
 	{
-		ULandscapeComponent* Component = It.Value;
-		FIntPoint ComponentKey = It.Key;
-		UTexture2D* ComponentHeightmap = Component->GetHeightmap();
-		
-		// Get the subregion of the heightmap that this component uses (differs due to texture sharing).
-		// HeightmapScaleBias ZW values give us the offset of the component in a shared texture. You might think that XY would
-		// give the portion of the texture it occupies, but no, XY are 1/size, for some reason. Just calculate the subregion
-		// size ourselves.
-		int32 ComponentSize = Component->NumSubsections * (Component->SubsectionSizeQuads + 1);
-
-		FIntPoint HeightmapOffset(0, 0);
-		FTextureResource* HeightmapResource = ComponentHeightmap->GetResource();
-		if (ensure(HeightmapResource))
-		{
-			// We get the overall heightmap size via the resource instead of direct GetSizeX/Y calls because apparently
-			// the latter are unreliable while the texture is being built.
-			HeightmapOffset = FIntPoint(
-				FMath::RoundToInt32(Component->HeightmapScaleBias.Z * HeightmapResource->GetSizeX()),
-				FMath::RoundToInt32(Component->HeightmapScaleBias.W * HeightmapResource->GetSizeY()));
-		}
+		const FName TargetLayerName = bIsHeightmap ? FName(TEXT("Heightmap")) : InWeightmapLayerNames[TargetLayerIndex];
+
+		RenderMergedTexture_RenderThread::FRenderInfo& MergeTextureRenderInfo = MergeTextureRenderInfos.Emplace_GetRef();
+		// For now, merge the texture at max resolution :
+		MergeTextureRenderInfo.SubsectionSizeQuads = SubsectionSizeQuads;
+		MergeTextureRenderInfo.NumSubsections = NumSubsections;
+		MergeTextureRenderInfo.bIsHeightmap = bIsHeightmap;
+		MergeTextureRenderInfo.bCompressHeight = bCompressHeight;
+		MergeTextureRenderInfo.TargetLayerName = TargetLayerName;
+
+		// Indices of the components being rendered by this target layer : 
+		FIntRect RenderTargetComponentIndicesBoundingRect;
+
+		for (auto It : OverlappedComponents)
+		{
+			ULandscapeComponent* Component = It.Value;
+			FIntPoint ComponentKey = It.Key;
+
+			UTexture2D* SourceTexture = nullptr;
+			FVector4 SourceTextureScaleBias;
+			int32 SourceTextureChannel = INDEX_NONE;
+
+			if (bIsHeightmap)
+			{
+				SourceTexture = Component->GetHeightmap();
+				SourceTextureScaleBias = Component->HeightmapScaleBias;
+			}
+			else
+			{
+				const TArray<UTexture2D*>& WeightmapTextures = Component->GetWeightmapTextures();
+				const TArray<FWeightmapLayerAllocationInfo>& AllocInfos = Component->GetWeightmapLayerAllocations();
+				const FWeightmapLayerAllocationInfo* AllocInfo = AllocInfos.FindByPredicate([TargetLayerName](const FWeightmapLayerAllocationInfo& InAllocInfo) { return InAllocInfo.IsAllocated() && (InAllocInfo.GetLayerName() == TargetLayerName); });
+				if (AllocInfo != nullptr)
+				{
+					SourceTexture = WeightmapTextures[AllocInfo->WeightmapTextureIndex];
+					check(SourceTexture != nullptr);
+					SourceTextureScaleBias = Component->WeightmapScaleBias;
+					SourceTextureChannel = AllocInfo->WeightmapTextureChannel;
+				}
+			}
+			
+			if (SourceTexture != nullptr)
+			{
+				// Get the subregion of the source texture that this component uses (differs due to texture sharing).
+				// SourceTextureScaleBias ZW values give us the offset of the component in a shared texture. You might think that XY would
+				// give the portion of the texture it occupies, but no, XY are 1/size, for some reason. Just calculate the subregion
+				// size ourselves.
+				int32 ComponentSize = Component->NumSubsections * (Component->SubsectionSizeQuads + 1);
+
+				FIntPoint SourceTextureOffset(0, 0);
+				FTextureResource* SourceTextureResource = SourceTexture->GetResource();
+				if (ensure(SourceTextureResource))
+				{
+					// We get the overall source texture size via the resource instead of direct GetSizeX/Y calls because the latter are unreliable while the texture is being built.
+					SourceTextureOffset = FIntPoint(
+						FMath::FloorToInt32(SourceTextureScaleBias.Z * SourceTextureResource->GetSizeX()),
+						FMath::FloorToInt32(SourceTextureScaleBias.W * SourceTextureResource->GetSizeY()));
+				}
 			
-		// When mips are partially loaded, we need to take that into consideration when merging the heightmap :
-		uint32 MipBias = ComponentHeightmap->GetNumMips() - ComponentHeightmap->GetNumResidentMips();
-
-		// Theoretically speaking, all of our component heightmaps should be powers of two when we include the duplicated
-		// rows/columns across subsections, so we shouldn't get weird truncation results here...
-		HeightmapOffset.X >>= MipBias;
-		HeightmapOffset.Y >>= MipBias;
-		ComponentSize >>= MipBias;
-
-		// Effective area of the texture affecting this component (because of texture sharing) :
-		FIntRect HeightmapSubregion(HeightmapOffset, HeightmapOffset + ComponentSize);
-		HeightmapMergeRenderInfo.ComponentHeightmapsToRender.Add(ComponentKey, FTexture2DResourceSubregion(ComponentHeightmap->GetResource()->GetTexture2DResource(), HeightmapSubregion));
-	}
-
-	// Create the transform that will go from output target UVs to world space: 
-	FVector OutputUVOrigin = InRenderAreaWorldTransform.TransformPosition(FVector(InRenderAreaExtents.Min.X, InRenderAreaExtents.Min.Y, 0.0));
-	FVector OutputUVScale = InRenderAreaWorldTransform.GetScale3D() * FVector(InRenderAreaExtents.GetSize(), 1.0);
-	FTransform OutputUVToWorld(InRenderAreaWorldTransform.GetRotation(), OutputUVOrigin, OutputUVScale);
-
-	// Create the transform that will go from merged heightmap UVs to world space. Note that this is slightly trickier because
-	// vertices in the landscape correspond to pixel centers. So UV (0,0) is not at the minimal landscape vertex, but is instead 
-	// half a quad further (one pixel is one quad in size, so the center of the first pixel ends up at the minimal vertex).
-	// For related reasons, the size of the merged heightmap in world coordinates is actually one quad bigger in each direction.
-	check((ComponentIndicesBoundingRect.Min.X < ComponentIndicesBoundingRect.Max.X) && (ComponentIndicesBoundingRect.Min.Y < ComponentIndicesBoundingRect.Max.Y));
-	FVector MergedHeightmapScale = (FVector(ComponentIndicesBoundingRect.Max - ComponentIndicesBoundingRect.Min) * static_cast<double>(ComponentSizeQuads) + 1) 
-		* LandscapeTransform.GetScale3D();
-	MergedHeightmapScale.Z = 1.0f;
-	FVector MergedHeightmapUVOrigin = LandscapeTransform.TransformPosition(FVector(ComponentIndicesBoundingRect.Min) * (double)ComponentSizeQuads - FVector(0.5,0.5,0));
-	FTransform MergedHeightmapUVToWorld(LandscapeTransform.GetRotation(), MergedHeightmapUVOrigin, MergedHeightmapScale);
-
-	HeightmapMergeRenderInfo.OutputUVToMergedHeightmapUV = OutputUVToWorld.ToMatrixWithScale() * MergedHeightmapUVToWorld.ToInverseMatrixWithScale();
-
-	// Extract the render thread version of the render target :
-	FTextureRenderTarget2DResource* OutputRenderTargetResource = OutRenderTarget->GameThread_GetRenderTargetResource()->GetTextureRenderTarget2DResource();
+				// When mips are partially loaded, we need to take that into consideration when merging the source texture :
+				uint32 MipBias = SourceTexture->GetNumMips() - SourceTexture->GetNumResidentMips();
+
+				// Theoretically speaking, all of our component source textures should be powers of two when we include the duplicated
+				// rows/columns across subsections, so we shouldn't get weird truncation results here...
+				SourceTextureOffset.X >>= MipBias;
+				SourceTextureOffset.Y >>= MipBias;
+				ComponentSize >>= MipBias;
+
+				// Effective area of the texture affecting this component (because of texture sharing):
+				FIntRect SourceTextureSubregion(SourceTextureOffset, SourceTextureOffset + ComponentSize);
+				MergeTextureRenderInfo.ComponentTexturesToRender.Add(ComponentKey, FTexture2DResourceSubregion(SourceTextureResource->GetTexture2DResource(), SourceTextureSubregion, SourceTextureChannel));
+
+				// Since this component will be rendered in the render target, we can now expand the render target's bounds :
+				RenderTargetComponentIndicesBoundingRect.Union(FIntRect(ComponentKey, ComponentKey + FIntPoint(1)));
+			}
+		}
+
+		// Create the transform that will go from output target UVs to world space: 
+		FVector OutputUVOrigin = FinalRenderAreaWorldTransform.TransformPosition(FVector(FinalRenderAreaExtents.Min.X, FinalRenderAreaExtents.Min.Y, 0.0));
+		FVector OutputUVScale = FinalRenderAreaWorldTransform.GetScale3D() * FVector(FinalRenderAreaExtents.GetSize(), 1.0);
+		FTransform OutputUVToWorld(FinalRenderAreaWorldTransform.GetRotation(), OutputUVOrigin, OutputUVScale);
+
+		// Create the transform that will go from the merged texture (atlas) UVs to world space. Note that this is slightly trickier because
+		// vertices in the landscape correspond to pixel centers. So UV (0,0) is not at the minimal landscape vertex, but is instead 
+		// half a quad further (one pixel is one quad in size, so the center of the first pixel ends up at the minimal vertex).
+		// For related reasons, the size of the merged texture in world coordinates is actually one quad bigger in each direction.
+		check(RenderTargetComponentIndicesBoundingRect.IsEmpty()
+			|| (RenderTargetComponentIndicesBoundingRect.Min.X < RenderTargetComponentIndicesBoundingRect.Max.X) && (RenderTargetComponentIndicesBoundingRect.Min.Y < RenderTargetComponentIndicesBoundingRect.Max.Y));
+		FVector MergedTextureScale = (FVector(RenderTargetComponentIndicesBoundingRect.Max - RenderTargetComponentIndicesBoundingRect.Min) * static_cast<double>(ComponentSizeQuads) + 1)
+			* LandscapeTransform.GetScale3D();
+		MergedTextureScale.Z = 1.0f;
+		FVector MergedTextureUVOrigin = LandscapeTransform.TransformPosition(FVector(RenderTargetComponentIndicesBoundingRect.Min) * (double)ComponentSizeQuads - FVector(0.5,0.5,0));
+		FTransform MergedTextureUVToWorld(LandscapeTransform.GetRotation(), MergedTextureUVOrigin, MergedTextureScale);
+
+		MergeTextureRenderInfo.OutputUVToMergedTextureUV = OutputUVToWorld.ToMatrixWithScale() * MergedTextureUVToWorld.ToInverseMatrixWithScale();
+	}
+
+	// Extract the render thread version of the output render target :
+	FTextureRenderTargetResource* OutputRenderTargetResource = OutRenderTarget->GameThread_GetRenderTargetResource();
 	check(OutputRenderTargetResource != nullptr);
 
-	ENQUEUE_RENDER_COMMAND(RenderHeightmap)([HeightmapMergeRenderInfo, OutputRenderTargetResource](FRHICommandListImmediate& RHICmdList)
+	ENQUEUE_RENDER_COMMAND(RenderMergedTexture)([MergeTextureRenderInfos, OutputRenderTargetResource, bIsHeightmap, RenderTargetFormat](FRHICommandListImmediate& RHICmdList)
 	{
-		FRDGBuilder GraphBuilder(RHICmdList, RDG_EVENT_NAME("RenderHeightmap"));
+		FRDGBuilder GraphBuilder(RHICmdList, RDG_EVENT_NAME("RenderMergedTexture"));
+
+		FTextureRenderTarget2DResource* OutputRenderTarget2DResource = OutputRenderTargetResource->GetTextureRenderTarget2DResource();
+		FTextureRenderTarget2DArrayResource* OutputRenderTarget2DArrayResource = OutputRenderTargetResource->GetTextureRenderTarget2DArrayResource();
+		check((OutputRenderTarget2DResource != nullptr) || (OutputRenderTarget2DArrayResource != nullptr)); // either a render target 2D array or a render target 2D
 
-		FRDGTextureRef OutputHeightmap = GraphBuilder.RegisterExternalTexture(CreateRenderTarget(OutputRenderTargetResource->GetTextureRHI(), TEXT("MergedHeightmap")));
-		RenderMergedLandscape_RenderThread::RenderMergedHeightmap(HeightmapMergeRenderInfo, GraphBuilder, OutputHeightmap);
+		FRDGTextureRef OutputTexture = GraphBuilder.RegisterExternalTexture(CreateRenderTarget(OutputRenderTargetResource->GetTextureRHI(), TEXT("MergedTexture")));
+
+		// If we perform a single merge, we can simply render to the final texture : 
+		const int32 NumTargetLayers = MergeTextureRenderInfos.Num();
+		if (NumTargetLayers == 1)
+		{
+			// If it's a texture array, we need to specify the slice index
+			const int16 ArraySlice = (OutputRenderTarget2DArrayResource != nullptr) ? 0 : -1;
+			FRenderTargetBinding RenderTargetBinding(OutputTexture, ERenderTargetLoadAction::ENoAction, /*InMipIndex = */0, ArraySlice);
+			RenderMergedTexture_RenderThread::RenderMergedTexture(MergeTextureRenderInfos[0], GraphBuilder, RenderTargetBinding);
+		}
+		// In the case of multiple target layers, we'll render them one by one and pack them on the available output channels : 
+		else
+		{
+			const int32 NumChannels = GPixelFormats[RenderTargetFormat].NumComponents;
+			const int32 NumChannelPackingOperations = FMath::DivideAndRoundUp<int32>(NumTargetLayers, NumChannels);
+			check(NumChannelPackingOperations > 0);
+			checkf((OutputRenderTarget2DArrayResource != nullptr) || (NumTargetLayers <= NumChannels), TEXT("Trying to merge %i weightmaps onto a 2D texture of %i channels only"), NumTargetLayers, NumChannels);
+			checkf(!bIsHeightmap, TEXT("We should only be able to merge multiple textures in the case of weightmaps"));
+			FRDGTextureSRVRef DummyBlackTextureSRV = GraphBuilder.CreateSRV(FRDGTextureSRVDesc::Create(GSystemTextures.GetBlackDummy(GraphBuilder)));
+
+			// We'll need temporary 1 channel-texture for each weightmap that will then be packed onto the needed channels. This is for weightmaps only so PF_G8 pixel format is what we need for 
+			FIntPoint OutputTextureSize(OutputTexture->Desc.GetSize().X, OutputTexture->Desc.GetSize().Y);
+			FRDGTextureDesc SingleChannelTextureDesc = FRDGTextureDesc::Create2D(OutputTextureSize, PF_G8, FClearValueBinding::Black, TexCreate_RenderTargetable | TexCreate_ShaderResource);
+			int32 TargetLayerIndex = 0; 
+			const int32 NumSlices = (OutputRenderTarget2DArrayResource != nullptr) ? OutputTexture->Desc.ArraySize : 1;
+			for (int32 SliceIndex = 0; SliceIndex < NumSlices; ++SliceIndex)
+			{
+				RDG_EVENT_SCOPE(GraphBuilder, "RenderMergedTexture Slice %d", SliceIndex);
+
+				// The last slice might have to render less than the actual number of channels of the texture : 
+				const int32 NumEffectiveChannels = FMath::Min(NumChannels, (NumTargetLayers - SliceIndex * NumChannels));
+				check((NumEffectiveChannels >= 0) && (NumEffectiveChannels <= NumChannels));
+				TArray<FRDGTextureRef, TInlineAllocator<4>> SingleChannelTextures;
+
+				// First, render the each channel independently : 
+				for (int32 ChannelIndex = 0; ChannelIndex < NumEffectiveChannels; ++ChannelIndex)
+				{
+					FRDGTextureRef& SingleChannelTexture = SingleChannelTextures.Add_GetRef(GraphBuilder.CreateTexture(SingleChannelTextureDesc, TEXT("LandscapeMergedTextureTargetLayer")));
+					FRenderTargetBinding SingleChannelRenderTargetBinding(SingleChannelTexture, ERenderTargetLoadAction::ENoAction);
+					RenderMergedTexture_RenderThread::RenderMergedTexture(MergeTextureRenderInfos[TargetLayerIndex], GraphBuilder, SingleChannelRenderTargetBinding);
+					// We have rendered a new target layer, move on to the next :
+					++TargetLayerIndex;
+					check(TargetLayerIndex <= NumTargetLayers);
+				}
+
+				// Now pack the channels directly to the final render target (slice)
+				{
+					FLandscapePackRGBAChannelsPS::FParameters* PackRGBAChannelsParams = GraphBuilder.AllocParameters<FLandscapePackRGBAChannelsPS::FParameters>();
+					PackRGBAChannelsParams->InNumChannels = NumEffectiveChannels;
+					for (int32 ChannelIndex = 0; ChannelIndex < 4; ++ChannelIndex)
+					{
+						PackRGBAChannelsParams->InSourceTextures[ChannelIndex] = (ChannelIndex < NumEffectiveChannels) ? GraphBuilder.CreateSRV(FRDGTextureSRVDesc::Create(SingleChannelTextures[ChannelIndex])) : DummyBlackTextureSRV;
+					}
+					const int16 ArraySlice = (OutputRenderTarget2DArrayResource != nullptr) ? SliceIndex : -1;
+					// If it's a texture 2D or a texture 2D array with individually targetable slices, we can pack directly using the slice's RTV :
+					if ((OutputRenderTarget2DArrayResource == nullptr) || EnumHasAnyFlags(OutputTexture->Desc.Flags, TexCreate_TargetArraySlicesIndependently))
+					{
+						FRenderTargetBinding RenderTargetBinding(OutputTexture, ERenderTargetLoadAction::ENoAction, /*InMipIndex = */0, ArraySlice);
+						PackRGBAChannelsParams->RenderTargets[0] = RenderTargetBinding;
+						FLandscapePackRGBAChannelsPS::PackRGBAChannels(GraphBuilder, PackRGBAChannelsParams, FIntRect(FIntPoint::ZeroValue, OutputTextureSize));
+					}
+					else
+					{
+						// Otherwise (2D array but with non-individually targetable slices), we need to render to another render target and use a copy : 
+						FRDGTextureDesc IntermediateRenderTargetDesc = FRDGTextureDesc::Create2D(OutputTextureSize, OutputTexture->Desc.Format, FClearValueBinding::Black, TexCreate_RenderTargetable);
+						FRDGTextureRef IntermediateRenderTarget = GraphBuilder.CreateTexture(IntermediateRenderTargetDesc, TEXT("PackedRGBASlice"));
+						FRenderTargetBinding RenderTargetBinding(IntermediateRenderTarget, ERenderTargetLoadAction::ENoAction, /*InMipIndex = */0, /*InArraySlice = */-1);
+						PackRGBAChannelsParams->RenderTargets[0] = RenderTargetBinding;
+						FLandscapePackRGBAChannelsPS::PackRGBAChannels(GraphBuilder, PackRGBAChannelsParams, FIntRect(FIntPoint::ZeroValue, OutputTextureSize));
+
+						FRHICopyTextureInfo CopyTextureInfo;
+						CopyTextureInfo.DestSliceIndex = ArraySlice;
+						AddCopyTexturePass(GraphBuilder, IntermediateRenderTarget, OutputTexture, CopyTextureInfo);
+					}
+				}
+			}
+		}
 
 		GraphBuilder.Execute();
 	});
+
+	return true;
+}
+
+bool ALandscape::RenderHeightmap(const FTransform& InRenderAreaWorldTransform, const FBox2D& InRenderAreaExtents, UTextureRenderTarget2D* OutRenderTarget)
+{
+	return RenderMergedTextureInternal(InRenderAreaWorldTransform, InRenderAreaExtents, /*InWeightmapLayerNames = */{}, OutRenderTarget);
+}
+
+bool ALandscape::RenderWeightmap(const FTransform& InRenderAreaWorldTransform, const FBox2D& InRenderAreaExtents, FName InWeightmapLayerName, UTextureRenderTarget2D* OutRenderTarget)
+{
+	return RenderMergedTextureInternal(InRenderAreaWorldTransform, InRenderAreaExtents, /*InWeightmapLayerNames = */{ InWeightmapLayerName }, OutRenderTarget);
+}
+
+bool ALandscape::RenderWeightmaps(const FTransform& InRenderAreaWorldTransform, const FBox2D& InRenderAreaExtents, const TArray<FName>& InWeightmapLayerNames, UTextureRenderTarget* OutRenderTarget)
+{
+	return RenderMergedTextureInternal(InRenderAreaWorldTransform, InRenderAreaExtents, InWeightmapLayerNames, OutRenderTarget);
 }
 
 #if WITH_EDITOR
diff --git a/Engine/Source/Runtime/Landscape/Private/LandscapeEdit.cpp b/Engine/Source/Runtime/Landscape/Private/LandscapeEdit.cpp
--- a/Engine/Source/Runtime/Landscape/Private/LandscapeEdit.cpp
+++ b/Engine/Source/Runtime/Landscape/Private/LandscapeEdit.cpp
@@ -8361,6 +8361,7 @@ bool ALandscapeProxy::LandscapeImportWeightmapFromRenderTarget(UTextureRenderTar
 
 bool ALandscapeProxy::LandscapeExportWeightmapToRenderTarget(UTextureRenderTarget2D* InRenderTarget, FName InLayerName)
 {
+	// TODO [jonathan.bard] : Not implemented... Is redundant with RenderWeightmap, which is available in engine as well, so it'd probably best to deprecate LandscapeExportWeightmapToRenderTarget, as well as LandscapeExportHeightmapToRenderTarget 
 	return false;
 }
 
diff --git a/Engine/Source/Runtime/Landscape/Private/LandscapeEditResources.h b/Engine/Source/Runtime/Landscape/Private/LandscapeEditResources.h
--- a/Engine/Source/Runtime/Landscape/Private/LandscapeEditResources.h
+++ b/Engine/Source/Runtime/Landscape/Private/LandscapeEditResources.h
@@ -65,20 +65,23 @@ private:
 };
 
 // ----------------------------------------------------------------------------------
-// Describes a texture and an associated subregion (for heightmap texture sharing):
+// Describes a texture and an associated subregion (for heightmap texture sharing), and, optionally, a single channel (for weightmap texture sharing) :
 struct FTexture2DResourceSubregion
 {
 	FTexture2DResourceSubregion() = default;
 
-	FTexture2DResourceSubregion(FTexture2DResource* InTexture, const FIntRect& InSubregion)
+	FTexture2DResourceSubregion(FTexture2DResource* InTexture, const FIntRect& InSubregion, int32 InChannelIndex = INDEX_NONE)
 		: Texture(InTexture)
 		, Subregion(InSubregion)
+		, ChannelIndex(InChannelIndex)
 	{}
 
 	FTexture2DResource* Texture = nullptr;
 	FIntRect Subregion;
+	int32 ChannelIndex = INDEX_NONE;
 };
 
+
 // ----------------------------------------------------------------------------------
 // Struct that helps tracking external textures in the render graph
 struct FLandscapeRDGTrackedTexture
