Add ability to override LOD settings on landscape via CVAR, allowing easy testing of new LOD settings in builds.

diff --git a/Engine/Source/Runtime/Landscape/Private/LandscapeRender.cpp b/Engine/Source/Runtime/Landscape/Private/LandscapeRender.cpp
--- a/Engine/Source/Runtime/Landscape/Private/LandscapeRender.cpp
+++ b/Engine/Source/Runtime/Landscape/Private/LandscapeRender.cpp
@@ -72,6 +72,35 @@ static void OnLODDistributionScaleChanged(IConsoleVariable* CVar)
 }
 #endif
 
+#if !UE_BUILD_SHIPPING
+float GLandscapeLOD0ScreenSizeOverride = -1.f;
+FAutoConsoleVariableRef CVarLandscapeLOD0ScreenSizeOverride(
+	TEXT("r.Landscape.Override.LOD0ScreenSize"),
+	GLandscapeLOD0ScreenSizeOverride,
+	TEXT("When > 0, force override the landscape LOD0ScreenSize property on all landscapes"),
+	FConsoleVariableDelegate::CreateStatic(&OnLODDistributionScaleChanged),
+	ECVF_Cheat
+);
+
+float GLandscapeLOD0DistributionOverride = -1.f;
+FAutoConsoleVariableRef CVarLandscapeLOD0DistributionOverride(
+	TEXT("r.Landscape.Override.LOD0Distribution"),
+	GLandscapeLOD0DistributionOverride,
+	TEXT("When > 0, force override the LOD0DistributionSetting property on all landscapes, and ignore r.LandscapeLOD0DistributionScale"),
+	FConsoleVariableDelegate::CreateStatic(&OnLODDistributionScaleChanged),
+	ECVF_Cheat
+);
+
+float GLandscapeLODDistributionOverride = -1.f;
+FAutoConsoleVariableRef CVarLandscapeLODDistributionOverride(
+	TEXT("r.Landscape.Override.LODDistribution"),
+	GLandscapeLODDistributionOverride,
+	TEXT("When > 0, force override the landscape LODDistributionSetting property on all landscapes, and ignore r.LandscapeLODDistributionScale"),
+	FConsoleVariableDelegate::CreateStatic(&OnLODDistributionScaleChanged),
+	ECVF_Cheat
+);
+#endif // !UE_BUILD_SHIPPING
+
 float GLandscapeLOD0DistributionScale = 1.f;
 FAutoConsoleVariableRef CVarLandscapeLOD0DistributionScale(
 	TEXT("r.LandscapeLOD0DistributionScale"),
@@ -1155,9 +1184,28 @@ FLandscapeComponentSceneProxy::FLandscapeComponentSceneProxy(ULandscapeComponent
 		HeightmapSubsectionOffsetV = ((float)(InComponent->SubsectionSizeQuads + 1) / (float)FMath::Max<int32>(1, HeightmapTexture->GetSizeY()));
 	}
 
-	float ScreenSizeRatioDivider = FMath::Max(InComponent->GetLandscapeProxy()->LOD0DistributionSetting * GLandscapeLOD0DistributionScale, 1.01f);
+	float LOD0ScreenSize = InComponent->GetLandscapeProxy()->LOD0ScreenSize;
+	float LOD0Distribution = InComponent->GetLandscapeProxy()->LOD0DistributionSetting * GLandscapeLOD0DistributionScale;
+	float LODDistribution = InComponent->GetLandscapeProxy()->LODDistributionSetting * GLandscapeLODDistributionScale;
+
+#if !UE_BUILD_SHIPPING
+	if (GLandscapeLOD0ScreenSizeOverride > 0.0)
+	{
+		LOD0ScreenSize = GLandscapeLOD0ScreenSizeOverride;
+	}
+	if (GLandscapeLOD0DistributionOverride > 0.0)
+	{
+		LOD0Distribution = GLandscapeLOD0DistributionOverride;
+	}
+	if (GLandscapeLODDistributionOverride > 0.0)
+	{
+		LODDistribution = GLandscapeLODDistributionOverride;
+	}
+#endif // !UE_BUILD_SHIPPING
+
+	float ScreenSizeRatioDivider = FMath::Max(LOD0Distribution, 1.01f);
 	// Cancel out so that landscape is not affected by r.StaticMeshLODDistanceScale
-	float CurrentScreenSizeRatio = InComponent->GetLandscapeProxy()->LOD0ScreenSize / CVarStaticMeshLODDistanceScale.GetValueOnAnyThread();
+	float CurrentScreenSizeRatio = LOD0ScreenSize / CVarStaticMeshLODDistanceScale.GetValueOnAnyThread();
 
 	LODScreenRatioSquared.AddUninitialized(MaxLOD + 1);
 
@@ -1166,7 +1214,7 @@ FLandscapeComponentSceneProxy::FLandscapeComponentSceneProxy(ULandscapeComponent
 	LODSettings.LOD0ScreenSizeSquared = FMath::Square(CurrentScreenSizeRatio);
 	CurrentScreenSizeRatio /= ScreenSizeRatioDivider;
 	LODSettings.LOD1ScreenSizeSquared = FMath::Square(CurrentScreenSizeRatio);
-	ScreenSizeRatioDivider = FMath::Max(InComponent->GetLandscapeProxy()->LODDistributionSetting * GLandscapeLODDistributionScale, 1.01f);
+	ScreenSizeRatioDivider = FMath::Max(LODDistribution, 1.01f);
 	LODSettings.LODOnePlusDistributionScalarSquared = FMath::Square(ScreenSizeRatioDivider);
 
 	// Other LODs
