Ensure Landscape BPBrush Normal Capture has the render target padding it requires during layer merge
This also gets rid of a lot of copy/pasted code in the render target allocation path in LandscapeEditLayers.cpp.

diff --git a/Engine/Source/Runtime/Landscape/Classes/Landscape.h b/Engine/Source/Runtime/Landscape/Classes/Landscape.h
--- a/Engine/Source/Runtime/Landscape/Classes/Landscape.h
+++ b/Engine/Source/Runtime/Landscape/Classes/Landscape.h
@@ -400,7 +400,9 @@ protected:
 
 private:
 	bool SupportsEditLayersLocalMerge();
-	void CreateLayersRenderingResource();
+	bool HasNormalCaptureBPBrushLayer();
+
+	void CreateLayersRenderingResource(bool bUseNormalCapture);
 	void PrepareEditLayersLocalMergeResources();
 	void UpdateLayersContent(bool bInWaitForStreaming = false, bool bInSkipMonitorLandscapeEdModeChanges = false, bool bIntermediateRender = false, bool bFlushRender = false);
 	void MonitorShaderCompilation();
@@ -465,7 +467,7 @@ private:
 
 	void CopyTexturePS(const FString& InSourceDebugName, FTextureResource* InSourceResource, const FString& InDestDebugName, FTextureResource* InDestResource) const;
 
-	void InitializeLayers();
+	void InitializeLayers(bool bUseNormalCapture);
 	
 	void PrintLayersDebugRT(const FString& InContext, UTextureRenderTarget2D* InDebugRT, uint8 InMipRender = 0, bool InOutputHeight = true, bool InOutputNormals = false) const;
 	void PrintLayersDebugTextureResource(const FString& InContext, FTextureResource* InTextureResource, uint8 InMipRender = 0, bool InOutputHeight = true, bool InOutputNormals = false) const;
@@ -564,7 +566,10 @@ private:
 
 	UPROPERTY(Transient)
 	bool bLandscapeLayersAreInitialized;
-	
+
+	UPROPERTY(Transient)
+	bool bLandscapeLayersAreInitializedForNormalCapture;
+
 	UPROPERTY(Transient)
 	bool bLandscapeLayersAreUsingLocalMerge;
 
diff --git a/Engine/Source/Runtime/Landscape/Private/LandscapeEditLayers.cpp b/Engine/Source/Runtime/Landscape/Private/LandscapeEditLayers.cpp
--- a/Engine/Source/Runtime/Landscape/Private/LandscapeEditLayers.cpp
+++ b/Engine/Source/Runtime/Landscape/Private/LandscapeEditLayers.cpp
@@ -1824,7 +1824,23 @@ bool ALandscape::SupportsEditLayersLocalMerge()
 	return !!LandscapeEditLayersLocalMerge;
 }
 
-void ALandscape::CreateLayersRenderingResource()
+bool ALandscape::HasNormalCaptureBPBrushLayer()
+{
+	ALandscape* Landscape = GetLandscapeActor();
+	check(Landscape);
+
+	for (const FLandscapeLayer& Layer : Landscape->LandscapeLayers)
+	{
+		if (Layer.bVisible && Algo::AnyOf(Layer.Brushes, [](const FLandscapeLayerBrush& Brush) { return (Brush.GetBrush() != nullptr) && Brush.GetBrush()->IsVisible() && Brush.GetBrush()->GetCaptureBoundaryNormals(); }))
+		{
+			return true;
+		}
+	}
+
+	return false;
+}
+
+void ALandscape::CreateLayersRenderingResource(bool bUseNormalCapture)
 {
 	ULandscapeInfo* Info = GetLandscapeInfo();
 	if (Info == nullptr)
@@ -1852,14 +1868,50 @@ void ALandscape::CreateLayersRenderingResource()
 	}
 	else
 	{
-		if (Landscape->HeightmapRTList.Num() == 0)
+		// Heightmap mip size
+		int32 ComponentVerts = ((SubsectionSizeQuads + 1) * NumSubsections);
+		int32 CurrentMipSizeX = ComponentVerts * ComponentCounts.X;
+		int32 CurrentMipSizeY = ComponentVerts * ComponentCounts.Y;
+
+		// when using normal capture, we require an extra padding of 2 pixels --
+		//   but we round up to a full component of padding because the mip downsample shader code currently requires perfect 2:1 ratios in resolutions between mips
+		//   and keeping CurrentMipSizeX/Y a multiple of ComponentVerts ensures that (ComponentVerts is guaranteed to be a power of 2) 
+		if (bUseNormalCapture)
+		{
+			int32 PaddingX = FMath::RoundUpToPowerOfTwo(CurrentMipSizeX) - CurrentMipSizeX;
+			int32 PaddingY = FMath::RoundUpToPowerOfTwo(CurrentMipSizeY) - CurrentMipSizeY;
+			if (PaddingX < 2)
+			{
+				CurrentMipSizeX += ComponentVerts;
+			}
+			if (PaddingY < 2)
+			{
+				CurrentMipSizeY += ComponentVerts;
+			}
+		}
+
+		bool bCreateFromScratch = (Landscape->HeightmapRTList.Num() == 0);
+		if (bCreateFromScratch)
 		{
 			Landscape->HeightmapRTList.Init(nullptr, (int32)EHeightmapRTType::HeightmapRT_Count);
+		}
 
-			int32 CurrentMipSizeX = ((SubsectionSizeQuads + 1) * NumSubsections) * ComponentCounts.X;
-			int32 CurrentMipSizeY = ((SubsectionSizeQuads + 1) * NumSubsections) * ComponentCounts.Y;
+		auto InitOrResizeRT = [](UTextureRenderTarget2D *RT, int32 ResX, int32 ResY, bool bInit)
+		{
+			if (bInit)
+			{
+				RT->InitAutoFormat(FMath::RoundUpToPowerOfTwo(ResX), FMath::RoundUpToPowerOfTwo(ResY));
+				RT->UpdateResourceImmediate(true);
+			}
+			else
+			{
+				RT->ResizeTarget(FMath::RoundUpToPowerOfTwo(ResX), FMath::RoundUpToPowerOfTwo(ResY));
+			}
+		};
 
-			for (int32 i = 0; i < (int32)EHeightmapRTType::HeightmapRT_Count; ++i)
+		for (int32 i = 0; i < (int32)EHeightmapRTType::HeightmapRT_Count; ++i)
+		{
+			if (bCreateFromScratch)
 			{
 				FText DisplayName = StaticEnum<EHeightmapRTType>()->GetDisplayValueAsText((EHeightmapRTType)i);
 				FName RTName = MakeUniqueObjectName(GetTransientPackage(), UTextureRenderTarget2D::StaticClass(), FName(*DisplayName.ToString()));
@@ -1869,61 +1921,39 @@ void ALandscape::CreateLayersRenderingResource()
 				Landscape->HeightmapRTList[i]->AddressX = TextureAddress::TA_Clamp;
 				Landscape->HeightmapRTList[i]->AddressY = TextureAddress::TA_Clamp;
 				Landscape->HeightmapRTList[i]->ClearColor = FLinearColor(0.0f, 0.0f, 0.0f, 0.0f);
-
-				if (i < (int32)EHeightmapRTType::HeightmapRT_Mip1) // Landscape size RT
-				{
-					Landscape->HeightmapRTList[i]->InitAutoFormat(FMath::RoundUpToPowerOfTwo(CurrentMipSizeX), FMath::RoundUpToPowerOfTwo(CurrentMipSizeY));
-				}
-				else // Mips
-				{
-					CurrentMipSizeX >>= 1;
-					CurrentMipSizeY >>= 1;
-					Landscape->HeightmapRTList[i]->InitAutoFormat(FMath::RoundUpToPowerOfTwo(CurrentMipSizeX), FMath::RoundUpToPowerOfTwo(CurrentMipSizeY));
-				}
-
-				Landscape->HeightmapRTList[i]->UpdateResourceImmediate(true);
-
-				// Only generate required mips RT
-				if (CurrentMipSizeX == ComponentCounts.X && CurrentMipSizeY == ComponentCounts.Y)
-				{
-					break;
-				}
 			}
-		}
-		else // Simply resize the render target
-		{
-			int32 CurrentMipSizeX = ((SubsectionSizeQuads + 1) * NumSubsections) * ComponentCounts.X;
-			int32 CurrentMipSizeY = ((SubsectionSizeQuads + 1) * NumSubsections) * ComponentCounts.Y;
 
-			for (int32 i = 0; i < (int32)EHeightmapRTType::HeightmapRT_Count; ++i)
+			if (i < (int32)EHeightmapRTType::HeightmapRT_Mip1) // Landscape size RT
 			{
-				if (i < (int32)EHeightmapRTType::HeightmapRT_Mip1) // Landscape size RT
-				{
-					Landscape->HeightmapRTList[i]->ResizeTarget(FMath::RoundUpToPowerOfTwo(CurrentMipSizeX), FMath::RoundUpToPowerOfTwo(CurrentMipSizeY));
-				}
-				else // Mips
-				{
-					CurrentMipSizeX >>= 1;
-					CurrentMipSizeY >>= 1;
-					Landscape->HeightmapRTList[i]->ResizeTarget(FMath::RoundUpToPowerOfTwo(CurrentMipSizeX), FMath::RoundUpToPowerOfTwo(CurrentMipSizeY));
-				}
+				InitOrResizeRT(Landscape->HeightmapRTList[i], CurrentMipSizeX, CurrentMipSizeY, bCreateFromScratch);
+			}
+			else // Mips
+			{
+				CurrentMipSizeX >>= 1;
+				CurrentMipSizeY >>= 1;
+				InitOrResizeRT(Landscape->HeightmapRTList[i], CurrentMipSizeX, CurrentMipSizeY, bCreateFromScratch);
+			}
 
-				// Only generate required mips RT
-				if (CurrentMipSizeX == ComponentCounts.X && CurrentMipSizeY == ComponentCounts.Y)
-				{
-					break;
-				}
+			// Only generate required mips RT
+			if (CurrentMipSizeX == ComponentCounts.X && CurrentMipSizeY == ComponentCounts.Y)
+			{
+				break;
 			}
 		}
 
-		if (Landscape->WeightmapRTList.Num() == 0)
+		// Weightmap mip size
+		CurrentMipSizeX = ((SubsectionSizeQuads + 1) * NumSubsections) * ComponentCounts.X;
+		CurrentMipSizeY = ((SubsectionSizeQuads + 1) * NumSubsections) * ComponentCounts.Y;
+		bCreateFromScratch = (Landscape->WeightmapRTList.Num() == 0);
+
+		if (bCreateFromScratch)
 		{
 			Landscape->WeightmapRTList.Init(nullptr, (int32)EWeightmapRTType::WeightmapRT_Count);
+		}
 
-			int32 CurrentMipSizeX = ((SubsectionSizeQuads + 1) * NumSubsections) * ComponentCounts.X;
-			int32 CurrentMipSizeY = ((SubsectionSizeQuads + 1) * NumSubsections) * ComponentCounts.Y;
-
-			for (int32 i = 0; i < (int32)EWeightmapRTType::WeightmapRT_Count; ++i)
+		for (int32 i = 0; i < (int32)EWeightmapRTType::WeightmapRT_Count; ++i)
+		{
+			if (bCreateFromScratch)
 			{
 				FText DisplayName = StaticEnum<EHeightmapRTType>()->GetDisplayValueAsText((EWeightmapRTType)i);
 				FName RTName = MakeUniqueObjectName(GetTransientPackage(), UTextureRenderTarget2D::StaticClass(), FName(*DisplayName.ToString()));
@@ -1935,52 +1965,30 @@ void ALandscape::CreateLayersRenderingResource()
 				Landscape->WeightmapRTList[i]->ClearColor = FLinearColor(0.0f, 0.0f, 0.0f, 0.0f);
 				Landscape->WeightmapRTList[i]->RenderTargetFormat = RTF_RGBA8;
 
-				if (i < (int32)EWeightmapRTType::WeightmapRT_Mip0) // Landscape size RT, only create the number of layer we have
-				{
-					Landscape->WeightmapRTList[i]->RenderTargetFormat = (i == (int32)EWeightmapRTType::WeightmapRT_Scratch_RGBA) ? RTF_RGBA8 : RTF_R8;
-					Landscape->WeightmapRTList[i]->InitAutoFormat(FMath::RoundUpToPowerOfTwo(CurrentMipSizeX), FMath::RoundUpToPowerOfTwo(CurrentMipSizeY));
-				}
-				else // Mips
-				{
-					Landscape->WeightmapRTList[i]->InitAutoFormat(FMath::RoundUpToPowerOfTwo(CurrentMipSizeX), FMath::RoundUpToPowerOfTwo(CurrentMipSizeY));
-
-					CurrentMipSizeX >>= 1;
-					CurrentMipSizeY >>= 1;
-				}
-
-				Landscape->WeightmapRTList[i]->UpdateResourceImmediate(true);
-
-				// Only generate required mips RT
-				if (CurrentMipSizeX < ComponentCounts.X && CurrentMipSizeY < ComponentCounts.Y)
+				// scratch 1/2/3 RTs are R8 format
+				if ((i >= (int32)EWeightmapRTType::WeightmapRT_Scratch1) &&
+					(i < (int32)EWeightmapRTType::WeightmapRT_Mip0))
 				{
-					break;
+					Landscape->WeightmapRTList[i]->RenderTargetFormat = RTF_R8;
 				}
 			}
-		}
-		else // Simply resize the render target
-		{
-			int32 CurrentMipSizeX = ((SubsectionSizeQuads + 1) * NumSubsections) * ComponentCounts.X;
-			int32 CurrentMipSizeY = ((SubsectionSizeQuads + 1) * NumSubsections) * ComponentCounts.Y;
 
-			for (int32 i = 0; i < (int32)EWeightmapRTType::WeightmapRT_Count; ++i)
+			if (i < (int32)EWeightmapRTType::WeightmapRT_Mip0)
 			{
-				if (i < (int32)EWeightmapRTType::WeightmapRT_Mip0) // Landscape size RT, only create the number of layer we have
-				{
-					Landscape->WeightmapRTList[i]->ResizeTarget(FMath::RoundUpToPowerOfTwo(CurrentMipSizeX), FMath::RoundUpToPowerOfTwo(CurrentMipSizeY));
-				}
-				else // Mips
-				{
-					Landscape->WeightmapRTList[i]->ResizeTarget(FMath::RoundUpToPowerOfTwo(CurrentMipSizeX), FMath::RoundUpToPowerOfTwo(CurrentMipSizeY));
+				InitOrResizeRT(Landscape->WeightmapRTList[i], CurrentMipSizeX, CurrentMipSizeY, bCreateFromScratch);
+			}
+			else // Mips
+			{
+				InitOrResizeRT(Landscape->WeightmapRTList[i], CurrentMipSizeX, CurrentMipSizeY, bCreateFromScratch);
 
-					CurrentMipSizeX >>= 1;
-					CurrentMipSizeY >>= 1;
-				}
+				CurrentMipSizeX >>= 1;
+				CurrentMipSizeY >>= 1;
+			}
 
-				// Only generate required mips RT
-				if (CurrentMipSizeX < ComponentCounts.X && CurrentMipSizeY < ComponentCounts.Y)
-				{
-					break;
-				}
+			// Only generate required mips RT
+			if (CurrentMipSizeX < ComponentCounts.X && CurrentMipSizeY < ComponentCounts.Y)
+			{
+				break;
 			}
 		}
 
@@ -8229,16 +8237,19 @@ void ALandscape::UpdateLayersContent(bool bInWaitForStreaming, bool bInSkipMonit
 	FLandscapeDirtyOnlyInModeScope DirtyOnlyInMode(LandscapeInfo);
 
 	// If we went from local merge to global merge or vice versa, we need to reinitialize layers : 
+	// If we toggled normal capture, we need to reinitialize layers
 	bool bUseLocalMerge = SupportsEditLayersLocalMerge();
-	if (bLandscapeLayersAreUsingLocalMerge != bUseLocalMerge)
+	bool bUseNormalCapture = HasNormalCaptureBPBrushLayer();
+	if ((bLandscapeLayersAreUsingLocalMerge != bUseLocalMerge) || (bLandscapeLayersAreInitializedForNormalCapture != bUseNormalCapture))
 	{
 		RequestLayersInitialization(/*bInRequestContentUpdate = */true);
 		bLandscapeLayersAreUsingLocalMerge = bUseLocalMerge;
+		bLandscapeLayersAreInitializedForNormalCapture = bUseNormalCapture;
 	}
 
 	if (!bLandscapeLayersAreInitialized)
 	{
-		InitializeLayers();
+		InitializeLayers(bUseNormalCapture);
 	}
 
 	if (!bInSkipMonitorLandscapeEdModeChanges)
@@ -8591,11 +8602,11 @@ uint32 ALandscape::UpdateAfterReadbackResolves(const TArrayView<FLandscapeEditLa
 	return NewUpdateFlags;
 }
 
-void ALandscape::InitializeLayers()
+void ALandscape::InitializeLayers(bool bUseNormalCapture)
 {
 	check(HasLayersContent());
 
-	CreateLayersRenderingResource();
+	CreateLayersRenderingResource(bUseNormalCapture);
 	InitializeLandscapeLayersWeightmapUsage();
 
 	bLandscapeLayersAreInitialized = true;
