Splitting out LandscapeGrassWeightExporter and LandscapeGrassMapBuilder to their own files.
(Mainly so that follow up CLs implementing runtime grass map generation are simpler and easier to review)

No functionality was changed, other than adding the AddGrassWeightShaderTypes() function
so we don't have to expose the landscape weightmap shaders to the rest of the landscape code.

diff --git a/Engine/Source/Runtime/Landscape/Classes/LandscapeProxy.h b/Engine/Source/Runtime/Landscape/Classes/LandscapeProxy.h
--- a/Engine/Source/Runtime/Landscape/Classes/LandscapeProxy.h
+++ b/Engine/Source/Runtime/Landscape/Classes/LandscapeProxy.h
@@ -1494,21 +1494,6 @@ private:
 
 
 #if WITH_EDITOR
-/**
- * Helper class used to Build or monitor outdated Grass maps of a world
- */
-class FLandscapeGrassMapsBuilder
-{
-public:
-	LANDSCAPE_API FLandscapeGrassMapsBuilder(UWorld* InWorld);
-	LANDSCAPE_API void Build();
-	LANDSCAPE_API int32 GetOutdatedGrassMapCount(bool bInForceUpdate = true) const;
-private:
-	UWorld* World;
-	mutable int32 OutdatedGrassMapCount;
-	mutable double GrassMapsLastCheckTime;
-};
-
 /**
  * Helper class used to Build or monitor Landscape GI Textures
  */
diff --git a/Engine/Source/Runtime/Landscape/Private/LandscapeGrass.cpp b/Engine/Source/Runtime/Landscape/Private/LandscapeGrass.cpp
--- a/Engine/Source/Runtime/Landscape/Private/LandscapeGrass.cpp
+++ b/Engine/Source/Runtime/Landscape/Private/LandscapeGrass.cpp
@@ -50,6 +50,7 @@
 #include "ShaderParameterUtils.h"
 #include "EngineModule.h"
 #include "LandscapeRender.h"
+#include "LandscapeGrassWeightExporter.h"
 #include "MaterialCompiler.h"
 #include "Algo/Accumulate.h"
 #include "Algo/ForEach.h"
@@ -847,14 +848,14 @@ static void GrassCVarSinkFunction()
 		{
 			if (!World->bIsTearingDown)
 			{
-			if (ULandscapeSubsystem* LandscapeSubsystem = World->GetSubsystem<ULandscapeSubsystem>())
-			{
-				LandscapeSubsystem->RegenerateGrass(/*bInFlushGrass = */true, /*bInForceSync = */true);
+				if (ULandscapeSubsystem* LandscapeSubsystem = World->GetSubsystem<ULandscapeSubsystem>())
+				{
+					LandscapeSubsystem->RegenerateGrass(/*bInFlushGrass = */true, /*bInForceSync = */true);
+				}
 			}
 		}
 	}
 }
-}
 
 static FAutoConsoleVariableSink CVarGrassSink(FConsoleCommandDelegate::CreateStatic(&GrassCVarSinkFunction));
 
@@ -863,674 +864,6 @@ static FAutoConsoleVariableSink CVarGrassSink(FConsoleCommandDelegate::CreateSta
 //
 
 #if WITH_EDITOR
-static bool ShouldCacheLandscapeGrassShaders(const FMeshMaterialShaderPermutationParameters& Parameters)
-{
-	// We only need grass weight shaders for Landscape vertex factories on desktop platforms
-	return (Parameters.MaterialParameters.bIsUsedWithLandscape || Parameters.MaterialParameters.bIsSpecialEngineMaterial) &&
-		IsFeatureLevelSupported(Parameters.Platform, ERHIFeatureLevel::SM5) &&
-		Parameters.VertexFactoryType == FindVertexFactoryType(FName(TEXT("FLandscapeFixedGridVertexFactory"), FNAME_Find)) &&
-		EnumHasAllFlags(Parameters.Flags, EShaderPermutationFlags::HasEditorOnlyData);
-}
-
-class FLandscapeGrassWeightShaderElementData : public FMeshMaterialShaderElementData
-{
-public:
-
-	int32 OutputPass;
-	FVector2f RenderOffset;
-};
-
-class FLandscapeGrassWeightVS : public FMeshMaterialShader
-{
-	DECLARE_SHADER_TYPE(FLandscapeGrassWeightVS, MeshMaterial);
-
-	LAYOUT_FIELD(FShaderParameter, RenderOffsetParameter);
-
-protected:
-
-	FLandscapeGrassWeightVS()
-	{}
-
-	FLandscapeGrassWeightVS(const FMeshMaterialShaderType::CompiledShaderInitializerType& Initializer)
-	: FMeshMaterialShader(Initializer)
-	{
-		RenderOffsetParameter.Bind(Initializer.ParameterMap, TEXT("RenderOffset"));
-		PassUniformBuffer.Bind(Initializer.ParameterMap, FSceneTextureUniformParameters::FTypeInfo::GetStructMetadata()->GetShaderVariableName());
-	}
-
-public:
-
-	static bool ShouldCompilePermutation(const FMeshMaterialShaderPermutationParameters& Parameters)
-	{
-		return ShouldCacheLandscapeGrassShaders(Parameters);
-	}
-
-	void GetShaderBindings(
-		const FScene* Scene,
-		ERHIFeatureLevel::Type FeatureLevel,
-		const FPrimitiveSceneProxy* PrimitiveSceneProxy,
-		const FMaterialRenderProxy& MaterialRenderProxy,
-		const FMaterial& Material,
-		const FMeshPassProcessorRenderState& DrawRenderState,
-		const FLandscapeGrassWeightShaderElementData& ShaderElementData,
-		FMeshDrawSingleShaderBindings& ShaderBindings) const
-	{
-		FMeshMaterialShader::GetShaderBindings(Scene, FeatureLevel, PrimitiveSceneProxy, MaterialRenderProxy, Material, DrawRenderState, ShaderElementData, ShaderBindings);
-
-		ShaderBindings.Add(RenderOffsetParameter, ShaderElementData.RenderOffset);
-	}
-};
-
-IMPLEMENT_MATERIAL_SHADER_TYPE(, FLandscapeGrassWeightVS, TEXT("/Engine/Private/LandscapeGrassWeight.usf"), TEXT("VSMain"), SF_Vertex);
-
-class FLandscapeGrassWeightPS : public FMeshMaterialShader
-{
-	DECLARE_SHADER_TYPE(FLandscapeGrassWeightPS, MeshMaterial);
-	LAYOUT_FIELD(FShaderParameter, OutputPassParameter);
-public:
-
-	static bool ShouldCompilePermutation(const FMeshMaterialShaderPermutationParameters& Parameters)
-	{
-		return ShouldCacheLandscapeGrassShaders(Parameters);
-	}
-
-	FLandscapeGrassWeightPS(const ShaderMetaType::CompiledShaderInitializerType& Initializer)
-	: FMeshMaterialShader(Initializer)
-	{
-		OutputPassParameter.Bind(Initializer.ParameterMap, TEXT("OutputPass"));
-		PassUniformBuffer.Bind(Initializer.ParameterMap, FSceneTextureUniformParameters::FTypeInfo::GetStructMetadata()->GetShaderVariableName());
-	}
-
-	FLandscapeGrassWeightPS()
-	{}
-
-	void GetShaderBindings(
-		const FScene* Scene,
-		ERHIFeatureLevel::Type FeatureLevel,
-		const FPrimitiveSceneProxy* PrimitiveSceneProxy,
-		const FMaterialRenderProxy& MaterialRenderProxy,
-		const FMaterial& Material,
-		const FMeshPassProcessorRenderState& DrawRenderState,
-		const FLandscapeGrassWeightShaderElementData& ShaderElementData,
-		FMeshDrawSingleShaderBindings& ShaderBindings) const
-	{
-		FMeshMaterialShader::GetShaderBindings(Scene, FeatureLevel, PrimitiveSceneProxy, MaterialRenderProxy, Material, DrawRenderState, ShaderElementData, ShaderBindings);
-
-		ShaderBindings.Add(OutputPassParameter, ShaderElementData.OutputPass);
-	}
-};
-
-IMPLEMENT_MATERIAL_SHADER_TYPE(, FLandscapeGrassWeightPS, TEXT("/Engine/Private/LandscapeGrassWeight.usf"), TEXT("PSMain"), SF_Pixel);
-
-class FLandscapeGrassWeightMeshProcessor : public FMeshPassProcessor
-{
-public:
-	FLandscapeGrassWeightMeshProcessor(const FScene* Scene, ERHIFeatureLevel::Type InFeatureLevel, const FSceneView* InViewIfDynamicMeshCommand, FMeshPassDrawListContext* InDrawListContext);
-
-	void AddMeshBatch(const FMeshBatch& RESTRICT MeshBatch,
-		uint64 BatchElementMask,
-		int32 NumPasses,
-		FVector2D ViewOffset,
-		float PassOffsetX,
-		int32 FirstHeightMipsPassIndex, 
-		const TArray<int32>& HeightMips,
-		const FPrimitiveSceneProxy* RESTRICT PrimitiveSceneProxy);
-
-	virtual void AddMeshBatch(const FMeshBatch& RESTRICT MeshBatch, uint64 BatchElementMask, const FPrimitiveSceneProxy* RESTRICT PrimitiveSceneProxy, int32 StaticMeshId = -1) override final
-	{
-		checkf(false, TEXT("Default AddMeshBatch can't be used as rendering requires extra parameters per pass."));
-	}
-
-
-private:
-	bool TryAddMeshBatch(
-		const FMeshBatch& RESTRICT MeshBatch,
-		uint64 BatchElementMask,
-		const FPrimitiveSceneProxy* RESTRICT PrimitiveSceneProxy,
-		int32 StaticMeshId,
-		const FMaterialRenderProxy& MaterialRenderProxy,
-		const FMaterial& MaterialResource,
-		int32 NumPasses,
-		FVector2D ViewOffset,
-		float PassOffsetX,
-		int32 FirstHeightMipsPassIndex,
-		const TArray<int32>& HeightMips);
-
-	bool Process(
-		const FMeshBatch& MeshBatch,
-		uint64 BatchElementMask,
-		const FPrimitiveSceneProxy* RESTRICT PrimitiveSceneProxy,
-		const FMaterialRenderProxy& RESTRICT MaterialRenderProxy,
-		const FMaterial& RESTRICT MaterialResource,
-		int32 NumPasses,
-		FVector2D ViewOffset,
-		float PassOffsetX,
-		int32 FirstHeightMipsPassIndex,
-		const TArray<int32>& HeightMips);
-
-	FMeshPassProcessorRenderState PassDrawRenderState;
-};
-
-FLandscapeGrassWeightMeshProcessor::FLandscapeGrassWeightMeshProcessor(const FScene* Scene, ERHIFeatureLevel::Type InFeatureLevel, const FSceneView* InViewIfDynamicMeshCommand, FMeshPassDrawListContext* InDrawListContext)
-	: FMeshPassProcessor(EMeshPass::Num, Scene, InFeatureLevel, InViewIfDynamicMeshCommand, InDrawListContext)
-{
-	PassDrawRenderState.SetBlendState(TStaticBlendState<>::GetRHI());
-	PassDrawRenderState.SetDepthStencilState(TStaticDepthStencilState<false, CF_Always>::GetRHI());
-}
-
-void FLandscapeGrassWeightMeshProcessor::AddMeshBatch(const FMeshBatch& RESTRICT MeshBatch,
-	uint64 BatchElementMask,
-	int32 NumPasses,
-	FVector2D ViewOffset,
-	float PassOffsetX,
-	int32 FirstHeightMipsPassIndex, 
-	const TArray<int32>& HeightMips, 
-	const FPrimitiveSceneProxy* RESTRICT PrimitiveSceneProxy)
-{
-	const FMaterialRenderProxy* MaterialRenderProxy = MeshBatch.MaterialRenderProxy;
-	while (MaterialRenderProxy)
-	{
-		const FMaterial* Material = MaterialRenderProxy->GetMaterialNoFallback(FeatureLevel);
-		if (Material)
-		{
-			if (TryAddMeshBatch(MeshBatch, BatchElementMask, PrimitiveSceneProxy, -1, *MaterialRenderProxy, *Material, NumPasses, ViewOffset, PassOffsetX, FirstHeightMipsPassIndex, HeightMips))
-			{
-				break;
-			}
-		}
-
-		MaterialRenderProxy = MaterialRenderProxy->GetFallback(FeatureLevel);
-	}
-}
-
-bool FLandscapeGrassWeightMeshProcessor::TryAddMeshBatch(
-	const FMeshBatch& RESTRICT MeshBatch,
-	uint64 BatchElementMask,
-	const FPrimitiveSceneProxy* RESTRICT PrimitiveSceneProxy,
-	int32 StaticMeshId,
-	const FMaterialRenderProxy& MaterialRenderProxy,
-	const FMaterial& MaterialResource,
-	int32 NumPasses,
-	FVector2D ViewOffset,
-	float PassOffsetX,
-	int32 FirstHeightMipsPassIndex,
-	const TArray<int32>& HeightMips)
-{
-	check(MeshBatch.VertexFactory != nullptr);
-	return Process(MeshBatch, BatchElementMask, PrimitiveSceneProxy, MaterialRenderProxy, MaterialResource, NumPasses, ViewOffset, PassOffsetX, FirstHeightMipsPassIndex, HeightMips);
-}
-
-bool FLandscapeGrassWeightMeshProcessor::Process(
-	const FMeshBatch& MeshBatch,
-	uint64 BatchElementMask,
-	const FPrimitiveSceneProxy* RESTRICT PrimitiveSceneProxy,
-	const FMaterialRenderProxy& RESTRICT MaterialRenderProxy,
-	const FMaterial& RESTRICT MaterialResource,
-	int32 NumPasses,
-	FVector2D ViewOffset,
-	float PassOffsetX,
-	int32 FirstHeightMipsPassIndex,
-	const TArray<int32>& HeightMips)
-{
-	const FVertexFactory* VertexFactory = MeshBatch.VertexFactory;
-
-	TMeshProcessorShaders<
-		FLandscapeGrassWeightVS,
-		FLandscapeGrassWeightPS> PassShaders;
-
-	FMaterialShaderTypes ShaderTypes;
-	ShaderTypes.AddShaderType<FLandscapeGrassWeightVS>();
-	ShaderTypes.AddShaderType<FLandscapeGrassWeightPS>();
-
-	FMaterialShaders Shaders;
-	if (!MaterialResource.TryGetShaders(ShaderTypes, VertexFactory->GetType(), Shaders))
-	{
-		return false;
-	}
-
-	Shaders.TryGetVertexShader(PassShaders.VertexShader);
-	Shaders.TryGetPixelShader(PassShaders.PixelShader);
-
-	const FMeshDrawingPolicyOverrideSettings OverrideSettings = ComputeMeshOverrideSettings(MeshBatch);
-	const ERasterizerFillMode MeshFillMode = ComputeMeshFillMode(MaterialResource, OverrideSettings);
-	const ERasterizerCullMode MeshCullMode = CM_None;
-
-	FLandscapeGrassWeightShaderElementData ShaderElementData;
-	ShaderElementData.InitializeMeshMaterialData(ViewIfDynamicMeshCommand, PrimitiveSceneProxy, MeshBatch, -1, true);
-
-	const FMeshDrawCommandSortKey SortKey = CalculateMeshStaticSortKey(PassShaders.VertexShader, PassShaders.PixelShader);
-
-	for (int32 PassIndex = 0; PassIndex < NumPasses; ++PassIndex)
-	{
-		ShaderElementData.OutputPass = (PassIndex >= FirstHeightMipsPassIndex) ? 0 : PassIndex;
-		ShaderElementData.RenderOffset = FVector2f(ViewOffset) + FVector2f(PassOffsetX * PassIndex, 0);	// LWC_TODO: Precision loss
-
-		uint64 Mask = (PassIndex >= FirstHeightMipsPassIndex) ? HeightMips[PassIndex - FirstHeightMipsPassIndex] : BatchElementMask;
-
-		BuildMeshDrawCommands(
-			MeshBatch,
-			Mask,
-			PrimitiveSceneProxy,
-			MaterialRenderProxy,
-			MaterialResource,
-			PassDrawRenderState,
-			PassShaders,
-			MeshFillMode,
-			MeshCullMode,
-			SortKey,
-			EMeshPassFeatures::Default,
-			ShaderElementData);
-	}
-
-	return true;
-}
-
-// data also accessible by render thread
-class FLandscapeGrassWeightExporter_RenderThread
-{
-	FLandscapeGrassWeightExporter_RenderThread(const TArray<int32>& InHeightMips)
-		: HeightMips(InHeightMips)
-		{
-	}
-
-	friend class FLandscapeGrassWeightExporter;
-
-public:
-	virtual ~FLandscapeGrassWeightExporter_RenderThread()
-	{}
-
-	struct FComponentInfo
-	{
-		TObjectPtr<ULandscapeComponent> Component = nullptr;
-		TArray<TObjectPtr<ULandscapeGrassType>> RequestedGrassTypes;
-		FVector2D ViewOffset = FVector2D::ZeroVector;
-		int32 PixelOffsetX = 0;
-		FLandscapeComponentSceneProxy* SceneProxy = nullptr;
-		int32 NumPasses = 0;
-		int32 FirstHeightMipsPassIndex = MAX_int32;
-
-		FComponentInfo(ULandscapeComponent* InComponent, bool bInNeedsGrassmap, bool bInNeedsHeightmap, const TArray<int32>& InHeightMips)
-			: Component(InComponent)
-			, SceneProxy((FLandscapeComponentSceneProxy*)InComponent->SceneProxy)
-		{
-			if (bInNeedsGrassmap)
-			{
-				RequestedGrassTypes = InComponent->GetGrassTypes();
-			}
-			int32 NumGrassMaps = RequestedGrassTypes.Num();
-			if (bInNeedsHeightmap || NumGrassMaps > 0)
-			{
-				NumPasses += FMath::DivideAndRoundUp(2 /* heightmap */ + NumGrassMaps, 4);
-			}
-			if (InHeightMips.Num() > 0)
-			{
-				FirstHeightMipsPassIndex = NumPasses;
-				NumPasses += InHeightMips.Num();
-			}
-		}
-	};
-
-	FSceneInterface* SceneInterface = nullptr;
-	FTextureRenderTarget2DResource* RenderTargetResource = nullptr;
-	TArray<FComponentInfo, TInlineAllocator<1>> ComponentInfos;
-	FIntPoint TargetSize;
-	TArray<int32> HeightMips;
-	float PassOffsetX;
-	FVector ViewOrigin;
-	FMatrix ViewRotationMatrix;
-	FMatrix ProjectionMatrix;
-
-	BEGIN_SHADER_PARAMETER_STRUCT(FLandscapeGrassPassParameters, )
-		SHADER_PARAMETER_STRUCT_REF(FViewUniformShaderParameters, View)
-		SHADER_PARAMETER_RDG_UNIFORM_BUFFER(FSceneUniformParameters, Scene)
-		SHADER_PARAMETER_STRUCT_INCLUDE(FInstanceCullingDrawParams, InstanceCullingDrawParams)
-		RENDER_TARGET_BINDING_SLOTS()
-	END_SHADER_PARAMETER_STRUCT()
-
-	void RenderLandscapeComponentToTexture_RenderThread(FRHICommandListImmediate& RHICmdList)
-	{
-		FRDGBuilder GraphBuilder(RHICmdList);
-
-		FSceneViewFamilyContext ViewFamily(FSceneViewFamily::ConstructionValues(nullptr, SceneInterface, FEngineShowFlags(ESFIM_Game))
-			.SetTime(FGameTime::GetTimeSinceAppStart()));
-
-		ViewFamily.LandscapeLODOverride = 0; // Force LOD render
-
-		// Ensure scene primitive rendering is valid (added primitives comitted, GPU-Scene updated, push/pop dynamic culling context).
-		FScenePrimitiveRenderingContextScopeHelper ScenePrimitiveRenderingContextScopeHelper(GetRendererModule().BeginScenePrimitiveRendering(GraphBuilder, &ViewFamily));
-		
-		FSceneViewInitOptions ViewInitOptions;
-		ViewInitOptions.SetViewRectangle(FIntRect(0, 0, TargetSize.X, TargetSize.Y));
-		ViewInitOptions.ViewOrigin = ViewOrigin;
-		ViewInitOptions.ViewRotationMatrix = ViewRotationMatrix;
-		ViewInitOptions.ProjectionMatrix = ProjectionMatrix;
-		ViewInitOptions.ViewFamily = &ViewFamily;
-
-		GetRendererModule().CreateAndInitSingleView(RHICmdList, &ViewFamily, &ViewInitOptions);
-		
-		const FSceneView* View = ViewFamily.Views[0];
-
-		FRDGTextureRef OutputTexture = GraphBuilder.RegisterExternalTexture(CreateRenderTarget(RenderTargetResource->GetTextureRHI(), TEXT("LandscapeGrass")));
-
-		auto* PassParameters = GraphBuilder.AllocParameters<FLandscapeGrassPassParameters>();
-		PassParameters->View = View->ViewUniformBuffer;
-		PassParameters->Scene = GetSceneUniformBufferRef(GraphBuilder, *View);
-		PassParameters->RenderTargets[0] = FRenderTargetBinding(OutputTexture, ERenderTargetLoadAction::EClear);
-
-		AddSimpleMeshPass(GraphBuilder, PassParameters, SceneInterface->GetRenderScene(), *View, nullptr, RDG_EVENT_NAME("LandscapeGrass"), View->UnscaledViewRect,
-			[&View, this](FDynamicPassMeshDrawListContext* DynamicMeshPassContext)
-			{
-				FLandscapeGrassWeightMeshProcessor PassMeshProcessor(
-					nullptr,
-					View->GetFeatureLevel(),
-					View,
-					DynamicMeshPassContext);
-
-				const uint64 DefaultBatchElementMask = 1 << 0; // LOD 0 only
-
-				for (auto& ComponentInfo : ComponentInfos)
-				{
-					if (ensure(ComponentInfo.SceneProxy))
-					{
-						const FMeshBatch& Mesh = ComponentInfo.SceneProxy->GetGrassMeshBatch();
-						Mesh.MaterialRenderProxy->UpdateUniformExpressionCacheIfNeeded(View->GetFeatureLevel());
-
-						PassMeshProcessor.AddMeshBatch(Mesh, DefaultBatchElementMask, ComponentInfo.NumPasses, ComponentInfo.ViewOffset, PassOffsetX, ComponentInfo.FirstHeightMipsPassIndex, HeightMips, ComponentInfo.SceneProxy);
-					}
-				}
-			});
-
-		GraphBuilder.Execute();
-	}
-};
-
-class FLandscapeGrassWeightExporter : public FLandscapeGrassWeightExporter_RenderThread
-{
-	TObjectPtr<ALandscapeProxy> LandscapeProxy;
-	int32 ComponentSizeVerts;
-	int32 SubsectionSizeQuads;
-	int32 NumSubsections;
-	TArray<TObjectPtr<ULandscapeGrassType>> GrassTypes;
-	TObjectPtr<UTextureRenderTarget2D> RenderTargetTexture;
-
-public:
-	FLandscapeGrassWeightExporter(ALandscapeProxy* InLandscapeProxy, TArrayView<ULandscapeComponent* const> InLandscapeComponents, bool bInNeedsGrassmap = true, bool bInNeedsHeightmap = true, const TArray<int32>& InHeightMips = {})
-		: FLandscapeGrassWeightExporter_RenderThread(InHeightMips)
-		, LandscapeProxy(InLandscapeProxy)
-		, ComponentSizeVerts(InLandscapeProxy->ComponentSizeQuads + 1)
-		, SubsectionSizeQuads(InLandscapeProxy->SubsectionSizeQuads)
-		, NumSubsections(InLandscapeProxy->NumSubsections)
-		, RenderTargetTexture(nullptr)
-	{
-		check(InLandscapeComponents.Num() > 0);
-		SceneInterface = InLandscapeComponents[0]->GetScene();
-
-		// todo: use a 2d target?
-		const int32 SingleTileWidth = ComponentSizeVerts;
-		TargetSize = FIntPoint(0, ComponentSizeVerts);
-
-		// First compute the total render target size and prepare ComponentInfos (each component has its own number of needed passes because some might have different materials, thus different grass types) :
-		ComponentInfos.Reserve(InLandscapeComponents.Num());
-		int32 CurrentPixelOffsetX = 0;
-		for (ULandscapeComponent* Component : InLandscapeComponents)
-		{
-			ensure(Component->SceneProxy);
-
-			FComponentInfo& ComponentInfo = ComponentInfos.Add_GetRef(FComponentInfo(Component, bInNeedsGrassmap, bInNeedsHeightmap, InHeightMips));
-			ComponentInfo.PixelOffsetX = CurrentPixelOffsetX;
-
-			CurrentPixelOffsetX += ComponentInfo.NumPasses * SingleTileWidth;
-		}
-		TargetSize.X = CurrentPixelOffsetX;
-
-		FIntPoint TargetSizeMinusOne(TargetSize - FIntPoint(1, 1));
-		PassOffsetX = 2.0f * (float)SingleTileWidth / (float)TargetSize.X;
-
-		// Then compute FComponentInfo's ViewOffset with the knowledge of the total render target size : 
-		for (FComponentInfo& ComponentInfo : ComponentInfos)
-		{
-			FIntPoint ComponentOffset = (ComponentInfo.Component->GetSectionBase() - LandscapeProxy->LandscapeSectionOffset);
-			FVector2D ViewOffset(-ComponentOffset.X, ComponentOffset.Y);
-			ViewOffset.X += ComponentInfo.PixelOffsetX;
-			ViewOffset /= (FVector2D(TargetSize) * 0.5f);
-			ComponentInfo.ViewOffset = ViewOffset;
-		}
-
-		// center of target area in world
-		FVector TargetCenter = LandscapeProxy->GetTransform().TransformPosition(FVector(TargetSizeMinusOne, 0.f)*0.5f);
-
-		// extent of target in world space
-		FVector TargetExtent = FVector(TargetSize, 0.0f)*LandscapeProxy->GetActorScale()*0.5f;
-
-		ViewOrigin = TargetCenter;
-		ViewRotationMatrix = FInverseRotationMatrix(LandscapeProxy->GetActorRotation());
-		ViewRotationMatrix *= FMatrix(FPlane(1.0f, 0.0f, 0.0f, 0.0f),
-							          FPlane(0.0f, -1.0f, 0.0f, 0.0f),
-							          FPlane(0.0f, 0.0f, -1.0f, 0.0f),
-							          FPlane(0.0f, 0.0f, 0.0f, 1.0f));
-
-		const float ZOffset = UE_OLD_WORLD_MAX;
-		ProjectionMatrix = FReversedZOrthoMatrix(
-			TargetExtent.X,
-			TargetExtent.Y,
-			0.5f / ZOffset,
-			ZOffset);
-
-		RenderTargetTexture = NewObject<UTextureRenderTarget2D>();
-		check(RenderTargetTexture);
-		RenderTargetTexture->ClearColor = FLinearColor::White;
-		RenderTargetTexture->TargetGamma = 1.0f;
-		const bool bForceLinearGamma = true;
-		RenderTargetTexture->InitCustomFormat(TargetSize.X, TargetSize.Y, PF_B8G8R8A8, bForceLinearGamma);
-		RenderTargetResource = RenderTargetTexture->GameThread_GetRenderTargetResource()->GetTextureRenderTarget2DResource();
-
-		UE::RenderCommandPipe::FSyncScope SyncScope;
-
-		// render
-		FLandscapeGrassWeightExporter_RenderThread* Exporter = this;
-		ENQUEUE_RENDER_COMMAND(FDrawSceneCommand)(
-			[Exporter](FRHICommandListImmediate& RHICmdList)
-			{
-				Exporter->RenderLandscapeComponentToTexture_RenderThread(RHICmdList);
-				FlushPendingDeleteRHIResources_RenderThread();
-			});
-	}
-
-	TMap<ULandscapeComponent*, TUniquePtr<FLandscapeComponentGrassData>, TInlineSetAllocator<1>> FetchResults()
-	{
-		TArray<FColor> Samples;
-		Samples.SetNumUninitialized(TargetSize.X*TargetSize.Y);
-
-		// Copy the contents of the remote texture to system memory
-		FReadSurfaceDataFlags ReadSurfaceDataFlags;
-		ReadSurfaceDataFlags.SetLinearToGamma(false);
-		RenderTargetResource->ReadPixels(Samples, ReadSurfaceDataFlags, FIntRect(0, 0, TargetSize.X, TargetSize.Y));
-
-		TMap<ULandscapeComponent*, TUniquePtr<FLandscapeComponentGrassData>, TInlineSetAllocator<1>> Results;
-		Results.Reserve(ComponentInfos.Num());
-		
-		// Local data will be moved in contiguous array at the end of export (to minimize slack waste)
-		TArray<uint16> HeightData;
-		TMap<ULandscapeGrassType*, TArray<uint8>> WeightData;
-
-		for (auto& ComponentInfo : ComponentInfos)
-		{
-			ULandscapeComponent* Component = ComponentInfo.Component;
-			ALandscapeProxy* Proxy = Component->GetLandscapeProxy();
-
-			TUniquePtr<FLandscapeComponentGrassData> NewGrassData = MakeUnique<FLandscapeComponentGrassData>(Component);
-
-			if (ComponentInfo.FirstHeightMipsPassIndex != MAX_int32)
-			{
-				HeightData.Empty(FMath::Square(ComponentSizeVerts));
-			}
-			else
-			{
-				HeightData.Empty(0);
-			}
-			NewGrassData->HeightMipData.Empty(HeightMips.Num());
-
-			WeightData.Empty();
-			TArray<TArray<uint8>*> GrassWeightArrays;
-			GrassWeightArrays.Empty(ComponentInfo.RequestedGrassTypes.Num());
-			for (auto GrassType : ComponentInfo.RequestedGrassTypes)
-			{
-				WeightData.Add(GrassType);
-			}
-
-			// need a second loop because the WeightData map will reallocate its arrays as grass types are added
-			for (auto GrassType : ComponentInfo.RequestedGrassTypes)
-			{
-				TArray<uint8>* DataArray = WeightData.Find(GrassType);
-				check(DataArray);
-				DataArray->Empty(FMath::Square(ComponentSizeVerts));
-				GrassWeightArrays.Add(DataArray);
-			}
-
-			// output debug bitmap
-#if UE_BUILD_DEBUG
-			static bool bOutputGrassBitmap = false;
-			if (bOutputGrassBitmap)
-			{
-				FString TempPath = FPaths::ScreenShotDir();
-				TempPath += TEXT("/GrassDebug");
-				IFileManager::Get().MakeDirectory(*TempPath, true);
-				FFileHelper::CreateBitmap(*(TempPath / "Grass"), TargetSize.X, TargetSize.Y, Samples.GetData(), nullptr, &IFileManager::Get(), nullptr, ComponentInfo.RequestedGrassTypes.Num() >= 2);
-			}
-#endif
-
-			for (int32 PassIdx = 0; PassIdx < ComponentInfo.NumPasses; PassIdx++)
-			{
-				FColor* SampleData = &Samples[ComponentInfo.PixelOffsetX + PassIdx*ComponentSizeVerts];
-				if (PassIdx < ComponentInfo.FirstHeightMipsPassIndex)
-				{
-					if (PassIdx == 0)
-					{
-						for (int32 y = 0; y < ComponentSizeVerts; y++)
-						{
-							for (int32 x = 0; x < ComponentSizeVerts; x++)
-							{
-								FColor& Sample = SampleData[x + y * TargetSize.X];
-								uint16 Height = (((uint16)Sample.R) << 8) + (uint16)(Sample.G);
-								HeightData.Add(Height);
-								if (ComponentInfo.RequestedGrassTypes.Num() > 0)
-								{
-									GrassWeightArrays[0]->Add(Sample.B);
-									if (ComponentInfo.RequestedGrassTypes.Num() > 1)
-									{
-										GrassWeightArrays[1]->Add(Sample.A);
-									}
-								}
-							}
-						}
-					}
-					else
-					{
-						for (int32 y = 0; y < ComponentSizeVerts; y++)
-						{
-							for (int32 x = 0; x < ComponentSizeVerts; x++)
-							{
-								FColor& Sample = SampleData[x + y * TargetSize.X];
-
-								int32 TypeIdx = PassIdx * 4 - 2;
-								GrassWeightArrays[TypeIdx++]->Add(Sample.R);
-								if (TypeIdx < ComponentInfo.RequestedGrassTypes.Num())
-								{
-									GrassWeightArrays[TypeIdx++]->Add(Sample.G);
-									if (TypeIdx < ComponentInfo.RequestedGrassTypes.Num())
-									{
-										GrassWeightArrays[TypeIdx++]->Add(Sample.B);
-										if (TypeIdx < ComponentInfo.RequestedGrassTypes.Num())
-										{
-											GrassWeightArrays[TypeIdx++]->Add(Sample.A);
-										}
-									}
-								}
-							}
-						}
-					}
-				}
-				else // PassIdx >= FirstHeightMipsPassIndex
-				{
-					const int32 Mip = HeightMips[PassIdx - ComponentInfo.FirstHeightMipsPassIndex];
-					int32 MipSizeVerts = NumSubsections * (SubsectionSizeQuads >> Mip);
-					TArray<uint16>& MipHeightData = NewGrassData->HeightMipData.Add(Mip);
-					for (int32 y = 0; y < MipSizeVerts; y++)
-					{
-						for (int32 x = 0; x < MipSizeVerts; x++)
-						{
-							FColor& Sample = SampleData[x + y * TargetSize.X];
-							uint16 Height = (((uint16)Sample.R) << 8) + (uint16)(Sample.G);
-							MipHeightData.Add(Height);
-						}
-					}
-				}
-			}
-
-			// remove null grass type if we had one (can occur if the node has null entries)
-			WeightData.Remove(nullptr);
-
-			// Remove any grass data that is entirely weight 0
-			for (auto Iter(WeightData.CreateIterator()); Iter; ++Iter)
-			{
-				if (Iter->Value.IndexOfByPredicate([&](const int8& Weight) { return Weight != 0; }) == INDEX_NONE)
-				{
-					Iter.RemoveCurrent();
-				}
-			}
-
-			NewGrassData->InitializeFrom(HeightData, WeightData);
-			Results.Add(Component, MoveTemp(NewGrassData));
-		}
-
-		return Results;
-	}
-
-	void ApplyResults()
-	{
-		TMap<ULandscapeComponent*, TUniquePtr<FLandscapeComponentGrassData>, TInlineSetAllocator<1>> NewGrassData = FetchResults();
-
-		for (auto&& GrassDataPair : NewGrassData)
-		{
-			ULandscapeComponent* Component = GrassDataPair.Key;
-			FLandscapeComponentGrassData* ComponentGrassData = GrassDataPair.Value.Release();
-			ALandscapeProxy* Proxy = Component->GetLandscapeProxy();
-
-			// Assign the new data (thread-safe)
-			Component->GrassData = MakeShareable(ComponentGrassData);
-			Component->GrassData->bIsDirty = true;
-
-			if (Proxy->bBakeMaterialPositionOffsetIntoCollision)
-			{
-				Component->DestroyCollisionData();
-				Component->UpdateCollisionData();
-			}
-		}
-	}
-
-	void AddReferencedObjects(UObject* InThis, FReferenceCollector& Collector)
-	{
-		if (RenderTargetTexture)
-		{
-			Collector.AddReferencedObject(RenderTargetTexture);
-		}
-
-		if (LandscapeProxy)
-		{
-			Collector.AddReferencedObject(LandscapeProxy);
-		}
-
-		for (auto& Info : ComponentInfos)
-		{
-			if (Info.Component)
-			{
-				Collector.AddReferencedObject(Info.Component);
-			}
-
-			Collector.AddReferencedObjects(Info.RequestedGrassTypes);
-		}
-	}
-};
 
 FLandscapeComponentGrassData::FLandscapeComponentGrassData(ULandscapeComponent* Component)
 	: GenerationHash(Component->ComputeGrassMapGenerationHash())
@@ -1547,11 +880,11 @@ void ULandscapeComponent::UpdateGrassTypes()
 		GrassTypes = Material->GetMaterial()->GetCachedExpressionData().GrassTypes;
 
 		for (const ULandscapeGrassType* GrassType : GrassTypes)
-{
+		{
 			if (GrassType != nullptr)
 			{
 				for (const FGrassVariety& GrassVariety : GrassType->GrassVarieties)
-	{
+				{
 					GrassTypesMaxDiscardDistance = FMath::Max((float)GrassVariety.EndCullDistance.GetValue(), GrassTypesMaxDiscardDistance);
 				}
 			}
@@ -1586,10 +919,10 @@ uint32 ULandscapeComponent::ComputeGrassMapGenerationHash() const
 		}
 
 	return Hash;
-		}
+}
 
 bool ULandscapeComponent::IsGrassMapOutdated() const
-		{
+{
 	return GrassData->HasValidData() ? (ComputeGrassMapGenerationHash() != GrassData->GenerationHash) : false;
 }
 
@@ -1613,8 +946,7 @@ bool ULandscapeComponent::CanRenderGrassMap() const
 
 	// We only need the GrassWeight shaders on the fixed grid vertex factory to render grass maps : 
 	FMaterialShaderTypes ShaderTypes;
-	ShaderTypes.AddShaderType<FLandscapeGrassWeightVS>();
-	ShaderTypes.AddShaderType<FLandscapeGrassWeightPS>();
+	UE::Landscape::Grass::AddGrassWeightShaderTypes(ShaderTypes);
 
 	FVertexFactoryType* LandscapeGrassVF = FindVertexFactoryType(FName(TEXT("FLandscapeFixedGridVertexFactory"), FNAME_Find));
 	if (!MaterialResource->HasShaders(ShaderTypes, LandscapeGrassVF))
@@ -3105,59 +2437,8 @@ void ALandscapeProxy::RemoveAllExclusionBoxes()
 	}
 }
 
-
 #if WITH_EDITOR
 
-FLandscapeGrassMapsBuilder::FLandscapeGrassMapsBuilder(UWorld* InOwner)
-	: World(InOwner)
-	, OutdatedGrassMapCount(0)
-	, GrassMapsLastCheckTime(0)
-{}
-
-
-void FLandscapeGrassMapsBuilder::Build()
-{
-	if (World)
-	{
-		int32 Count = GetOutdatedGrassMapCount();
-		FScopedSlowTask SlowTask(static_cast<float>(Count), (LOCTEXT("GrassMaps_BuildGrassMaps", "Building Grass maps")));
-		SlowTask.MakeDialog();
-
-		for (TActorIterator<ALandscapeProxy> ProxyIt(World); ProxyIt; ++ProxyIt)
-		{
-			ProxyIt->BuildGrassMaps(&SlowTask);
-		}
-	}
-}
-
-int32 FLandscapeGrassMapsBuilder::GetOutdatedGrassMapCount(bool bInForceUpdate) const
-{
-	if (World)
-	{
-		bool bUpdate = bInForceUpdate || GLandscapeEditModeActive;
-		if (!bUpdate)
-		{
-			double GrassMapsTimeNow = FPlatformTime::Seconds();
-			// Recheck every 20 secs to handle the case where levels may have been Streamed in/out
-			if ((GrassMapsTimeNow - GrassMapsLastCheckTime) > 20)
-			{
-				GrassMapsLastCheckTime = GrassMapsTimeNow;
-				bUpdate = true;
-			}
-		}
-
-		if (bUpdate)
-		{
-			OutdatedGrassMapCount = 0;
-			for (TActorIterator<ALandscapeProxy> ProxyIt(World); ProxyIt; ++ProxyIt)
-			{
-				OutdatedGrassMapCount += ProxyIt->GetOutdatedGrassMapCount();
-			}
-		}
-	}
-	return OutdatedGrassMapCount;
-}
-
 void ALandscapeProxy::BuildGrassMaps(FScopedSlowTask* InSlowTask)
 {
 	if (!HasAnyFlags(RF_ClassDefaultObject))
diff --git a/Engine/Source/Runtime/Landscape/Private/LandscapeGrassMapsBuilder.cpp b/Engine/Source/Runtime/Landscape/Private/LandscapeGrassMapsBuilder.cpp
new file mode 100644
--- /dev/null
+++ b/Engine/Source/Runtime/Landscape/Private/LandscapeGrassMapsBuilder.cpp
@@ -0,0 +1,66 @@
+// Copyright Epic Games, Inc. All Rights Reserved.
+
+#include "LandscapeGrassMapsBuilder.h"
+#include "Misc/ScopedSlowTask.h"
+#include "Landscape.h"
+#include "LandscapeProxy.h"
+#include "LandscapeGrassWeightExporter.h"
+#include "EngineUtils.h"
+
+#define LOCTEXT_NAMESPACE "Landscape"
+
+extern int32 GGrassEnable;
+
+
+#if WITH_EDITOR
+
+FLandscapeGrassMapsBuilder::FLandscapeGrassMapsBuilder(UWorld* InOwner)
+	: World(InOwner)
+	, OutdatedGrassMapCount(0)
+	, GrassMapsLastCheckTime(0)
+{}
+
+void FLandscapeGrassMapsBuilder::Build()
+{
+	if (World)
+	{
+		int32 Count = GetOutdatedGrassMapCount();
+		FScopedSlowTask SlowTask(static_cast<float>(Count), (LOCTEXT("GrassMaps_BuildGrassMaps", "Building Grass maps")));
+		SlowTask.MakeDialog();
+
+		for (TActorIterator<ALandscapeProxy> ProxyIt(World); ProxyIt; ++ProxyIt)
+		{
+			ProxyIt->BuildGrassMaps(&SlowTask);
+		}
+	}
+}
+
+int32 FLandscapeGrassMapsBuilder::GetOutdatedGrassMapCount(bool bInForceUpdate) const
+{
+	if (World)
+	{
+		bool bUpdate = bInForceUpdate || GLandscapeEditModeActive;
+		if (!bUpdate)
+		{
+			double GrassMapsTimeNow = FPlatformTime::Seconds();
+			// Recheck every 20 secs to handle the case where levels may have been Streamed in/out
+			if ((GrassMapsTimeNow - GrassMapsLastCheckTime) > 20)
+			{
+				GrassMapsLastCheckTime = GrassMapsTimeNow;
+				bUpdate = true;
+			}
+		}
+
+		if (bUpdate)
+		{
+			OutdatedGrassMapCount = 0;
+			for (TActorIterator<ALandscapeProxy> ProxyIt(World); ProxyIt; ++ProxyIt)
+			{
+				OutdatedGrassMapCount += ProxyIt->GetOutdatedGrassMapCount();
+			}
+		}
+	}
+	return OutdatedGrassMapCount;
+}
+
+#endif // WITH_EDITOR
\ No newline at end of file
diff --git a/Engine/Source/Runtime/Landscape/Private/LandscapeGrassMapsBuilder.h b/Engine/Source/Runtime/Landscape/Private/LandscapeGrassMapsBuilder.h
new file mode 100644
--- /dev/null
+++ b/Engine/Source/Runtime/Landscape/Private/LandscapeGrassMapsBuilder.h
@@ -0,0 +1,24 @@
+// Copyright Epic Games, Inc. All Rights Reserved.
+#pragma once
+
+#include "CoreMinimal.h"
+#include "LandscapeComponent.h"
+
+#if WITH_EDITOR
+
+/**
+ * Helper class used to Build or monitor outdated Grass maps of a world
+ */
+class FLandscapeGrassMapsBuilder
+{
+public:
+	LANDSCAPE_API FLandscapeGrassMapsBuilder(UWorld* InWorld);
+	LANDSCAPE_API void Build();
+	LANDSCAPE_API int32 GetOutdatedGrassMapCount(bool bInForceUpdate = true) const;
+private:
+	UWorld* World;
+	mutable int32 OutdatedGrassMapCount;
+	mutable double GrassMapsLastCheckTime;
+};
+
+#endif // WITH_EDITOR
diff --git a/Engine/Source/Runtime/Landscape/Private/LandscapeGrassWeightExporter.cpp b/Engine/Source/Runtime/Landscape/Private/LandscapeGrassWeightExporter.cpp
new file mode 100644
--- /dev/null
+++ b/Engine/Source/Runtime/Landscape/Private/LandscapeGrassWeightExporter.cpp
@@ -0,0 +1,624 @@
+// Copyright Epic Games, Inc. All Rights Reserved.
+
+#include "LandscapeGrassWeightExporter.h"
+#include "SceneRendererInterface.h"
+#include "DataDrivenShaderPlatformInfo.h"
+#include "Materials/Material.h"
+#include "LandscapeGrassType.h"
+#include "Engine/TextureRenderTarget2D.h"
+#include "EngineModule.h"
+#include "LandscapeRender.h"
+#include "MeshPassProcessor.h"
+#include "MeshPassProcessor.inl"
+#include "SimpleMeshDrawCommandPass.h"
+#include "TextureResource.h"
+
+class FLandscapeGrassWeightVS;
+class FLandscapeGrassWeightPS;
+class FLandscapeGrassWeightShaderElementData;
+
+#if WITH_EDITOR
+
+BEGIN_SHADER_PARAMETER_STRUCT(FLandscapeGrassPassParameters, )
+	SHADER_PARAMETER_STRUCT_REF(FViewUniformShaderParameters, View)
+	SHADER_PARAMETER_RDG_UNIFORM_BUFFER(FSceneUniformParameters, Scene)
+	SHADER_PARAMETER_STRUCT_INCLUDE(FInstanceCullingDrawParams, InstanceCullingDrawParams)
+	RENDER_TARGET_BINDING_SLOTS()
+	END_SHADER_PARAMETER_STRUCT()
+
+class FLandscapeGrassWeightMeshProcessor : public FMeshPassProcessor
+{
+public:
+	FLandscapeGrassWeightMeshProcessor(const FScene* Scene, ERHIFeatureLevel::Type InFeatureLevel, const FSceneView* InViewIfDynamicMeshCommand, FMeshPassDrawListContext* InDrawListContext);
+
+	void AddMeshBatch(const FMeshBatch& RESTRICT MeshBatch,
+		uint64 BatchElementMask,
+		int32 NumPasses,
+		FVector2D ViewOffset,
+		float PassOffsetX,
+		int32 FirstHeightMipsPassIndex,
+		const TArray<int32>& HeightMips,
+		const FPrimitiveSceneProxy* RESTRICT PrimitiveSceneProxy);
+
+	virtual void AddMeshBatch(const FMeshBatch& RESTRICT MeshBatch, uint64 BatchElementMask, const FPrimitiveSceneProxy* RESTRICT PrimitiveSceneProxy, int32 StaticMeshId = -1) override final
+	{
+		checkf(false, TEXT("Default AddMeshBatch can't be used as rendering requires extra parameters per pass."));
+	}
+
+private:
+	bool TryAddMeshBatch(
+		const FMeshBatch& RESTRICT MeshBatch,
+		uint64 BatchElementMask,
+		const FPrimitiveSceneProxy* RESTRICT PrimitiveSceneProxy,
+		int32 StaticMeshId,
+		const FMaterialRenderProxy& MaterialRenderProxy,
+		const FMaterial& MaterialResource,
+		int32 NumPasses,
+		FVector2D ViewOffset,
+		float PassOffsetX,
+		int32 FirstHeightMipsPassIndex,
+		const TArray<int32>& HeightMips);
+
+	bool Process(
+		const FMeshBatch& MeshBatch,
+		uint64 BatchElementMask,
+		const FPrimitiveSceneProxy* RESTRICT PrimitiveSceneProxy,
+		const FMaterialRenderProxy& RESTRICT MaterialRenderProxy,
+		const FMaterial& RESTRICT MaterialResource,
+		int32 NumPasses,
+		FVector2D ViewOffset,
+		float PassOffsetX,
+		int32 FirstHeightMipsPassIndex,
+		const TArray<int32>& HeightMips);
+
+	FMeshPassProcessorRenderState PassDrawRenderState;
+};
+
+FLandscapeGrassWeightMeshProcessor::FLandscapeGrassWeightMeshProcessor(const FScene* Scene, ERHIFeatureLevel::Type InFeatureLevel, const FSceneView* InViewIfDynamicMeshCommand, FMeshPassDrawListContext* InDrawListContext)
+	: FMeshPassProcessor(EMeshPass::Num, Scene, InFeatureLevel, InViewIfDynamicMeshCommand, InDrawListContext)
+{
+	PassDrawRenderState.SetBlendState(TStaticBlendState<>::GetRHI());
+	PassDrawRenderState.SetDepthStencilState(TStaticDepthStencilState<false, CF_Always>::GetRHI());
+}
+
+void FLandscapeGrassWeightMeshProcessor::AddMeshBatch(const FMeshBatch& RESTRICT MeshBatch,
+	uint64 BatchElementMask,
+	int32 NumPasses,
+	FVector2D ViewOffset,
+	float PassOffsetX,
+	int32 FirstHeightMipsPassIndex,
+	const TArray<int32>& HeightMips,
+	const FPrimitiveSceneProxy* RESTRICT PrimitiveSceneProxy)
+{
+	const FMaterialRenderProxy* MaterialRenderProxy = MeshBatch.MaterialRenderProxy;
+	while (MaterialRenderProxy)
+	{
+		const FMaterial* Material = MaterialRenderProxy->GetMaterialNoFallback(FeatureLevel);
+		if (Material)
+		{
+			if (TryAddMeshBatch(MeshBatch, BatchElementMask, PrimitiveSceneProxy, -1, *MaterialRenderProxy, *Material, NumPasses, ViewOffset, PassOffsetX, FirstHeightMipsPassIndex, HeightMips))
+			{
+				break;
+			}
+		}
+
+		MaterialRenderProxy = MaterialRenderProxy->GetFallback(FeatureLevel);
+	}
+}
+
+bool FLandscapeGrassWeightMeshProcessor::TryAddMeshBatch(
+	const FMeshBatch& RESTRICT MeshBatch,
+	uint64 BatchElementMask,
+	const FPrimitiveSceneProxy* RESTRICT PrimitiveSceneProxy,
+	int32 StaticMeshId,
+	const FMaterialRenderProxy& MaterialRenderProxy,
+	const FMaterial& MaterialResource,
+	int32 NumPasses,
+	FVector2D ViewOffset,
+	float PassOffsetX,
+	int32 FirstHeightMipsPassIndex,
+	const TArray<int32>& HeightMips)
+{
+	check(MeshBatch.VertexFactory != nullptr);
+	return Process(MeshBatch, BatchElementMask, PrimitiveSceneProxy, MaterialRenderProxy, MaterialResource, NumPasses, ViewOffset, PassOffsetX, FirstHeightMipsPassIndex, HeightMips);
+}
+
+bool FLandscapeGrassWeightMeshProcessor::Process(
+	const FMeshBatch& MeshBatch,
+	uint64 BatchElementMask,
+	const FPrimitiveSceneProxy* RESTRICT PrimitiveSceneProxy,
+	const FMaterialRenderProxy& RESTRICT MaterialRenderProxy,
+	const FMaterial& RESTRICT MaterialResource,
+	int32 NumPasses,
+	FVector2D ViewOffset,
+	float PassOffsetX,
+	int32 FirstHeightMipsPassIndex,
+	const TArray<int32>& HeightMips)
+{
+	const FVertexFactory* VertexFactory = MeshBatch.VertexFactory;
+
+	TMeshProcessorShaders<
+		FLandscapeGrassWeightVS,
+		FLandscapeGrassWeightPS> PassShaders;
+
+	FMaterialShaderTypes ShaderTypes;
+	ShaderTypes.AddShaderType<FLandscapeGrassWeightVS>();
+	ShaderTypes.AddShaderType<FLandscapeGrassWeightPS>();
+
+	FMaterialShaders Shaders;
+	if (!MaterialResource.TryGetShaders(ShaderTypes, VertexFactory->GetType(), Shaders))
+	{
+		return false;
+	}
+
+	Shaders.TryGetVertexShader(PassShaders.VertexShader);
+	Shaders.TryGetPixelShader(PassShaders.PixelShader);
+
+	const FMeshDrawingPolicyOverrideSettings OverrideSettings = ComputeMeshOverrideSettings(MeshBatch);
+	const ERasterizerFillMode MeshFillMode = ComputeMeshFillMode(MaterialResource, OverrideSettings);
+	const ERasterizerCullMode MeshCullMode = CM_None;
+
+	FLandscapeGrassWeightShaderElementData ShaderElementData;
+	ShaderElementData.InitializeMeshMaterialData(ViewIfDynamicMeshCommand, PrimitiveSceneProxy, MeshBatch, -1, true);
+
+	const FMeshDrawCommandSortKey SortKey = CalculateMeshStaticSortKey(PassShaders.VertexShader, PassShaders.PixelShader);
+
+	for (int32 PassIndex = 0; PassIndex < NumPasses; ++PassIndex)
+	{
+		ShaderElementData.OutputPass = (PassIndex >= FirstHeightMipsPassIndex) ? 0 : PassIndex;
+		ShaderElementData.RenderOffset = FVector2f(ViewOffset) + FVector2f(PassOffsetX * PassIndex, 0);	// LWC_TODO: Precision loss
+
+		uint64 Mask = (PassIndex >= FirstHeightMipsPassIndex) ? HeightMips[PassIndex - FirstHeightMipsPassIndex] : BatchElementMask;
+
+		BuildMeshDrawCommands(
+			MeshBatch,
+			Mask,
+			PrimitiveSceneProxy,
+			MaterialRenderProxy,
+			MaterialResource,
+			PassDrawRenderState,
+			PassShaders,
+			MeshFillMode,
+			MeshCullMode,
+			SortKey,
+			EMeshPassFeatures::Default,
+			ShaderElementData);
+	}
+
+	return true;
+}
+
+
+void FLandscapeGrassWeightExporter_RenderThread::RenderLandscapeComponentToTexture_RenderThread(FRHICommandListImmediate& RHICmdList)
+{
+	FRDGBuilder GraphBuilder(RHICmdList);
+
+	FSceneViewFamilyContext ViewFamily(FSceneViewFamily::ConstructionValues(nullptr, SceneInterface, FEngineShowFlags(ESFIM_Game))
+		.SetTime(FGameTime::GetTimeSinceAppStart()));
+
+	ViewFamily.LandscapeLODOverride = 0; // Force LOD render
+
+	// Ensure scene primitive rendering is valid (added primitives comitted, GPU-Scene updated, push/pop dynamic culling context).
+	FScenePrimitiveRenderingContextScopeHelper ScenePrimitiveRenderingContextScopeHelper(GetRendererModule().BeginScenePrimitiveRendering(GraphBuilder, &ViewFamily));
+
+	FSceneViewInitOptions ViewInitOptions;
+	ViewInitOptions.SetViewRectangle(FIntRect(0, 0, TargetSize.X, TargetSize.Y));
+	ViewInitOptions.ViewOrigin = ViewOrigin;
+	ViewInitOptions.ViewRotationMatrix = ViewRotationMatrix;
+	ViewInitOptions.ProjectionMatrix = ProjectionMatrix;
+	ViewInitOptions.ViewFamily = &ViewFamily;
+
+	GetRendererModule().CreateAndInitSingleView(RHICmdList, &ViewFamily, &ViewInitOptions);
+
+	const FSceneView* View = ViewFamily.Views[0];
+
+	FRDGTextureRef OutputTexture = GraphBuilder.RegisterExternalTexture(CreateRenderTarget(RenderTargetResource->GetTextureRHI(), TEXT("LandscapeGrass")));
+
+	auto* PassParameters = GraphBuilder.AllocParameters<FLandscapeGrassPassParameters>();
+	PassParameters->View = View->ViewUniformBuffer;
+	PassParameters->Scene = GetSceneUniformBufferRef(GraphBuilder, *View);
+	PassParameters->RenderTargets[0] = FRenderTargetBinding(OutputTexture, ERenderTargetLoadAction::EClear);
+
+	AddSimpleMeshPass(GraphBuilder, PassParameters, SceneInterface->GetRenderScene(), *View, nullptr, RDG_EVENT_NAME("LandscapeGrass"), View->UnscaledViewRect,
+		[&View, this](FDynamicPassMeshDrawListContext* DynamicMeshPassContext)
+		{
+			FLandscapeGrassWeightMeshProcessor PassMeshProcessor(
+				nullptr,
+				View->GetFeatureLevel(),
+				View,
+				DynamicMeshPassContext);
+
+			const uint64 DefaultBatchElementMask = 1 << 0; // LOD 0 only
+
+			for (auto& ComponentInfo : ComponentInfos)
+			{
+				if (ensure(ComponentInfo.SceneProxy))
+				{
+					const FMeshBatch& Mesh = ComponentInfo.SceneProxy->GetGrassMeshBatch();
+					Mesh.MaterialRenderProxy->UpdateUniformExpressionCacheIfNeeded(View->GetFeatureLevel());
+
+					PassMeshProcessor.AddMeshBatch(Mesh, DefaultBatchElementMask, ComponentInfo.NumPasses, ComponentInfo.ViewOffset, PassOffsetX, ComponentInfo.FirstHeightMipsPassIndex, HeightMips, ComponentInfo.SceneProxy);
+				}
+			}
+		});
+
+	GraphBuilder.Execute();
+}
+
+
+FLandscapeGrassWeightExporter::FLandscapeGrassWeightExporter(ALandscapeProxy* InLandscapeProxy, TArrayView<ULandscapeComponent* const> InLandscapeComponents, bool bInNeedsGrassmap, bool bInNeedsHeightmap, const TArray<int32>& InHeightMips)
+	: FLandscapeGrassWeightExporter_RenderThread(InHeightMips)
+	, LandscapeProxy(InLandscapeProxy)
+	, ComponentSizeVerts(InLandscapeProxy->ComponentSizeQuads + 1)
+	, SubsectionSizeQuads(InLandscapeProxy->SubsectionSizeQuads)
+	, NumSubsections(InLandscapeProxy->NumSubsections)
+	, RenderTargetTexture(nullptr)
+{
+	check(InLandscapeComponents.Num() > 0);
+	SceneInterface = InLandscapeComponents[0]->GetScene();
+
+	// todo: use a 2d target?
+	const int32 SingleTileWidth = ComponentSizeVerts;
+	TargetSize = FIntPoint(0, ComponentSizeVerts);
+
+	// First compute the total render target size and prepare ComponentInfos (each component has its own number of needed passes because some might have different materials, thus different grass types) :
+	ComponentInfos.Reserve(InLandscapeComponents.Num());
+	int32 CurrentPixelOffsetX = 0;
+	for (ULandscapeComponent* Component : InLandscapeComponents)
+	{
+		ensure(Component->SceneProxy);
+
+		FComponentInfo& ComponentInfo = ComponentInfos.Add_GetRef(FComponentInfo(Component, bInNeedsGrassmap, bInNeedsHeightmap, InHeightMips));
+		ComponentInfo.PixelOffsetX = CurrentPixelOffsetX;
+
+		CurrentPixelOffsetX += ComponentInfo.NumPasses * SingleTileWidth;
+	}
+	TargetSize.X = CurrentPixelOffsetX;
+
+	FIntPoint TargetSizeMinusOne(TargetSize - FIntPoint(1, 1));
+	PassOffsetX = 2.0f * (float)SingleTileWidth / (float)TargetSize.X;
+
+	// Then compute FComponentInfo's ViewOffset with the knowledge of the total render target size : 
+	for (FComponentInfo& ComponentInfo : ComponentInfos)
+	{
+		FIntPoint ComponentOffset = (ComponentInfo.Component->GetSectionBase() - LandscapeProxy->LandscapeSectionOffset);
+		FVector2D ViewOffset(-ComponentOffset.X, ComponentOffset.Y);
+		ViewOffset.X += ComponentInfo.PixelOffsetX;
+		ViewOffset /= (FVector2D(TargetSize) * 0.5f);
+		ComponentInfo.ViewOffset = ViewOffset;
+	}
+
+	// center of target area in world
+	FVector TargetCenter = LandscapeProxy->GetTransform().TransformPosition(FVector(TargetSizeMinusOne, 0.f) * 0.5f);
+
+	// extent of target in world space
+	FVector TargetExtent = FVector(TargetSize, 0.0f) * LandscapeProxy->GetActorScale() * 0.5f;
+
+	ViewOrigin = TargetCenter;
+	ViewRotationMatrix = FInverseRotationMatrix(LandscapeProxy->GetActorRotation());
+	ViewRotationMatrix *= FMatrix(FPlane(1.0f, 0.0f, 0.0f, 0.0f),
+		FPlane(0.0f, -1.0f, 0.0f, 0.0f),
+		FPlane(0.0f, 0.0f, -1.0f, 0.0f),
+		FPlane(0.0f, 0.0f, 0.0f, 1.0f));
+
+	const float ZOffset = UE_OLD_WORLD_MAX;
+	ProjectionMatrix = FReversedZOrthoMatrix(
+		TargetExtent.X,
+		TargetExtent.Y,
+		0.5f / ZOffset,
+		ZOffset);
+
+	RenderTargetTexture = NewObject<UTextureRenderTarget2D>();
+	check(RenderTargetTexture);
+	RenderTargetTexture->ClearColor = FLinearColor::White;
+	RenderTargetTexture->TargetGamma = 1.0f;
+	const bool bForceLinearGamma = true;
+	RenderTargetTexture->InitCustomFormat(TargetSize.X, TargetSize.Y, PF_B8G8R8A8, bForceLinearGamma);
+	RenderTargetResource = RenderTargetTexture->GameThread_GetRenderTargetResource()->GetTextureRenderTarget2DResource();
+
+	UE::RenderCommandPipe::FSyncScope SyncScope;
+
+	// render
+	FLandscapeGrassWeightExporter_RenderThread* Exporter = this;
+	ENQUEUE_RENDER_COMMAND(FDrawSceneCommand)(
+		[Exporter](FRHICommandListImmediate& RHICmdList)
+		{
+			Exporter->RenderLandscapeComponentToTexture_RenderThread(RHICmdList);
+			FlushPendingDeleteRHIResources_RenderThread();
+		});
+}
+
+TMap<ULandscapeComponent*, TUniquePtr<FLandscapeComponentGrassData>, TInlineSetAllocator<1>> FLandscapeGrassWeightExporter::FetchResults()
+{
+	TArray<FColor> Samples;
+	Samples.SetNumUninitialized(TargetSize.X * TargetSize.Y);
+
+	// Copy the contents of the remote texture to system memory
+	FReadSurfaceDataFlags ReadSurfaceDataFlags;
+	ReadSurfaceDataFlags.SetLinearToGamma(false);
+	RenderTargetResource->ReadPixels(Samples, ReadSurfaceDataFlags, FIntRect(0, 0, TargetSize.X, TargetSize.Y));
+
+	TMap<ULandscapeComponent*, TUniquePtr<FLandscapeComponentGrassData>, TInlineSetAllocator<1>> Results;
+	Results.Reserve(ComponentInfos.Num());
+
+	// Local data will be moved in contiguous array at the end of export (to minimize slack waste)
+	TArray<uint16> HeightData;
+	TMap<ULandscapeGrassType*, TArray<uint8>> WeightData;
+
+	for (auto& ComponentInfo : ComponentInfos)
+	{
+		ULandscapeComponent* Component = ComponentInfo.Component;
+		ALandscapeProxy* Proxy = Component->GetLandscapeProxy();
+
+		TUniquePtr<FLandscapeComponentGrassData> NewGrassData = MakeUnique<FLandscapeComponentGrassData>(Component);
+
+		if (ComponentInfo.FirstHeightMipsPassIndex != MAX_int32)
+		{
+			HeightData.Empty(FMath::Square(ComponentSizeVerts));
+		}
+		else
+		{
+			HeightData.Empty(0);
+		}
+		NewGrassData->HeightMipData.Empty(HeightMips.Num());
+
+		WeightData.Empty();
+		TArray<TArray<uint8>*> GrassWeightArrays;
+		GrassWeightArrays.Empty(ComponentInfo.RequestedGrassTypes.Num());
+		for (auto GrassType : ComponentInfo.RequestedGrassTypes)
+		{
+			WeightData.Add(GrassType);
+		}
+
+		// need a second loop because the WeightData map will reallocate its arrays as grass types are added
+		for (auto GrassType : ComponentInfo.RequestedGrassTypes)
+		{
+			TArray<uint8>* DataArray = WeightData.Find(GrassType);
+			check(DataArray);
+			DataArray->Empty(FMath::Square(ComponentSizeVerts));
+			GrassWeightArrays.Add(DataArray);
+		}
+
+		// output debug bitmap
+#if UE_BUILD_DEBUG
+		static bool bOutputGrassBitmap = false;
+		if (bOutputGrassBitmap)
+		{
+			FString TempPath = FPaths::ScreenShotDir();
+			TempPath += TEXT("/GrassDebug");
+			IFileManager::Get().MakeDirectory(*TempPath, true);
+			FFileHelper::CreateBitmap(*(TempPath / "Grass"), TargetSize.X, TargetSize.Y, Samples.GetData(), nullptr, &IFileManager::Get(), nullptr, ComponentInfo.RequestedGrassTypes.Num() >= 2);
+		}
+#endif
+
+		for (int32 PassIdx = 0; PassIdx < ComponentInfo.NumPasses; PassIdx++)
+		{
+			FColor* SampleData = &Samples[ComponentInfo.PixelOffsetX + PassIdx * ComponentSizeVerts];
+			if (PassIdx < ComponentInfo.FirstHeightMipsPassIndex)
+			{
+				if (PassIdx == 0)
+				{
+					for (int32 y = 0; y < ComponentSizeVerts; y++)
+					{
+						for (int32 x = 0; x < ComponentSizeVerts; x++)
+						{
+							FColor& Sample = SampleData[x + y * TargetSize.X];
+							uint16 Height = (((uint16)Sample.R) << 8) + (uint16)(Sample.G);
+							HeightData.Add(Height);
+							if (ComponentInfo.RequestedGrassTypes.Num() > 0)
+							{
+								GrassWeightArrays[0]->Add(Sample.B);
+								if (ComponentInfo.RequestedGrassTypes.Num() > 1)
+								{
+									GrassWeightArrays[1]->Add(Sample.A);
+								}
+							}
+						}
+					}
+				}
+				else
+				{
+					for (int32 y = 0; y < ComponentSizeVerts; y++)
+					{
+						for (int32 x = 0; x < ComponentSizeVerts; x++)
+						{
+							FColor& Sample = SampleData[x + y * TargetSize.X];
+
+							int32 TypeIdx = PassIdx * 4 - 2;
+							GrassWeightArrays[TypeIdx++]->Add(Sample.R);
+							if (TypeIdx < ComponentInfo.RequestedGrassTypes.Num())
+							{
+								GrassWeightArrays[TypeIdx++]->Add(Sample.G);
+								if (TypeIdx < ComponentInfo.RequestedGrassTypes.Num())
+								{
+									GrassWeightArrays[TypeIdx++]->Add(Sample.B);
+									if (TypeIdx < ComponentInfo.RequestedGrassTypes.Num())
+									{
+										GrassWeightArrays[TypeIdx++]->Add(Sample.A);
+									}
+								}
+							}
+						}
+					}
+				}
+			}
+			else // PassIdx >= FirstHeightMipsPassIndex
+			{
+				const int32 Mip = HeightMips[PassIdx - ComponentInfo.FirstHeightMipsPassIndex];
+				int32 MipSizeVerts = NumSubsections * (SubsectionSizeQuads >> Mip);
+				TArray<uint16>& MipHeightData = NewGrassData->HeightMipData.Add(Mip);
+				for (int32 y = 0; y < MipSizeVerts; y++)
+				{
+					for (int32 x = 0; x < MipSizeVerts; x++)
+					{
+						FColor& Sample = SampleData[x + y * TargetSize.X];
+						uint16 Height = (((uint16)Sample.R) << 8) + (uint16)(Sample.G);
+						MipHeightData.Add(Height);
+					}
+				}
+			}
+		}
+
+		// remove null grass type if we had one (can occur if the node has null entries)
+		WeightData.Remove(nullptr);
+
+		// Remove any grass data that is entirely weight 0
+		for (auto Iter(WeightData.CreateIterator()); Iter; ++Iter)
+		{
+			if (Iter->Value.IndexOfByPredicate([&](const int8& Weight) { return Weight != 0; }) == INDEX_NONE)
+			{
+				Iter.RemoveCurrent();
+			}
+		}
+
+		NewGrassData->InitializeFrom(HeightData, WeightData);
+		Results.Add(Component, MoveTemp(NewGrassData));
+	}
+
+	return Results;
+}
+
+void FLandscapeGrassWeightExporter::ApplyResults()
+{
+	TMap<ULandscapeComponent*, TUniquePtr<FLandscapeComponentGrassData>, TInlineSetAllocator<1>> NewGrassData = FetchResults();
+
+	for (auto&& GrassDataPair : NewGrassData)
+	{
+		ULandscapeComponent* Component = GrassDataPair.Key;
+		FLandscapeComponentGrassData* ComponentGrassData = GrassDataPair.Value.Release();
+		ALandscapeProxy* Proxy = Component->GetLandscapeProxy();
+
+		// Assign the new data (thread-safe)
+		Component->GrassData = MakeShareable(ComponentGrassData);
+		Component->GrassData->bIsDirty = true;
+
+		if (Proxy->bBakeMaterialPositionOffsetIntoCollision)
+		{
+			Component->DestroyCollisionData();
+			Component->UpdateCollisionData();
+		}
+	}
+}
+
+void FLandscapeGrassWeightExporter::AddReferencedObjects(UObject* InThis, FReferenceCollector& Collector)
+{
+	if (RenderTargetTexture)
+	{
+		Collector.AddReferencedObject(RenderTargetTexture);
+	}
+
+	if (LandscapeProxy)
+	{
+		Collector.AddReferencedObject(LandscapeProxy);
+	}
+
+	for (auto& Info : ComponentInfos)
+	{
+		if (Info.Component)
+		{
+			Collector.AddReferencedObject(Info.Component);
+		}
+
+		Collector.AddReferencedObjects(Info.RequestedGrassTypes);
+	}
+}
+
+static bool ShouldCacheLandscapeGrassShaders(const FMeshMaterialShaderPermutationParameters& Parameters)
+{
+	// We only need grass weight shaders for Landscape vertex factories on desktop platforms
+	return (Parameters.MaterialParameters.bIsUsedWithLandscape || Parameters.MaterialParameters.bIsSpecialEngineMaterial) &&
+		IsFeatureLevelSupported(Parameters.Platform, ERHIFeatureLevel::SM5) &&
+		Parameters.VertexFactoryType == FindVertexFactoryType(FName(TEXT("FLandscapeFixedGridVertexFactory"), FNAME_Find)) &&
+		EnumHasAllFlags(Parameters.Flags, EShaderPermutationFlags::HasEditorOnlyData);
+}
+
+class FLandscapeGrassWeightVS : public FMeshMaterialShader
+{
+	DECLARE_SHADER_TYPE(FLandscapeGrassWeightVS, MeshMaterial);
+
+	LAYOUT_FIELD(FShaderParameter, RenderOffsetParameter);
+
+protected:
+
+	FLandscapeGrassWeightVS()
+	{}
+
+	FLandscapeGrassWeightVS(const FMeshMaterialShaderType::CompiledShaderInitializerType& Initializer)
+		: FMeshMaterialShader(Initializer)
+	{
+		RenderOffsetParameter.Bind(Initializer.ParameterMap, TEXT("RenderOffset"));
+		PassUniformBuffer.Bind(Initializer.ParameterMap, FSceneTextureUniformParameters::FTypeInfo::GetStructMetadata()->GetShaderVariableName());
+	}
+
+public:
+
+	static bool ShouldCompilePermutation(const FMeshMaterialShaderPermutationParameters& Parameters)
+	{
+		return ShouldCacheLandscapeGrassShaders(Parameters);
+	}
+
+	void GetShaderBindings(
+		const FScene* Scene,
+		ERHIFeatureLevel::Type FeatureLevel,
+		const FPrimitiveSceneProxy* PrimitiveSceneProxy,
+		const FMaterialRenderProxy& MaterialRenderProxy,
+		const FMaterial& Material,
+		const FMeshPassProcessorRenderState& DrawRenderState,
+		const FLandscapeGrassWeightShaderElementData& ShaderElementData,
+		FMeshDrawSingleShaderBindings& ShaderBindings) const
+	{
+		FMeshMaterialShader::GetShaderBindings(Scene, FeatureLevel, PrimitiveSceneProxy, MaterialRenderProxy, Material, DrawRenderState, ShaderElementData, ShaderBindings);
+
+		ShaderBindings.Add(RenderOffsetParameter, ShaderElementData.RenderOffset);
+	}
+};
+
+IMPLEMENT_MATERIAL_SHADER_TYPE(, FLandscapeGrassWeightVS, TEXT("/Engine/Private/LandscapeGrassWeight.usf"), TEXT("VSMain"), SF_Vertex);
+
+class FLandscapeGrassWeightPS : public FMeshMaterialShader
+{
+	DECLARE_SHADER_TYPE(FLandscapeGrassWeightPS, MeshMaterial);
+	LAYOUT_FIELD(FShaderParameter, OutputPassParameter);
+public:
+
+	static bool ShouldCompilePermutation(const FMeshMaterialShaderPermutationParameters& Parameters)
+	{
+		return ShouldCacheLandscapeGrassShaders(Parameters);
+	}
+
+	FLandscapeGrassWeightPS(const ShaderMetaType::CompiledShaderInitializerType& Initializer)
+		: FMeshMaterialShader(Initializer)
+	{
+		OutputPassParameter.Bind(Initializer.ParameterMap, TEXT("OutputPass"));
+		PassUniformBuffer.Bind(Initializer.ParameterMap, FSceneTextureUniformParameters::FTypeInfo::GetStructMetadata()->GetShaderVariableName());
+	}
+
+	FLandscapeGrassWeightPS()
+	{}
+
+	void GetShaderBindings(
+		const FScene* Scene,
+		ERHIFeatureLevel::Type FeatureLevel,
+		const FPrimitiveSceneProxy* PrimitiveSceneProxy,
+		const FMaterialRenderProxy& MaterialRenderProxy,
+		const FMaterial& Material,
+		const FMeshPassProcessorRenderState& DrawRenderState,
+		const FLandscapeGrassWeightShaderElementData& ShaderElementData,
+		FMeshDrawSingleShaderBindings& ShaderBindings) const
+	{
+		FMeshMaterialShader::GetShaderBindings(Scene, FeatureLevel, PrimitiveSceneProxy, MaterialRenderProxy, Material, DrawRenderState, ShaderElementData, ShaderBindings);
+
+		ShaderBindings.Add(OutputPassParameter, ShaderElementData.OutputPass);
+	}
+};
+
+IMPLEMENT_MATERIAL_SHADER_TYPE(, FLandscapeGrassWeightPS, TEXT("/Engine/Private/LandscapeGrassWeight.usf"), TEXT("PSMain"), SF_Pixel);
+
+void UE::Landscape::Grass::AddGrassWeightShaderTypes(FMaterialShaderTypes& InOutShaderTypes)
+{
+	InOutShaderTypes.AddShaderType<FLandscapeGrassWeightVS>();
+	InOutShaderTypes.AddShaderType<FLandscapeGrassWeightPS>();
+}
+
+#endif // WITH_EDITOR
+
diff --git a/Engine/Source/Runtime/Landscape/Private/LandscapeGrassWeightExporter.h b/Engine/Source/Runtime/Landscape/Private/LandscapeGrassWeightExporter.h
new file mode 100644
--- /dev/null
+++ b/Engine/Source/Runtime/Landscape/Private/LandscapeGrassWeightExporter.h
@@ -0,0 +1,99 @@
+// Copyright Epic Games, Inc. All Rights Reserved.
+#pragma once
+
+#include "CoreMinimal.h"
+#include "LandscapeComponent.h"
+#include "MeshMaterialShader.h"
+
+class ALandscapeProxy;
+class FLandscapeComponentSceneProxy;
+class FTextureRenderTarget2DResource;
+class UTextureRenderTarget2D;
+
+// data also accessible by render thread
+class FLandscapeGrassWeightExporter_RenderThread
+{
+	FLandscapeGrassWeightExporter_RenderThread(const TArray<int32>& InHeightMips)
+		: HeightMips(InHeightMips)
+	{
+	}
+
+	friend class FLandscapeGrassWeightExporter;
+
+public:
+	virtual ~FLandscapeGrassWeightExporter_RenderThread()
+	{}
+
+	struct FComponentInfo
+	{
+		TObjectPtr<ULandscapeComponent> Component = nullptr;
+		TArray<TObjectPtr<ULandscapeGrassType>> RequestedGrassTypes;
+		FVector2D ViewOffset = FVector2D::ZeroVector;
+		int32 PixelOffsetX = 0;
+		FLandscapeComponentSceneProxy* SceneProxy = nullptr;
+		int32 NumPasses = 0;
+		int32 FirstHeightMipsPassIndex = MAX_int32;
+
+		FComponentInfo(ULandscapeComponent* InComponent, bool bInNeedsGrassmap, bool bInNeedsHeightmap, const TArray<int32>& InHeightMips)
+			: Component(InComponent)
+			, SceneProxy((FLandscapeComponentSceneProxy*)InComponent->SceneProxy)
+		{
+			if (bInNeedsGrassmap)
+			{
+				RequestedGrassTypes = InComponent->GetGrassTypes();
+			}
+			int32 NumGrassMaps = RequestedGrassTypes.Num();
+			if (bInNeedsHeightmap || NumGrassMaps > 0)
+			{
+				NumPasses += FMath::DivideAndRoundUp(2 /* heightmap */ + NumGrassMaps, 4);
+			}
+			if (InHeightMips.Num() > 0)
+			{
+				FirstHeightMipsPassIndex = NumPasses;
+				NumPasses += InHeightMips.Num();
+			}
+		}
+	};
+
+	FSceneInterface* SceneInterface = nullptr;
+	FTextureRenderTarget2DResource* RenderTargetResource = nullptr;
+	TArray<FComponentInfo, TInlineAllocator<1>> ComponentInfos;
+	FIntPoint TargetSize;
+	TArray<int32> HeightMips;
+	float PassOffsetX;
+	FVector ViewOrigin;
+	FMatrix ViewRotationMatrix;
+	FMatrix ProjectionMatrix;
+
+	void RenderLandscapeComponentToTexture_RenderThread(FRHICommandListImmediate& RHICmdList);
+};
+
+class FLandscapeGrassWeightShaderElementData : public FMeshMaterialShaderElementData
+{
+public:
+
+	int32 OutputPass;
+	FVector2f RenderOffset;
+};
+
+class FLandscapeGrassWeightExporter : public FLandscapeGrassWeightExporter_RenderThread
+{
+	TObjectPtr<ALandscapeProxy> LandscapeProxy;
+	int32 ComponentSizeVerts;
+	int32 SubsectionSizeQuads;
+	int32 NumSubsections;
+	TArray<TObjectPtr<ULandscapeGrassType>> GrassTypes;
+	TObjectPtr<UTextureRenderTarget2D> RenderTargetTexture;
+
+public:
+	FLandscapeGrassWeightExporter(ALandscapeProxy* InLandscapeProxy, TArrayView<ULandscapeComponent* const> InLandscapeComponents, bool bInNeedsGrassmap = true, bool bInNeedsHeightmap = true, const TArray<int32>& InHeightMips = {});
+	TMap<ULandscapeComponent*, TUniquePtr<FLandscapeComponentGrassData>, TInlineSetAllocator<1>> FetchResults();
+	void ApplyResults();
+	void AddReferencedObjects(UObject* InThis, FReferenceCollector& Collector);
+};
+
+namespace UE::Landscape::Grass
+{
+	void AddGrassWeightShaderTypes(FMaterialShaderTypes& InOutShaderTypes);
+}
+
diff --git a/Engine/Source/Runtime/Landscape/Private/LandscapeSubsystem.cpp b/Engine/Source/Runtime/Landscape/Private/LandscapeSubsystem.cpp
--- a/Engine/Source/Runtime/Landscape/Private/LandscapeSubsystem.cpp
+++ b/Engine/Source/Runtime/Landscape/Private/LandscapeSubsystem.cpp
@@ -18,6 +18,7 @@
 #include "LandscapeRender.h"
 #include "LandscapePrivate.h"
 #include "LandscapeSettings.h"
+#include "LandscapeGrassMapsBuilder.h"
 #include "ProfilingDebugging/CsvProfiler.h"
 #include "WorldPartition/WorldPartitionSubsystem.h"
 #include "ActorPartition/ActorPartitionSubsystem.h"
