Fix Grass Update on newly created Landscapes

diff --git a/Engine/Source/Runtime/Landscape/Classes/LandscapeComponent.h b/Engine/Source/Runtime/Landscape/Classes/LandscapeComponent.h
--- a/Engine/Source/Runtime/Landscape/Classes/LandscapeComponent.h
+++ b/Engine/Source/Runtime/Landscape/Classes/LandscapeComponent.h
@@ -514,7 +514,10 @@ private:
 	UPROPERTY(EditAnywhere, Category = LandscapeComponent)
 	TArray<FLandscapePerLODMaterialOverride> PerLODOverrideMaterials;
 
-	/** Cached list of grass types supported by the component's material */
+	/** The value of the landscape material AllStateCRC the last time the GrassTypes array was updated from it */
+	uint32 LastLandscapeMaterialAllStateCRCWhenGrassTypesBuilt = 0;
+
+	/** Cached list of grass types supported by the component's material. Call UpdateGrassTypes() to ensure this array is up to date */
 	UPROPERTY()
 	TArray<TObjectPtr<ULandscapeGrassType>> GrassTypes;
 
@@ -522,7 +525,6 @@ private:
 	UPROPERTY()
 	float GrassTypesMaxDiscardDistance = 0.0f;
 
-
 public:
 
 	/** Uniquely identifies this component's built map data. */
@@ -795,11 +797,12 @@ public:
 
 	/** Gets the landscape info object for this landscape */
 	LANDSCAPE_API ULandscapeInfo* GetLandscapeInfo() const;
-	
-	const TArray<ULandscapeGrassType*>& GetGrassTypes() const { return GrassTypes; }
 
-	/** Temporarily sets the grass type for this component. Any call to UpdateGrassTypes will override what has been set using this method */
-	void SetGrassTypes(const TArray<ULandscapeGrassType*>& InGrassTypes) { GrassTypes = InGrassTypes; }
+	/** Returns the array of grass types used by the landscape material. Call UpdateGrassTypes first to ensure this array is up to date. */
+	const TArray<TObjectPtr<ULandscapeGrassType>>& GetGrassTypes() const { return GrassTypes; }
+
+	/** Temporarily sets the grass type for this component. Any call to UpdateGrassTypes may override what has been set using this method. */
+	void SetGrassTypes(const TArray<TObjectPtr<ULandscapeGrassType>>& InGrassTypes)	{ GrassTypes = InGrassTypes; }
 	
 	bool MaterialHasGrass() const { return !GetGrassTypes().IsEmpty(); }
 	
@@ -807,8 +810,11 @@ public:
 	void SetGrassTypesMaxDiscardDistance(const float InGrassTypesMaxDiscardDistance) { GrassTypesMaxDiscardDistance = InGrassTypesMaxDiscardDistance; }
 
 #if WITH_EDITOR
-	/** Recomputes the list of grass types and other related info (in case there was a material change, for example) : */
-	void UpdateGrassTypes();
+	/** If the LandscapeMaterial has changed, updates the GrassTypes array. Returns true if the GrassTypes array was updated. */
+	bool UpdateGrassTypes(bool bForceUpdate = false);
+
+	/** Recomputes the maximum discard distance across all grass types in the GrassTypes array. */
+	void UpdateGrassTypesMaxDiscardDistance();
 
 	/** Deletes a material layer from the current edit layer on this component, removing all its data, adjusting other layer's weightmaps if necessary, etc. */
 	LANDSCAPE_API void DeleteLayer(ULandscapeLayerInfoObject* LayerInfo, FLandscapeEditDataInterface& LandscapeEdit);
@@ -839,7 +845,7 @@ public:
 	/** Computes a hash representing the state of the material and grasstypes used by this component. */
 	LANDSCAPE_API uint32 ComputeGrassMapGenerationHash() const;
 
-	/* Is the grassmap data outdated, eg by a material */
+	/* Returns true if the component HAS grass data, but it is not up to date */
 	bool IsGrassMapOutdated() const;
 
 	/** Renders the heightmap of this component (including material world-position-offset) at the specified LOD */
diff --git a/Engine/Source/Runtime/Landscape/Private/LandscapeGrass.cpp b/Engine/Source/Runtime/Landscape/Private/LandscapeGrass.cpp
--- a/Engine/Source/Runtime/Landscape/Private/LandscapeGrass.cpp
+++ b/Engine/Source/Runtime/Landscape/Private/LandscapeGrass.cpp
@@ -870,23 +870,33 @@ FLandscapeComponentGrassData::FLandscapeComponentGrassData(ULandscapeComponent*
 {
 }
 
-void ULandscapeComponent::UpdateGrassTypes()
-	{
-	GrassTypesMaxDiscardDistance = 0.0f;
-	GrassTypes.Reset();
-	
+bool ULandscapeComponent::UpdateGrassTypes(bool bForceUpdate)
+{
+	bool bChanged = false;
 	if (UMaterialInterface* Material = GetLandscapeMaterial())
 	{
-		GrassTypes = Material->GetMaterial()->GetCachedExpressionData().GrassTypes;
+		uint32 CurMaterialAllStateCRC = Material->ComputeAllStateCRC();
+		if (bForceUpdate || (LastLandscapeMaterialAllStateCRCWhenGrassTypesBuilt != CurMaterialAllStateCRC))
+		{
+			LastLandscapeMaterialAllStateCRCWhenGrassTypesBuilt = CurMaterialAllStateCRC;
+			bChanged = true;
+
+			GrassTypes = Material->GetMaterial()->GetCachedExpressionData().GrassTypes;
+		}
+	}
+	return bChanged;
+}
 
-		for (const ULandscapeGrassType* GrassType : GrassTypes)
+void ULandscapeComponent::UpdateGrassTypesMaxDiscardDistance()
+{
+	GrassTypesMaxDiscardDistance = 0.0f;
+	for (const ULandscapeGrassType* GrassType : GrassTypes)
+	{
+		if (GrassType != nullptr)
 		{
-			if (GrassType != nullptr)
+			for (const FGrassVariety& GrassVariety : GrassType->GrassVarieties)
 			{
-				for (const FGrassVariety& GrassVariety : GrassType->GrassVarieties)
-				{
-					GrassTypesMaxDiscardDistance = FMath::Max((float)GrassVariety.EndCullDistance.GetValue(), GrassTypesMaxDiscardDistance);
-				}
+				GrassTypesMaxDiscardDistance = FMath::Max((float)GrassVariety.EndCullDistance.GetValue(), GrassTypesMaxDiscardDistance);
 			}
 		}
 	}
@@ -909,21 +919,19 @@ uint32 ULandscapeComponent::ComputeGrassMapGenerationHash() const
 		// If anything changes in the grass types, we should take that into account as well :
 		for (ULandscapeGrassType* GrassType : GrassTypes)
 		{
-			if (GrassType == nullptr)
+			if (GrassType != nullptr)
 			{
-				continue;
-			}
-
-			Hash = FCrc::TypeCrc32(GrassType->StateHash, Hash);
+				Hash = FCrc::TypeCrc32(GrassType->StateHash, Hash);
 			}
 		}
+	}
 
 	return Hash;
 }
 
 bool ULandscapeComponent::IsGrassMapOutdated() const
 {
-	return GrassData->HasValidData() ? (ComputeGrassMapGenerationHash() != GrassData->GenerationHash) : false;
+	return GrassData->HasValidData() && (ComputeGrassMapGenerationHash() != GrassData->GenerationHash);
 }
 
 bool ULandscapeComponent::CanRenderGrassMap() const
@@ -1369,6 +1377,10 @@ void ULandscapeGrassType::PostEditChangeProperty(FPropertyChangedEvent& Property
 {
 	Super::PostEditChangeProperty(PropertyChangedEvent);
 
+	// Look for landscape components using this grass type, and flush their grass data
+	// This is not required to trigger a grass update (the state hash changing will do that for us)
+	// This is only necessary to immediately delete the instances while dragging in the UI.
+	// (otherwise the instances stick around until we actually rebuild them)
 	if (GIsEditor)
 	{
 		for (TObjectIterator<ALandscapeProxy> It; It; ++It)
@@ -1376,22 +1388,15 @@ void ULandscapeGrassType::PostEditChangeProperty(FPropertyChangedEvent& Property
 			ALandscapeProxy* Proxy = *It;
 			if (Proxy->GetWorld() && !Proxy->GetWorld()->IsPlayInEditor())
 			{
-				const UMaterialInterface* MaterialInterface = Proxy->GetLandscapeMaterial();
-				if (MaterialInterface)
+				for (ULandscapeComponent* Component : Proxy->LandscapeComponents)
 				{
-					TArray<const UMaterialExpressionLandscapeGrassOutput*> GrassExpressions;
-					MaterialInterface->GetMaterial()->GetAllExpressionsOfType<UMaterialExpressionLandscapeGrassOutput>(GrassExpressions);
-
-					// Should only be one grass type node
-					if (GrassExpressions.Num() > 0)
+					for (ULandscapeGrassType* GrassType : Component->GetGrassTypes())
 					{
-						for (auto& Output : GrassExpressions[0]->GrassTypes)
+						if (GrassType == this)
 						{
-							if (Output.GrassType == this)
-							{
-								Proxy->FlushGrassComponents();
-								break;
-							}
+							TSet<ULandscapeComponent*> Components = { Component };
+							Proxy->FlushGrassComponents(&Components, /* bFlushGrassMaps= */ true);
+							break;
 						}
 					}
 				}
@@ -1399,6 +1404,7 @@ void ULandscapeGrassType::PostEditChangeProperty(FPropertyChangedEvent& Property
 		}
 	}
 
+	// Update the state hash (so we detect the change and rebuild the relevant grass)
 	StateHash = ComputeStateHash();
 }
 #endif
@@ -2435,6 +2441,7 @@ void ALandscapeProxy::BuildGrassMaps(FScopedSlowTask* InSlowTask)
 		bHasLandscapeGrass = false;
 		for (ULandscapeComponent* Component : LandscapeComponents)
 		{
+			Component->UpdateGrassTypes(/* bForceUpdate = */ true);
 			bHasLandscapeGrass |= (Component->GetGrassTypes().Num() > 0);
 			GrassTypesMaxDiscardDistance = FMath::Max(Component->GetGrassTypesMaxDiscardDistance(), GrassTypesMaxDiscardDistance);
 		}
@@ -2524,22 +2531,25 @@ void ALandscapeProxy::UpdateGrassDataStatus(TSet<UTexture2D*>* OutCurrentForcedS
 				}
 			}
 
-			bool bIsGrassMapOutdated = Component->IsGrassMapOutdated();
-			if (bIsGrassMapOutdated)
+			// if the landscape material has changed, update the grass types (needed to ensure ComputeGrassMapGenerationHash is correct)
+			Component->UpdateGrassTypes();
+			bool bGrassMapGenerationHashChanged = (Component->ComputeGrassMapGenerationHash() != Component->GrassData->GenerationHash);
+
+			// outdated meaning: it has grass data, but it is not up to date
+			bool bIsGrassMapOutdated = Component->GrassData->HasValidData() && bGrassMapGenerationHashChanged;
+
+			if (bGrassMapGenerationHashChanged)
 			{
-#if WITH_EDITOR
-				// If the material has changed we need to cache the GrassTypes and related info again : 
-				Component->UpdateGrassTypes();
-#endif // WITH_EDITOR
+				Component->UpdateGrassTypesMaxDiscardDistance();
+			}
 
-				if (OutOutdatedComponents)
+			if (bIsGrassMapOutdated && OutOutdatedComponents)
 			{
 				OutOutdatedComponents->Add(Component);
 			}
-			}
 
 			// Needs to be called after UpdateGrassTypes
-			TArray<ULandscapeGrassType*> GrassTypes = Component->GetGrassTypes();
+			const TArray<TObjectPtr<ULandscapeGrassType>>& GrassTypes = Component->GetGrassTypes();
 			bool bHasGrassTypes = (GrassTypes.Num() > 0);
 			if (bHasGrassTypes || bBakeMaterialPositionOffsetIntoCollision)
 			{
