Made sure landscapes that belong to a level instance are not editable. The only landscapes that appear in landscape mode (and that can be updated) are now those that belong to the main level.
#tests Editor

diff --git a/Engine/Source/Runtime/Landscape/Private/LandscapeEdit.cpp b/Engine/Source/Runtime/Landscape/Private/LandscapeEdit.cpp
--- a/Engine/Source/Runtime/Landscape/Private/LandscapeEdit.cpp
+++ b/Engine/Source/Runtime/Landscape/Private/LandscapeEdit.cpp
@@ -2566,6 +2566,13 @@ FIntRect ULandscapeComponent::GetComponentExtent() const
 //
 bool ULandscapeInfo::SupportsLandscapeEditing() const
 {
+	// Don't let landscapes from level instances be edited : they can only be edited in their source level (note that technically, the IsEditing test is not necessary as they cannot be edited 
+	//  in level instance mode, since it's mutually exclusive with landscape mode, but let's keep it for describing the intention here) :
+	if (LandscapeActor.IsValid() && LandscapeActor->IsInLevelInstance() && !LandscapeActor->IsInEditingLevelInstance())
+	{
+		return false;
+	}
+
 	bool bSupportsEditing = true;
 	ForEachLandscapeProxy([&bSupportsEditing](ALandscapeProxy* Proxy)
 	{
@@ -2574,6 +2581,7 @@ bool ULandscapeInfo::SupportsLandscapeEditing() const
 			bSupportsEditing = false;
 			return false;
 		}
+
 		return true;
 
 	});
diff --git a/Engine/Source/Runtime/Landscape/Private/LandscapeEditLayers.cpp b/Engine/Source/Runtime/Landscape/Private/LandscapeEditLayers.cpp
--- a/Engine/Source/Runtime/Landscape/Private/LandscapeEditLayers.cpp
+++ b/Engine/Source/Runtime/Landscape/Private/LandscapeEditLayers.cpp
@@ -8222,7 +8222,7 @@ void ALandscape::UpdateLayersContent(bool bInWaitForStreaming, bool bInSkipMonit
 	bool bResourcesReady = PrepareTextureResources(bInWaitForStreaming);
 
 	ULandscapeInfo* LandscapeInfo = GetLandscapeInfo();
-	if ((LandscapeInfo == nullptr) || !CanHaveLayersContent() || !LandscapeInfo->AreAllComponentsRegistered() || !FApp::CanEverRender())
+	if ((LandscapeInfo == nullptr) || !CanHaveLayersContent() || !LandscapeInfo->AreAllComponentsRegistered() || !FApp::CanEverRender() || !LandscapeInfo->SupportsLandscapeEditing())
 	{
 		return;
 	}
