Fixed BP functions RenderHeightmap/Weightmap/s functions so that they don't require their transform and extents pin to be connected, since it's a supported case already.
Note : this API change is backwards-compatble both on the code and data side

Also, made it explicit in the code that WeightmapScaleBias is *not* to be used, since it has a different meaning than HeightmapScaleBias. This was compensated by the FloorToInt32 that is done later, but still, it might avoid any future confusion to explicitly use a 0 bias instead (which is always true on weightmaps where there's no texture sharing).
#tests editor

diff --git a/Engine/Source/Runtime/Landscape/Classes/Landscape.h b/Engine/Source/Runtime/Landscape/Classes/Landscape.h
--- a/Engine/Source/Runtime/Landscape/Classes/Landscape.h
+++ b/Engine/Source/Runtime/Landscape/Classes/Landscape.h
@@ -274,7 +274,7 @@ public:
 	* @return false in case of failure (e.g. invalid inputs, incompatible render target format...)
 	*/
 	UFUNCTION(BlueprintCallable, Category = "Landscape|Runtime")
-	LANDSCAPE_API bool RenderHeightmap(const FTransform& InWorldTransform, const FBox2D& InExtents, UTextureRenderTarget2D* OutRenderTarget);
+	LANDSCAPE_API bool RenderHeightmap(FTransform InWorldTransform, FBox2D InExtents, UTextureRenderTarget2D* OutRenderTarget);
 
 	/**
 	* Render the final weightmap for the requested layer, in the requested top-down window, as one -atlased- texture in the provided render target 2D
@@ -286,7 +286,7 @@ public:
 	* @return false in case of failure (e.g. invalid inputs, incompatible render target format...)
 	*/
 	UFUNCTION(BlueprintCallable, Category = "Landscape|Runtime")
-	LANDSCAPE_API bool RenderWeightmap(const FTransform& InWorldTransform, const FBox2D& InExtents, FName InWeightmapLayerName, UTextureRenderTarget2D* OutRenderTarget);
+	LANDSCAPE_API bool RenderWeightmap(FTransform InWorldTransform, FBox2D InExtents, FName InWeightmapLayerName, UTextureRenderTarget2D* OutRenderTarget);
 
 	/**
 	* Render the final weightmaps for the requested layers, in the requested top-down window, as one -atlased- texture in the provided render target (2D or 2DArray) 
@@ -300,7 +300,7 @@ public:
 	* @return false in case of failure (e.g. invalid inputs, incompatible render target format...)
 	*/
 	UFUNCTION(BlueprintCallable, Category = "Landscape|Runtime")
-	LANDSCAPE_API bool RenderWeightmaps(const FTransform& InWorldTransform, const FBox2D& InExtents, const TArray<FName>& InWeightmapLayerNames, UTextureRenderTarget* OutRenderTarget);
+	LANDSCAPE_API bool RenderWeightmaps(FTransform InWorldTransform, FBox2D InExtents, const TArray<FName>& InWeightmapLayerNames, UTextureRenderTarget* OutRenderTarget);
 
 	bool IsValidRenderTargetFormatHeightmap(EPixelFormat InRenderTargetFormat, bool& bOutCompressHeight);
 	bool IsValidRenderTargetFormatWeightmap(EPixelFormat InRenderTargetFormat, int32& OutNumChannels);
diff --git a/Engine/Source/Runtime/Landscape/Private/Landscape.cpp b/Engine/Source/Runtime/Landscape/Private/Landscape.cpp
--- a/Engine/Source/Runtime/Landscape/Private/Landscape.cpp
+++ b/Engine/Source/Runtime/Landscape/Private/Landscape.cpp
@@ -3218,13 +3218,13 @@ bool ALandscape::RenderMergedTextureInternal(const FTransform& InRenderAreaWorld
 			FIntPoint ComponentKey = It.Key;
 
 			UTexture2D* SourceTexture = nullptr;
-			FVector4 SourceTextureScaleBias;
+			FVector2D SourceTextureBias;
 			int32 SourceTextureChannel = INDEX_NONE;
 
 			if (bIsHeightmap)
 			{
 				SourceTexture = Component->GetHeightmap();
-				SourceTextureScaleBias = Component->HeightmapScaleBias;
+				SourceTextureBias = FVector2D(Component->HeightmapScaleBias.Z, Component->HeightmapScaleBias.W);
 			}
 			else
 			{
@@ -3235,7 +3235,9 @@ bool ALandscape::RenderMergedTextureInternal(const FTransform& InRenderAreaWorld
 				{
 					SourceTexture = WeightmapTextures[AllocInfo->WeightmapTextureIndex];
 					check(SourceTexture != nullptr);
-					SourceTextureScaleBias = Component->WeightmapScaleBias;
+					// Note : don't use WeightmapScaleBias here, it has a different meaning than HeightmapScaleBias (very conveniently!) : this is compensated by the FloorToInt32 later on, 
+					//  but still, let's set this to zero here and use the fact that there's no texture sharing on weightmaps : 
+					SourceTextureBias = FVector2D::ZeroVector;
 					SourceTextureChannel = AllocInfo->WeightmapTextureChannel;
 				}
 			}
@@ -3243,9 +3245,7 @@ bool ALandscape::RenderMergedTextureInternal(const FTransform& InRenderAreaWorld
 			if (SourceTexture != nullptr)
 			{
 				// Get the subregion of the source texture that this component uses (differs due to texture sharing).
-				// SourceTextureScaleBias ZW values give us the offset of the component in a shared texture. You might think that XY would
-				// give the portion of the texture it occupies, but no, XY are 1/size, for some reason. Just calculate the subregion
-				// size ourselves.
+				// SourceTextureBias values give us the offset of the component in a shared texture
 				int32 ComponentSize = Component->NumSubsections * (Component->SubsectionSizeQuads + 1);
 
 				FIntPoint SourceTextureOffset(0, 0);
@@ -3254,8 +3254,8 @@ bool ALandscape::RenderMergedTextureInternal(const FTransform& InRenderAreaWorld
 				{
 					// We get the overall source texture size via the resource instead of direct GetSizeX/Y calls because the latter are unreliable while the texture is being built.
 					SourceTextureOffset = FIntPoint(
-						FMath::FloorToInt32(SourceTextureScaleBias.Z * SourceTextureResource->GetSizeX()),
-						FMath::FloorToInt32(SourceTextureScaleBias.W * SourceTextureResource->GetSizeY()));
+						FMath::FloorToInt32(SourceTextureBias.X * SourceTextureResource->GetSizeX()),
+						FMath::FloorToInt32(SourceTextureBias.Y * SourceTextureResource->GetSizeY()));
 				}
 			
 				// When mips are partially loaded, we need to take that into consideration when merging the source texture :
@@ -3393,17 +3393,17 @@ bool ALandscape::RenderMergedTextureInternal(const FTransform& InRenderAreaWorld
 	return true;
 }
 
-bool ALandscape::RenderHeightmap(const FTransform& InRenderAreaWorldTransform, const FBox2D& InRenderAreaExtents, UTextureRenderTarget2D* OutRenderTarget)
+bool ALandscape::RenderHeightmap(FTransform InRenderAreaWorldTransform, FBox2D InRenderAreaExtents, UTextureRenderTarget2D* OutRenderTarget)
 {
 	return RenderMergedTextureInternal(InRenderAreaWorldTransform, InRenderAreaExtents, /*InWeightmapLayerNames = */{}, OutRenderTarget);
 }
 
-bool ALandscape::RenderWeightmap(const FTransform& InRenderAreaWorldTransform, const FBox2D& InRenderAreaExtents, FName InWeightmapLayerName, UTextureRenderTarget2D* OutRenderTarget)
+bool ALandscape::RenderWeightmap(FTransform InRenderAreaWorldTransform, FBox2D InRenderAreaExtents, FName InWeightmapLayerName, UTextureRenderTarget2D* OutRenderTarget)
 {
 	return RenderMergedTextureInternal(InRenderAreaWorldTransform, InRenderAreaExtents, /*InWeightmapLayerNames = */{ InWeightmapLayerName }, OutRenderTarget);
 }
 
-bool ALandscape::RenderWeightmaps(const FTransform& InRenderAreaWorldTransform, const FBox2D& InRenderAreaExtents, const TArray<FName>& InWeightmapLayerNames, UTextureRenderTarget* OutRenderTarget)
+bool ALandscape::RenderWeightmaps(FTransform InRenderAreaWorldTransform, FBox2D InRenderAreaExtents, const TArray<FName>& InWeightmapLayerNames, UTextureRenderTarget* OutRenderTarget)
 {
 	return RenderMergedTextureInternal(InRenderAreaWorldTransform, InRenderAreaExtents, InWeightmapLayerNames, OutRenderTarget);
 }
