Use determinsitic GUID for landscape heightfield collision.

diff --git a/Engine/Source/Runtime/Landscape/Private/LandscapeCollision.cpp b/Engine/Source/Runtime/Landscape/Private/LandscapeCollision.cpp
--- a/Engine/Source/Runtime/Landscape/Private/LandscapeCollision.cpp
+++ b/Engine/Source/Runtime/Landscape/Private/LandscapeCollision.cpp
@@ -831,7 +831,10 @@ void ULandscapeHeightfieldCollisionComponent::CreateCollisionObject()
 
 		if (!HeightfieldGuid.IsValid())
 		{
-			HeightfieldGuid = FGuid::NewGuid();
+#if !WITH_EDITORONLY_DATA
+			uint32 CollisionHash = 0;
+#endif
+			HeightfieldGuid = FGuid::NewDeterministicGuid(GetPathName(), CollisionHash);
 			bCheckDDC = false;
 		}
 		else
@@ -1022,7 +1025,10 @@ void ULandscapeHeightfieldCollisionComponent::CreateCollisionObject(
 	}
 #endif
 
-	HeightfieldGuid = FGuid::NewGuid();
+#if !WITH_EDITORONLY_DATA
+	uint32 CollisionHash = 0;
+#endif
+	HeightfieldGuid = FGuid::NewDeterministicGuid(GetPathName(), CollisionHash);
 
 	HeightfieldRef = GSharedHeightfieldRefs.Add(HeightfieldGuid, new FHeightfieldGeometryRef(HeightfieldGuid));
 	HeightfieldRef->HeightfieldGeometry = Chaos::FHeightFieldPtr(new Chaos::FHeightField(Heights, PhysicalMaterialIds, CollisionSizeVerts, CollisionSizeVerts, Chaos::FVec3(1)));
