Added BP method to retrieve landscape target (paint) layer names

#tests editor

diff --git a/Engine/Source/Runtime/Landscape/Classes/Landscape.h b/Engine/Source/Runtime/Landscape/Classes/Landscape.h
--- a/Engine/Source/Runtime/Landscape/Classes/Landscape.h
+++ b/Engine/Source/Runtime/Landscape/Classes/Landscape.h
@@ -302,10 +302,19 @@ public:
 	UFUNCTION(BlueprintCallable, Category = "Landscape|Runtime")
 	LANDSCAPE_API bool RenderWeightmaps(FTransform InWorldTransform, FBox2D InExtents, const TArray<FName>& InWeightmapLayerNames, UTextureRenderTarget* OutRenderTarget);
 
+	/** 
+	* Retrieves the names of valid paint layers on this landscape (editor-only : returns nothing at runtime) 
+	* @Param bInIncludeVisibilityLayer whether the visibility layer's name should be included in the list or not
+	* @return the list of paint layer names
+	*/
+	UFUNCTION(BlueprintCallable, BlueprintPure=false, Category = "Landscape|Editor", meta=(DevelopmentOnly))
+	TArray<FName> GetTargetLayerNames(bool bInIncludeVisibilityLayer = false) const;
+
 	bool IsValidRenderTargetFormatHeightmap(EPixelFormat InRenderTargetFormat, bool& bOutCompressHeight);
 	bool IsValidRenderTargetFormatWeightmap(EPixelFormat InRenderTargetFormat, int32& OutNumChannels);
 
 #if WITH_EDITOR
+
 	/** Computes & returns bounds containing all landscape proxies (if any) or this landscape's bounds otherwise. Note that in non-WP worlds this will call GetLoadedBounds(). */
 	LANDSCAPE_API FBox GetCompleteBounds() const;
 	void RegisterLandscapeEdMode(ILandscapeEdModeInterface* InLandscapeEdMode) { LandscapeEdMode = InLandscapeEdMode; }
diff --git a/Engine/Source/Runtime/Landscape/Private/LandscapeEdit.cpp b/Engine/Source/Runtime/Landscape/Private/LandscapeEdit.cpp
--- a/Engine/Source/Runtime/Landscape/Private/LandscapeEdit.cpp
+++ b/Engine/Source/Runtime/Landscape/Private/LandscapeEdit.cpp
@@ -84,6 +84,7 @@ LandscapeEdit.cpp: Landscape editing
 #include "LandscapeSplineActor.h"
 #endif
 #include "Algo/Count.h"
+#include "Algo/Transform.h"
 #include "Algo/ForEach.h"
 #include "Serialization/MemoryWriter.h"
 #include "Engine/Canvas.h"
@@ -8282,8 +8283,30 @@ bool ALandscapeProxy::LandscapeExportHeightmapToRenderTarget(UTextureRenderTarge
 	return true;
 }
 
+TArray<FName> ALandscape::GetTargetLayerNames(bool bInIncludeVisibilityLayer) const
+{
+	TArray<FName> Result;
+
 #if WITH_EDITOR
+	if (ULandscapeInfo* LandscapeInfo = GetLandscapeInfo())
+	{
+		Algo::TransformIf(LandscapeInfo->Layers, Result,
+			[bInIncludeVisibilityLayer](const FLandscapeInfoLayerSettings& InSettings) 
+			{ 
+				return (InSettings.LayerInfoObj != nullptr) 
+				&& (bInIncludeVisibilityLayer || (InSettings.LayerInfoObj != ALandscapeProxy::VisibilityLayer)); 
+			},
+			[](const FLandscapeInfoLayerSettings& InSettings) { return InSettings.GetLayerName(); });
+	}
+#else // WITH_EDITOR
+	FMessageLog("Blueprint").Error(LOCTEXT("GetTargetLayerNames_Runtime.", "GetTargetLayerNames_EditorOnly: this cannot be called at runtime"));
+#endif // !WITH_EDITOR
 
+	return Result;
+}
+
+
+#if WITH_EDITOR
 
 bool ALandscapeProxy::LandscapeImportWeightmapFromRenderTarget(UTextureRenderTarget2D* InRenderTarget, FName InLayerName)
 {
