Fix landscape grass being culled 100% in cooked builds.
(Regression introduced in previous change to LandscapeGrass.cpp)

diff --git a/Engine/Source/Runtime/Landscape/Private/Landscape.cpp b/Engine/Source/Runtime/Landscape/Private/Landscape.cpp
--- a/Engine/Source/Runtime/Landscape/Private/Landscape.cpp
+++ b/Engine/Source/Runtime/Landscape/Private/Landscape.cpp
@@ -1354,6 +1354,7 @@ void ULandscapeComponent::PostLoad()
 	{
 		UpdateGrassTypes();
 	}
+	UpdateGrassTypesMaxDiscardDistance();
 
 #if !UE_BUILD_SHIPPING
 	if (MobileCombinationMaterialInstances.Num() == 0)
diff --git a/Engine/Source/Runtime/Landscape/Private/LandscapeGrass.cpp b/Engine/Source/Runtime/Landscape/Private/LandscapeGrass.cpp
--- a/Engine/Source/Runtime/Landscape/Private/LandscapeGrass.cpp
+++ b/Engine/Source/Runtime/Landscape/Private/LandscapeGrass.cpp
@@ -882,6 +882,7 @@ bool ULandscapeComponent::UpdateGrassTypes(bool bForceUpdate)
 			bChanged = true;
 
 			GrassTypes = Material->GetMaterial()->GetCachedExpressionData().GrassTypes;
+			UpdateGrassTypesMaxDiscardDistance();
 		}
 	}
 	return bChanged;
@@ -2538,11 +2539,6 @@ void ALandscapeProxy::UpdateGrassDataStatus(TSet<UTexture2D*>* OutCurrentForcedS
 			// outdated meaning: it has grass data, but it is not up to date
 			bool bIsGrassMapOutdated = Component->GrassData->HasValidData() && bGrassMapGenerationHashChanged;
 
-			if (bGrassMapGenerationHashChanged)
-			{
-				Component->UpdateGrassTypesMaxDiscardDistance();
-			}
-
 			if (bIsGrassMapOutdated && OutOutdatedComponents)
 			{
 				OutOutdatedComponents->Add(Component);
