* Fixed Lumen cards for landscape : those require the landscape fixed grid vertex factory now (this was working in the editor because of an unrelated bug)
* Included those shaders in PSO precaching as well
#tests editor

diff --git a/Engine/Source/Runtime/Landscape/Private/Landscape.cpp b/Engine/Source/Runtime/Landscape/Private/Landscape.cpp
--- a/Engine/Source/Runtime/Landscape/Private/Landscape.cpp
+++ b/Engine/Source/Runtime/Landscape/Private/Landscape.cpp
@@ -1064,6 +1064,8 @@ void ULandscapeComponent::UpdatedSharedPropertiesFromActor()
 
 void ULandscapeComponent::PostLoad()
 {
+	using namespace UE::Landscape;
+
 	Super::PostLoad();
 
 	if (IsComponentPSOPrecachingEnabled())
@@ -1086,16 +1088,10 @@ void ULandscapeComponent::PostLoad()
 		}
 
 		// we need the fixed grid vertex factory for both virtual texturing and grass
-		bool bNeedsFixedGridVertexFactory = UseVirtualTexturing(GMaxRHIShaderPlatform);
-		// This cvar is defined in the water plugin and searching for it should return nullptr if the plugin is not loaded
-		const bool bWaterPluginLoaded = IConsoleManager::Get().FindConsoleVariable(TEXT("r.Water.WaterInfo.RenderMethod")) != nullptr;
-		bNeedsFixedGridVertexFactory |= bWaterPluginLoaded;
-		if (bNeedsFixedGridVertexFactory)
+		if (NeedsFixedGridVertexFactory(GMaxRHIShaderPlatform))
 		{
 			VertexFactoryDataList.Add(FPSOPrecacheVertexFactoryData(&FLandscapeFixedGridVertexFactory::StaticType));
 		}
-
-		using namespace UE::Landscape;
 		
 		if (Culling::UseCulling(GMaxRHIShaderPlatform))
 		{
diff --git a/Engine/Source/Runtime/Landscape/Private/LandscapeRender.cpp b/Engine/Source/Runtime/Landscape/Private/LandscapeRender.cpp
--- a/Engine/Source/Runtime/Landscape/Private/LandscapeRender.cpp
+++ b/Engine/Source/Runtime/Landscape/Private/LandscapeRender.cpp
@@ -209,6 +209,29 @@ static FAutoConsoleVariableRef CVarLandscapeRayTracingGeometryFractionalLODUpdat
 
 #endif
 
+
+// ----------------------------------------------------------------------------------
+
+namespace UE::Landscape
+{
+bool NeedsFixedGridVertexFactory(EShaderPlatform InShaderPlatform)
+{
+	bool bNeedsFixedGridVertexFactory = false;
+	// We need the fixed grid vertex factory for virtual texturing : 
+	bNeedsFixedGridVertexFactory |= UseVirtualTexturing(InShaderPlatform);
+
+	// We need the fixed grid vertex factory for rendering Landscape into Lumen Surface Cache: 
+	bNeedsFixedGridVertexFactory |= DoesPlatformSupportLumenGI(InShaderPlatform);
+
+	// We need the fixed grid vertex factory for rendering the water info texture : 
+	// This cvar is defined in the water plugin and searching for it should return nullptr if the plugin is not loaded
+	const bool bWaterPluginLoaded = IConsoleManager::Get().FindConsoleVariable(TEXT("r.Water.WaterInfo.RenderMethod")) != nullptr;
+	bNeedsFixedGridVertexFactory |= bWaterPluginLoaded;
+
+	return bNeedsFixedGridVertexFactory;
+}
+} // namespace UE::Landscape
+
 //
 // FLandscapeDebugOptions
 //
@@ -1434,14 +1457,10 @@ void FLandscapeComponentSceneProxy::CreateRenderThreadResources(FRHICommandListB
 			SharedBuffers->VertexFactory = LandscapeXYOffsetVertexFactory;
 		}
 
-		// We need the fixed grid vertex factory for virtual texturing, grass and for rendering the water info texture and for rendering Landscape into Lumen Surface Cache: 
-		bool bNeedsFixedGridVertexFactory = UseVirtualTexturing(GetScene().GetShaderPlatform()) || DoesPlatformSupportLumenGI(GMaxRHIShaderPlatform);
-
-		// This cvar is defined in the water plugin and searching for it should return nullptr if the plugin is not loaded
-		const bool bWaterPluginLoaded = IConsoleManager::Get().FindConsoleVariable(TEXT("r.Water.WaterInfo.RenderMethod")) != nullptr;
-		bNeedsFixedGridVertexFactory |= bWaterPluginLoaded;
+		bool bNeedsFixedGridVertexFactory = NeedsFixedGridVertexFactory(GetScene().GetShaderPlatform());
 
 #if WITH_EDITOR
+		// We also need the fixed vertex factor for grass/physical materials : 
 		bNeedsFixedGridVertexFactory |= (SharedBuffers->GrassIndexBuffer != nullptr);
 #endif // WITH_EDITOR
 
@@ -3643,8 +3662,9 @@ public:
 				const bool bIsGrassShaderType = Algo::Find(GetGrassShaderTypes(), ShaderType->GetFName()) != nullptr;
 				const bool bIsGPULightmassShaderType = Algo::Find(GetGPULightmassShaderTypes(), ShaderType->GetFName()) != nullptr;
 				const bool bIsRuntimeVirtualTextureShaderType = Algo::Find(GetRuntimeVirtualTextureShaderTypes(), ShaderType->GetFName()) != nullptr;
+				const bool bIsLumen = Algo::Find(GetLumenCardShaderTypes(), ShaderType->GetFName()) != nullptr;
 
-				const bool bIsShaderTypeUsingFixedGrid = bIsGrassShaderType || bIsRuntimeVirtualTextureShaderType || bIsGPULightmassShaderType;
+				const bool bIsShaderTypeUsingFixedGrid = bIsGrassShaderType || bIsRuntimeVirtualTextureShaderType || bIsGPULightmassShaderType || bIsLumen;
 				
 				static const FName RayTracingDynamicGeometryConverterCS = FName(TEXT("FRayTracingDynamicGeometryConverterCS"));
 				const bool bIsRayTracingShaderType = ShaderType->GetFName() == RayTracingDynamicGeometryConverterCS;
@@ -3843,6 +3863,8 @@ public:
 			FName(TEXT("FTrivialMaterialCHS")),
 #endif // RHI_RAYTRACING
 
+			// No Lumen on thumbnails
+			FName(TEXT("FLumenCardCS")),
 			FName(TEXT("FLumenCardVS")),
 			FName(TEXT("FLumenCardPS<true>")),
 			FName(TEXT("FLumenCardPS<false>")),
@@ -3894,6 +3916,18 @@ public:
 		};
 		return ShaderTypes;
 	}
+
+	static const TArray<FName>& GetLumenCardShaderTypes()
+	{
+		static const TArray<FName> ShaderTypes =
+		{
+			FName(TEXT("FLumenCardVS")),
+			// We only need bMultiViewCapture == false as landscape components are non-Nanite :
+			FName(TEXT("FLumenCardPS<false>")),
+		};
+		return ShaderTypes;
+	}
+
 };
 
 FMaterialResource* ULandscapeMaterialInstanceConstant::AllocatePermutationResource()
diff --git a/Engine/Source/Runtime/Landscape/Public/LandscapeRender.h b/Engine/Source/Runtime/Landscape/Public/LandscapeRender.h
--- a/Engine/Source/Runtime/Landscape/Public/LandscapeRender.h
+++ b/Engine/Source/Runtime/Landscape/Public/LandscapeRender.h
@@ -39,6 +39,7 @@ LandscapeRender.h: New terrain rendering
 
 class FLandscapeComponentSceneProxy;
 enum class ERuntimeVirtualTextureMaterialType : uint8;
+enum EShaderPlatform : uint16;
 
 #if RHI_RAYTRACING
 struct FLandscapeRayTracingImpl;
@@ -93,6 +94,10 @@ LANDSCAPE_API extern TObjectPtr<UMaterialInterface> GLandscapeLayerUsageMaterial
 LANDSCAPE_API extern TObjectPtr<UMaterialInterface> GLandscapeDirtyMaterial;
 #endif
 
+namespace UE::Landscape
+{
+bool NeedsFixedGridVertexFactory(EShaderPlatform InShaderPlatform);
+} // namespace UE::Landscape
 
 /** The uniform shader parameters for a landscape draw call. */
 BEGIN_GLOBAL_SHADER_PARAMETER_STRUCT(FLandscapeUniformShaderParameters, LANDSCAPE_API)
