Fixed crash when updating landscape edit layers (account for missing layer data, like we do everywhere else in the code)
#tests editor

diff --git a/Engine/Source/Runtime/Landscape/Private/LandscapeEditLayers.cpp b/Engine/Source/Runtime/Landscape/Private/LandscapeEditLayers.cpp
--- a/Engine/Source/Runtime/Landscape/Private/LandscapeEditLayers.cpp
+++ b/Engine/Source/Runtime/Landscape/Private/LandscapeEditLayers.cpp
@@ -4531,8 +4531,13 @@ void ALandscape::PrepareLayersHeightmapsLocalMergeRenderThreadData(const FUpdate
 			{
 				if (Layer.bVisible && !InMergeParams.bSkipBrush)
 				{
-					UTexture2D* LayerHeightmap = Component->GetLayerData(Layer.Guid)->HeightmapData.Texture.Get();
-					NewComponentRenderInfo.VisibleLayerHeightmapTextures.Add(FTexture2DResourceSubregion(LayerHeightmap->GetResource()->GetTexture2DResource(), ComponentTextureSubregion));
+					if (FLandscapeLayerComponentData* ComponentLayerData = Component->GetLayerData(Layer.Guid))
+					{
+						if (UTexture2D* LayerHeightmap = ComponentLayerData->HeightmapData.Texture.Get())
+						{
+							NewComponentRenderInfo.VisibleLayerHeightmapTextures.Add(FTexture2DResourceSubregion(LayerHeightmap->GetResource()->GetTexture2DResource(), ComponentTextureSubregion));
+						}
+					}
 				}
 			}
 		}
@@ -5980,12 +5985,14 @@ void ALandscape::ReallocateLayersWeightmaps(FUpdateLayersContentContext& InUpdat
 			// No need for an allocation if the edit layer is invisible : 
 			if (Layer.bVisible)
 			{
-				FLandscapeLayerComponentData* LayerComponentData = Component->GetLayerData(Layer.Guid);
-				for (const FWeightmapLayerAllocationInfo& LayerWeightmapAllocation : LayerComponentData->WeightmapData.LayerAllocations)
+				if (FLandscapeLayerComponentData* LayerComponentData = Component->GetLayerData(Layer.Guid))
 				{
-					if (LayerWeightmapAllocation.LayerInfo != nullptr)
+					for (const FWeightmapLayerAllocationInfo& LayerWeightmapAllocation : LayerComponentData->WeightmapData.LayerAllocations)
 					{
-						ComponentLayerAlloc->AddUnique(LayerWeightmapAllocation.LayerInfo);
+						if (LayerWeightmapAllocation.LayerInfo != nullptr)
+						{
+							ComponentLayerAlloc->AddUnique(LayerWeightmapAllocation.LayerInfo);
+						}
 					}
 				}
 			}
