Fixed potential crashes in landscape code when iteration over garbage actor
#tests editor

diff --git a/Engine/Source/Editor/LandscapeEditor/Private/LandscapeEdModePaintTools.cpp b/Engine/Source/Editor/LandscapeEditor/Private/LandscapeEdModePaintTools.cpp
--- a/Engine/Source/Editor/LandscapeEditor/Private/LandscapeEdModePaintTools.cpp
+++ b/Engine/Source/Editor/LandscapeEditor/Private/LandscapeEdModePaintTools.cpp
@@ -79,7 +79,7 @@ namespace
 		// Find matching runtime virtual texture components and invalidate dirty region.
 		if (RuntimeVirtualTextures.Num() > 0)
 		{
-			for (TObjectIterator<URuntimeVirtualTextureComponent> It; It; ++It)
+			for (TObjectIterator<URuntimeVirtualTextureComponent> It(/*AdditionalExclusionFlags = */RF_ClassDefaultObject, /*bIncludeDerivedClasses = */true, /*InInternalExclusionFlags = */EInternalObjectFlags::Garbage); It; ++It)
 			{
 				if (It->GetVirtualTexture() != nullptr && RuntimeVirtualTextures.Contains(It->GetVirtualTexture()))
 				{
diff --git a/Engine/Source/Runtime/Landscape/Private/Landscape.cpp b/Engine/Source/Runtime/Landscape/Private/Landscape.cpp
--- a/Engine/Source/Runtime/Landscape/Private/Landscape.cpp
+++ b/Engine/Source/Runtime/Landscape/Private/Landscape.cpp
@@ -135,7 +135,7 @@ static void PrintNumLandscapeShadows()
 {
 	int32 NumComponents = 0;
 	int32 NumShadowCasters = 0;
-	for (TObjectIterator<ULandscapeComponent> It; It; ++It)
+	for (TObjectIterator<ULandscapeComponent> It(/*AdditionalExclusionFlags = */RF_ClassDefaultObject, /*bIncludeDerivedClasses = */true, /*InInternalExclusionFlags = */EInternalObjectFlags::Garbage); It; ++It)
 	{
 		ULandscapeComponent* LC = *It;
 		NumComponents++;
diff --git a/Engine/Source/Runtime/Landscape/Private/LandscapeEdit.cpp b/Engine/Source/Runtime/Landscape/Private/LandscapeEdit.cpp
--- a/Engine/Source/Runtime/Landscape/Private/LandscapeEdit.cpp
+++ b/Engine/Source/Runtime/Landscape/Private/LandscapeEdit.cpp
@@ -5357,7 +5357,7 @@ void ULandscapeLayerInfoObject::PostEditChangeProperty(FPropertyChangedEvent& Pr
 	{
 		if (PropertyName == NAME_PhysMaterial || PropertyName == NAME_MinimumCollisionRelevanceWeight)
 		{
-			for (TObjectIterator<ALandscapeProxy> It; It; ++It)
+			for (TObjectIterator<ALandscapeProxy> It(/*AdditionalExclusionFlags = */RF_ClassDefaultObject, /*bIncludeDerivedClasses = */true, /*InInternalExclusionFlags = */EInternalObjectFlags::Garbage); It; ++It)
 			{
 				ALandscapeProxy* Proxy = *It;
 				if (Proxy->GetWorld() && !Proxy->GetWorld()->IsPlayInEditor())
@@ -5380,7 +5380,7 @@ void ULandscapeLayerInfoObject::PostEditChangeProperty(FPropertyChangedEvent& Pr
 		else if (PropertyName == NAME_LayerUsageDebugColor || PropertyName == NAME_R || PropertyName == NAME_G || PropertyName == NAME_B || PropertyName == NAME_A)
 		{
 			LayerUsageDebugColor.A = 1.0f;
-			for (TObjectIterator<ALandscapeProxy> It; It; ++It)
+			for (TObjectIterator<ALandscapeProxy> It(/*AdditionalExclusionFlags = */RF_ClassDefaultObject, /*bIncludeDerivedClasses = */true, /*InInternalExclusionFlags = */EInternalObjectFlags::Garbage); It; ++It)
 			{
 				ALandscapeProxy* Proxy = *It;
 				if (Proxy->GetWorld() && !Proxy->GetWorld()->IsPlayInEditor())
@@ -5395,7 +5395,7 @@ void ULandscapeLayerInfoObject::PostEditChangeProperty(FPropertyChangedEvent& Pr
 				PropertyName == GET_MEMBER_NAME_CHECKED(ULandscapeLayerInfoObject, SplineFalloffModulationScale) ||
 				PropertyName == GET_MEMBER_NAME_CHECKED(ULandscapeLayerInfoObject, SplineFalloffModulationTiling))
 		{
-			for (TObjectIterator<ULandscapeInfo> It; It; ++It)
+			for (TObjectIterator<ULandscapeInfo> It(/*AdditionalExclusionFlags = */RF_ClassDefaultObject, /*bIncludeDerivedClasses = */true, /*InInternalExclusionFlags = */EInternalObjectFlags::Garbage); It; ++It)
 			{
 				if(ALandscape* Landscape = It->LandscapeActor.Get())
 				{
@@ -5411,7 +5411,7 @@ void ULandscapeLayerInfoObject::PostEditUndo()
 	Super::PostEditUndo();
 
 	// Force the update of spline data for the (potentially) affected landscapes in case of an undo
-	for (TObjectIterator<ULandscapeInfo> It; It; ++It)
+	for (TObjectIterator<ULandscapeInfo> It(/*AdditionalExclusionFlags = */RF_ClassDefaultObject, /*bIncludeDerivedClasses = */true, /*InInternalExclusionFlags = */EInternalObjectFlags::Garbage); It; ++It)
 	{
 		if (ALandscape* Landscape = It->LandscapeActor.Get())
 		{
@@ -6705,61 +6705,34 @@ void ALandscape::PostEditChangeProperty(FPropertyChangedEvent& PropertyChangedEv
 	else if (GIsEditor && PropertyName == GET_MEMBER_NAME_CHECKED(ALandscapeProxy, MaxLODLevel))
 	{
 		MaxLODLevel = FMath::Clamp<int32>(MaxLODLevel, -1, FMath::CeilLogTwo(SubsectionSizeQuads + 1) - 1);
-		bPropagateToProxies = true;
 	}
 	else if (PropertyName == GET_MEMBER_NAME_CHECKED(ALandscapeProxy, ComponentScreenSizeToUseSubSections))
 	{
 		ComponentScreenSizeToUseSubSections = FMath::Clamp<float>(ComponentScreenSizeToUseSubSections, 0.01f, 1.0f);
-		bPropagateToProxies = true;
 	}
 	else if (PropertyName == GET_MEMBER_NAME_CHECKED(ALandscapeProxy, LODDistributionSetting))
 	{
 		LODDistributionSetting = FMath::Clamp<float>(LODDistributionSetting, 1.0f, 10.0f);
-		bPropagateToProxies = true;
 	}
 	else if (PropertyName == GET_MEMBER_NAME_CHECKED(ALandscapeProxy, LOD0DistributionSetting))
 	{
 		LOD0DistributionSetting = FMath::Clamp<float>(LOD0DistributionSetting, 1.0f, 10.0f);
-		bPropagateToProxies = true;
 	}
 	else if (PropertyName == GET_MEMBER_NAME_CHECKED(ALandscapeProxy, LOD0ScreenSize))
 	{
 		LOD0ScreenSize = FMath::Clamp<float>(LOD0ScreenSize, 0.1f, 10.0f);
-		bPropagateToProxies = true;
 	}
 	else if (PropertyName == GET_MEMBER_NAME_CHECKED(ALandscapeProxy, CollisionMipLevel))
 	{
 		CollisionMipLevel = FMath::Clamp<int32>(CollisionMipLevel, 0, FMath::CeilLogTwo(SubsectionSizeQuads + 1) - 1);
-		bPropagateToProxies = true;
-	}
-	else if (PropertyName == GET_MEMBER_NAME_CHECKED(ALandscapeProxy, DefaultPhysMaterial))
-	{
-		bPropagateToProxies = true;
 	}
 	else if (PropertyName == GET_MEMBER_NAME_CHECKED(ALandscapeProxy, SimpleCollisionMipLevel))
 	{
 		SimpleCollisionMipLevel = FMath::Clamp<int32>(SimpleCollisionMipLevel, 0, FMath::CeilLogTwo(SubsectionSizeQuads + 1) - 1);
-		bPropagateToProxies = true;
 	}
 	else if (PropertyName == GET_MEMBER_NAME_CHECKED(ALandscapeProxy, bBakeMaterialPositionOffsetIntoCollision))
 	{
-		bPropagateToProxies = true;
-	}
-	else if ((PropertyName == GET_MEMBER_NAME_CHECKED(ALandscapeProxy, bEnableNanite))
-		|| (PropertyName == GET_MEMBER_NAME_CHECKED(ALandscapeProxy, NaniteLODIndex))
-		|| (PropertyName == GET_MEMBER_NAME_CHECKED(ALandscapeProxy, bNaniteSkirtEnabled))
-		|| (PropertyName == GET_MEMBER_NAME_CHECKED(ALandscapeProxy, NaniteSkirtDepth))
-		)
-	{
-		bPropagateToProxies = true;
-	}
-	else if ((PropertyName == GET_MEMBER_NAME_CHECKED(ALandscapeProxy, RuntimeVirtualTextures))
-		|| (PropertyName == GET_MEMBER_NAME_CHECKED(ALandscapeProxy, VirtualTextureRenderPassType))
-		|| (PropertyName == GET_MEMBER_NAME_CHECKED(ALandscapeProxy, bVirtualTextureRenderWithQuad))
-		|| (PropertyName == GET_MEMBER_NAME_CHECKED(ALandscapeProxy, bVirtualTextureRenderWithQuadHQ))
-		|| (PropertyName == GET_MEMBER_NAME_CHECKED(ALandscapeProxy, VirtualTextureNumLods)) 
-		|| (PropertyName == GET_MEMBER_NAME_CHECKED(ALandscapeProxy, VirtualTextureLodBias)))
-	{
+		// TODO [jonathan.bard] : why propagate to proxies?
 		bPropagateToProxies = true;
 	}
 	else if (GIsEditor && PropertyName == GET_MEMBER_NAME_CHECKED(ALandscapeProxy, StaticLightingResolution))
@@ -6788,11 +6761,11 @@ void ALandscape::PostEditChangeProperty(FPropertyChangedEvent& PropertyChangedEv
 	}
 	else if (PropertyName == GET_MEMBER_NAME_CHECKED(ALandscape, LODGroupKey))
 	{
-		// need to invoke GetSharedProperties to copy this change down to the proxies
-		bPropagateToProxies = true;
 		bMarkAllLandscapeRenderStateDirty = true;
 	}
-	else if (GIsEditor && IsSharedProperty(PropertyName))
+	
+	// If the property that has changed is overridable or inherited, synchronize the change on all landcape proxies :
+	if (GIsEditor && IsSharedProperty(PropertyName))
 	{
 		bPropagateToProxies = true;
 	}
diff --git a/Engine/Source/Runtime/Landscape/Private/LandscapeEditLayers.cpp b/Engine/Source/Runtime/Landscape/Private/LandscapeEditLayers.cpp
--- a/Engine/Source/Runtime/Landscape/Private/LandscapeEditLayers.cpp
+++ b/Engine/Source/Runtime/Landscape/Private/LandscapeEditLayers.cpp
@@ -185,7 +185,7 @@ TAutoConsoleVariable<int32> CVarLandscapeRemoveEmptyPaintLayersOnEdit(
 
 void OnLandscapeEditLayersLocalMergeChanged(IConsoleVariable* CVar)
 {
-	for (TObjectIterator<UWorld> It; It; ++It)
+	for (TObjectIterator<UWorld> It(/*AdditionalExclusionFlags = */RF_ClassDefaultObject, /*bIncludeDerivedClasses = */true, /*InInternalExclusionFlags = */EInternalObjectFlags::Garbage); It; ++It)
 	{
 		UWorld* CurrentWorld = *It;
 		if (!CurrentWorld->IsGameWorld())
@@ -240,7 +240,7 @@ private:
 	void ClearDirty()
 	{
 		bool bCleared = false;
-		for (TObjectIterator<UWorld> It; It; ++It)
+		for (TObjectIterator<UWorld> It(/*AdditionalExclusionFlags = */RF_ClassDefaultObject, /*bIncludeDerivedClasses = */true, /*InInternalExclusionFlags = */EInternalObjectFlags::Garbage); It; ++It)
 		{
 			UWorld* CurrentWorld = *It;
 			if (!CurrentWorld->IsGameWorld())
@@ -3575,7 +3575,7 @@ void ALandscape::InvalidateRVTForTextures(const TSet<UTexture2D*>& InTextures)
 		if (!PrimitiveComponentsToInvalidate.IsEmpty())
 		{
 			// Now invalidate the RVT regions that correspond to these components :
-			for (TObjectIterator<URuntimeVirtualTextureComponent> It; It; ++It)
+			for (TObjectIterator<URuntimeVirtualTextureComponent> It(/*AdditionalExclusionFlags = */RF_ClassDefaultObject, /*bIncludeDerivedClasses = */true, /*InInternalExclusionFlags = */EInternalObjectFlags::Garbage); It; ++It)
 			{
 				for (UPrimitiveComponent* PrimitiveComponent : PrimitiveComponentsToInvalidate)
 				{
diff --git a/Engine/Source/Runtime/Landscape/Private/LandscapeGrass.cpp b/Engine/Source/Runtime/Landscape/Private/LandscapeGrass.cpp
--- a/Engine/Source/Runtime/Landscape/Private/LandscapeGrass.cpp
+++ b/Engine/Source/Runtime/Landscape/Private/LandscapeGrass.cpp
@@ -1384,7 +1384,7 @@ void ULandscapeGrassType::PostEditChangeProperty(FPropertyChangedEvent& Property
 	// (otherwise the instances stick around until we actually rebuild them)
 	if (GIsEditor)
 	{
-		for (TObjectIterator<ALandscapeProxy> It; It; ++It)
+		for (TObjectIterator<ALandscapeProxy> It(/*AdditionalExclusionFlags = */RF_ClassDefaultObject, /*bIncludeDerivedClasses = */true, /*InInternalExclusionFlags = */EInternalObjectFlags::Garbage); It; ++It)
 		{
 			ALandscapeProxy* Proxy = *It;
 			if (Proxy->GetWorld() && !Proxy->GetWorld()->IsPlayInEditor())
diff --git a/Engine/Source/Runtime/Landscape/Private/LandscapeSplines.cpp b/Engine/Source/Runtime/Landscape/Private/LandscapeSplines.cpp
--- a/Engine/Source/Runtime/Landscape/Private/LandscapeSplines.cpp
+++ b/Engine/Source/Runtime/Landscape/Private/LandscapeSplines.cpp
@@ -91,7 +91,7 @@ struct FLandscapeFixSplines
 	{
 		bool bFixed = false;
 
-		for (TObjectIterator<UWorld> It; It; ++It)
+		for (TObjectIterator<UWorld> It(/*AdditionalExclusionFlags = */RF_ClassDefaultObject, /*bIncludeDerivedClasses = */true, /*InInternalExclusionFlags = */EInternalObjectFlags::Garbage); It; ++It)
 		{
 			UWorld* CurrentWorld = *It;
 			if (!CurrentWorld->IsGameWorld())
