Disabled landscape.RemoveEmptyPaintLayersOnEdit by default since it causes more problems (landscape actors being dirty on load for no apparent reason) than it solves. We need a proper separation between editor and final data to solve this for good.
#tests editor

diff --git a/Engine/Source/Runtime/Landscape/Private/LandscapeEdit.cpp b/Engine/Source/Runtime/Landscape/Private/LandscapeEdit.cpp
--- a/Engine/Source/Runtime/Landscape/Private/LandscapeEdit.cpp
+++ b/Engine/Source/Runtime/Landscape/Private/LandscapeEdit.cpp
@@ -6764,7 +6764,7 @@ void ALandscape::PostEditChangeProperty(FPropertyChangedEvent& PropertyChangedEv
 		bMarkAllLandscapeRenderStateDirty = true;
 	}
 	
-	// If the property that has changed is overridable or inherited, synchronize the change on all landcape proxies :
+	// If the property that has changed is overridable or inherited, synchronize the change on all landscape proxies :
 	if (GIsEditor && IsSharedProperty(PropertyName))
 	{
 		bPropagateToProxies = true;
diff --git a/Engine/Source/Runtime/Landscape/Private/LandscapeEditLayers.cpp b/Engine/Source/Runtime/Landscape/Private/LandscapeEditLayers.cpp
--- a/Engine/Source/Runtime/Landscape/Private/LandscapeEditLayers.cpp
+++ b/Engine/Source/Runtime/Landscape/Private/LandscapeEditLayers.cpp
@@ -180,7 +180,9 @@ TAutoConsoleVariable<int32> CVarLandscapeValidateProxyWeightmapUsages(
 
 TAutoConsoleVariable<int32> CVarLandscapeRemoveEmptyPaintLayersOnEdit(
 	TEXT("landscape.RemoveEmptyPaintLayersOnEdit"),
-	1,
+	// TODO [jonathan.bard] : this has been disabled for now, since it can lead to a permanent dirty-on-load state for landscape, where the edit layers will do a new weightmap allocation for the missing layer
+	//  (e.g. if a BP brush writes to it), only to remove it after readback, which will lead to the actor to be marked dirty. We need to separate the final from the source weightmap data to avoid this issue : 
+	0, 
 	TEXT("This will analyze weightmaps on readback and remove unneeded allocations (for unpainted layers)."));
 
 void OnLandscapeEditLayersLocalMergeChanged(IConsoleVariable* CVar)
