Adding async readback functionality to the grass weight exporter, and enabling it in non-editor builds

This is in preparation for the runtime grass generation

#public 652ec918f1386bb0f23dabfb, 652ecfb1d82862fb3879916b

diff --git a/Engine/Source/Runtime/Landscape/Private/LandscapeAsyncTextureReadback.cpp b/Engine/Source/Runtime/Landscape/Private/LandscapeAsyncTextureReadback.cpp
new file mode 100644
--- /dev/null
+++ b/Engine/Source/Runtime/Landscape/Private/LandscapeAsyncTextureReadback.cpp
@@ -0,0 +1,118 @@
+// Copyright Epic Games, Inc. All Rights Reserved.
+
+#include "LandscapeAsyncTextureReadback.h"
+#include "DataDrivenShaderPlatformInfo.h"
+#include "RenderGraphBuilder.h"
+#include "RenderGraphUtils.h"
+
+void FLandscapeAsyncTextureReadback::StartReadback_RenderThread(FRDGBuilder& GraphBuilder, FRDGTextureRef RDGTexture)
+{
+	check(!bAsyncReadbackSubmitOnRenderThread && !AsyncReadback);
+	check(RDGTexture->Desc.Format == PF_B8G8R8A8);
+	AsyncReadback = MakeUnique<FRHIGPUTextureReadback>(TEXT("LandscapeGrassReadback"));
+	AddEnqueueCopyPass(GraphBuilder, AsyncReadback.Get(), RDGTexture);
+	FIntVector Size = RDGTexture->Desc.GetSize();
+	TextureWidth = Size.X;
+	TextureHeight = Size.Y;
+	check(Size.Z == 1);
+
+	bAsyncReadbackSubmitOnRenderThread = true;
+}
+
+void FLandscapeAsyncTextureReadback::FinishReadback_RenderThread()
+{
+	check(bAsyncReadbackSubmitOnRenderThread && AsyncReadback.IsValid());
+	int32 RowPitchInPixels = 0;
+	int32 BufferHeight = 0;
+	void* SrcData = AsyncReadback->Lock(RowPitchInPixels, &BufferHeight);
+	check(SrcData);
+	check(RowPitchInPixels >= TextureWidth);
+	check(BufferHeight >= TextureHeight);
+
+	// copy into ReadbackResults
+	ReadbackResults.SetNumUninitialized(TextureWidth * TextureHeight);
+
+	// OpenGL does not really support BGRA images and uses channnel swizzling to emulate them
+	// so when we read them back we get internal RGBA representation
+	const bool bSwapRBChannels = IsOpenGLPlatform(GMaxRHIShaderPlatform);
+
+	if (!bSwapRBChannels && TextureWidth == RowPitchInPixels)
+	{
+		memcpy(ReadbackResults.GetData(), SrcData, TextureWidth * TextureHeight * sizeof(FColor));
+	}
+	else
+	{
+		// copy row by row
+		FColor* Dst = ReadbackResults.GetData();
+		FColor* Src = (FColor*)SrcData;
+		if (bSwapRBChannels)
+		{
+			for (int y = 0; y < TextureHeight; y++)
+			{
+				for (int x = 0; x < TextureWidth; x++)
+				{
+					// swap B and R channels when copying
+					Dst->B = Src->R;
+					Dst->G = Src->G;
+					Dst->R = Src->B;
+					Dst->A = Src->A;
+					Dst++;
+					Src++;
+				}
+				Src += RowPitchInPixels - TextureWidth;
+			}
+		}
+		else
+		{
+			for (int y = 0; y < TextureHeight; y++)
+			{
+				memcpy(Dst, Src, TextureWidth * sizeof(FColor));
+				Dst += TextureWidth;
+				Src += RowPitchInPixels;
+			}
+		}
+	}
+
+	AsyncReadback->Unlock();
+	AsyncReadback.Reset();
+
+	FPlatformMisc::MemoryBarrier();
+	bAsyncReadbackCompleteOnRenderThread = true;
+}
+
+bool FLandscapeAsyncTextureReadback::CheckAndUpdate()
+{
+	if (bAsyncReadbackCompleteOnRenderThread)
+	{
+		// already done
+		return true;
+	}
+
+	// not done yet -- queue update check on render thread
+	FLandscapeAsyncTextureReadback* Readback = this;
+	ENQUEUE_RENDER_COMMAND(FLandscapeAsyncTextureReadback_CheckAndUpdate)(
+		[Readback](FRHICommandListImmediate& RHICmdList)
+		{
+			if (Readback->bAsyncReadbackCompleteOnRenderThread)
+			{
+				// Actually finished already.. game thread was just asking too early.
+				// No need to do anything, game thread will advance next time it checks.
+			}
+			else if (Readback->bAsyncReadbackSubmitOnRenderThread)
+			{
+				// Readback was submit but not found to be complete yet -- let's check if it completed.
+				check(Readback->AsyncReadback.IsValid());
+				if (Readback->AsyncReadback->IsReady())
+				{
+					Readback->FinishReadback_RenderThread();
+				}
+			}
+			else
+			{
+				// not submit yet.. nothing to do
+			}
+		});
+
+	return false;
+}
+
diff --git a/Engine/Source/Runtime/Landscape/Private/LandscapeAsyncTextureReadback.h b/Engine/Source/Runtime/Landscape/Private/LandscapeAsyncTextureReadback.h
new file mode 100644
--- /dev/null
+++ b/Engine/Source/Runtime/Landscape/Private/LandscapeAsyncTextureReadback.h
@@ -0,0 +1,55 @@
+// Copyright Epic Games, Inc. All Rights Reserved.
+#pragma once
+
+#include "CoreMinimal.h"
+#include "RHIGPUReadback.h"
+#include "RenderGraphFwd.h"
+
+// Class that performs an async GPU texture readback for landscape purposes.
+// Only supports BGRA8 texture format, and returns the results as an array of FColors.
+class FLandscapeAsyncTextureReadback
+{
+private:
+	// render thread state
+	bool bAsyncReadbackSubmitOnRenderThread = false;
+	bool bAsyncReadbackCompleteOnRenderThread = false;
+	TUniquePtr<FRHIGPUTextureReadback> AsyncReadback;
+
+	// results
+	int32 TextureWidth = 0;
+	int32 TextureHeight = 0;
+	TArray<FColor> ReadbackResults;
+
+	void FinishReadback_RenderThread();
+
+public:
+	FLandscapeAsyncTextureReadback() {}
+
+	// use this to start an async readback operation from the render thread (on a render graph texture)
+	void StartReadback_RenderThread(FRDGBuilder& GraphBuilder, FRDGTextureRef RDGTexture);
+
+	// TODO: add other StartReadback functions as needed
+
+	// Check readback status and update the process if needed. Return true when the AsyncReadbackResults are available.
+	// You must call this occasionally or the readback may never complete.
+	bool CheckAndUpdate();
+
+	// Returns true when async readback results are available.  Call GetResults() to retrieve them.
+	bool IsComplete()
+	{
+		return bAsyncReadbackCompleteOnRenderThread;
+	}
+
+	// Retrieve the async readback results.  Requires readback to be complete.
+	// This function returns its internal memory buffer, relinquishing control over it, so this function can only be called once.
+	TArray<FColor> TakeResults(FIntPoint* OutSize)
+	{
+		check(bAsyncReadbackCompleteOnRenderThread);
+		if (OutSize)
+		{
+			*OutSize = FIntPoint(TextureWidth, TextureHeight);
+		}
+		return MoveTemp(ReadbackResults);
+	}
+};
+
diff --git a/Engine/Source/Runtime/Landscape/Private/LandscapeGrassWeightExporter.cpp b/Engine/Source/Runtime/Landscape/Private/LandscapeGrassWeightExporter.cpp
--- a/Engine/Source/Runtime/Landscape/Private/LandscapeGrassWeightExporter.cpp
+++ b/Engine/Source/Runtime/Landscape/Private/LandscapeGrassWeightExporter.cpp
@@ -12,12 +12,17 @@
 #include "MeshPassProcessor.inl"
 #include "SimpleMeshDrawCommandPass.h"
 #include "TextureResource.h"
+#include "RenderCaptureInterface.h"
 
+#if WITH_EDITOR
 class FLandscapeGrassWeightVS;
 class FLandscapeGrassWeightPS;
-class FLandscapeGrassWeightShaderElementData;
 
-#if WITH_EDITOR
+int32 RenderCaptureNextGrassmapDraws = 0;
+static FAutoConsoleVariableRef CVarRenderCaptureNextGrassmapDraws(
+	TEXT("landscape.RenderCaptureNextGrassmapDraws"),
+	RenderCaptureNextGrassmapDraws,
+	TEXT("Trigger render captures during the next N grassmap draw calls."));
 
 BEGIN_SHADER_PARAMETER_STRUCT(FLandscapeGrassPassParameters, )
 	SHADER_PARAMETER_STRUCT_REF(FViewUniformShaderParameters, View)
@@ -26,6 +31,13 @@ BEGIN_SHADER_PARAMETER_STRUCT(FLandscapeGrassPassParameters, )
 	RENDER_TARGET_BINDING_SLOTS()
 	END_SHADER_PARAMETER_STRUCT()
 
+class FLandscapeGrassWeightShaderElementData : public FMeshMaterialShaderElementData
+{
+public:
+	int32 OutputPass;
+	FVector2f RenderOffset;
+};
+
 class FLandscapeGrassWeightMeshProcessor : public FMeshPassProcessor
 {
 public:
@@ -242,12 +254,17 @@ void FLandscapeGrassWeightExporter_RenderThread::RenderLandscapeComponentToTextu
 			}
 		});
 
+	if (bUseAsyncReadback)
+	{
+		AsyncReadback.StartReadback_RenderThread(GraphBuilder, OutputTexture);
+	}
+
 	GraphBuilder.Execute();
 }
 
 
-FLandscapeGrassWeightExporter::FLandscapeGrassWeightExporter(ALandscapeProxy* InLandscapeProxy, TArrayView<ULandscapeComponent* const> InLandscapeComponents, bool bInNeedsGrassmap, bool bInNeedsHeightmap, const TArray<int32>& InHeightMips)
-	: FLandscapeGrassWeightExporter_RenderThread(InHeightMips)
+FLandscapeGrassWeightExporter::FLandscapeGrassWeightExporter(ALandscapeProxy* InLandscapeProxy, TArrayView<ULandscapeComponent* const> InLandscapeComponents, bool bInNeedsGrassmap, bool bInNeedsHeightmap, const TArray<int32>& InHeightMips, bool bInUseAsyncReadback)
+	: FLandscapeGrassWeightExporter_RenderThread(InHeightMips, bInUseAsyncReadback)
 	, LandscapeProxy(InLandscapeProxy)
 	, ComponentSizeVerts(InLandscapeProxy->ComponentSizeQuads + 1)
 	, SubsectionSizeQuads(InLandscapeProxy->SubsectionSizeQuads)
@@ -318,6 +335,9 @@ FLandscapeGrassWeightExporter::FLandscapeGrassWeightExporter(ALandscapeProxy* In
 
 	UE::RenderCommandPipe::FSyncScope SyncScope;
 
+	RenderCaptureInterface::FScopedCapture RenderCapture((RenderCaptureNextGrassmapDraws != 0), TEXT("LandscapeGrassmapCapture"));
+	RenderCaptureNextGrassmapDraws = FMath::Max(0, RenderCaptureNextGrassmapDraws - 1);
+
 	// render
 	FLandscapeGrassWeightExporter_RenderThread* Exporter = this;
 	ENQUEUE_RENDER_COMMAND(FDrawSceneCommand)(
@@ -330,15 +350,26 @@ FLandscapeGrassWeightExporter::FLandscapeGrassWeightExporter(ALandscapeProxy* In
 
 TMap<ULandscapeComponent*, TUniquePtr<FLandscapeComponentGrassData>, TInlineSetAllocator<1>> FLandscapeGrassWeightExporter::FetchResults()
 {
+	TMap<ULandscapeComponent*, TUniquePtr<FLandscapeComponentGrassData>, TInlineSetAllocator<1>> Results;
 	TArray<FColor> Samples;
-	Samples.SetNumUninitialized(TargetSize.X * TargetSize.Y);
 
-	// Copy the contents of the remote texture to system memory
-	FReadSurfaceDataFlags ReadSurfaceDataFlags;
-	ReadSurfaceDataFlags.SetLinearToGamma(false);
-	RenderTargetResource->ReadPixels(Samples, ReadSurfaceDataFlags, FIntRect(0, 0, TargetSize.X, TargetSize.Y));
+	if (bUseAsyncReadback)
+	{
+		FIntPoint Size;
+		Samples = AsyncReadback.TakeResults(&Size);
+		check(Size == TargetSize);
+	}
+	else
+	{
+		Samples.SetNumUninitialized(TargetSize.X * TargetSize.Y);
+
+		// Copy the contents of the remote texture to system memory (SamplesDataBuffer)
+		// This is a synchronous operation that may stall the cpu, waiting for the GPU data to be generated and copied
+		FReadSurfaceDataFlags ReadSurfaceDataFlags;
+		ReadSurfaceDataFlags.SetLinearToGamma(false);
+		RenderTargetResource->ReadPixels(Samples, ReadSurfaceDataFlags, FIntRect(0, 0, TargetSize.X, TargetSize.Y));
+	}
 
-	TMap<ULandscapeComponent*, TUniquePtr<FLandscapeComponentGrassData>, TInlineSetAllocator<1>> Results;
 	Results.Reserve(ComponentInfos.Num());
 
 	// Local data will be moved in contiguous array at the end of export (to minimize slack waste)
@@ -348,7 +379,6 @@ TMap<ULandscapeComponent*, TUniquePtr<FLandscapeComponentGrassData>, TInlineSetA
 	for (auto& ComponentInfo : ComponentInfos)
 	{
 		ULandscapeComponent* Component = ComponentInfo.Component;
-		ALandscapeProxy* Proxy = Component->GetLandscapeProxy();
 
 		TUniquePtr<FLandscapeComponentGrassData> NewGrassData = MakeUnique<FLandscapeComponentGrassData>(Component);
 
@@ -360,7 +390,9 @@ TMap<ULandscapeComponent*, TUniquePtr<FLandscapeComponentGrassData>, TInlineSetA
 		{
 			HeightData.Empty(0);
 		}
+#if WITH_EDITORONLY_DATA
 		NewGrassData->HeightMipData.Empty(HeightMips.Num());
+#endif // WITH_EDITORONLY_DATA
 
 		WeightData.Empty();
 		TArray<TArray<uint8>*> GrassWeightArrays;
@@ -391,6 +423,19 @@ TMap<ULandscapeComponent*, TUniquePtr<FLandscapeComponentGrassData>, TInlineSetA
 		}
 #endif
 
+		// FColor memory layout matches our BGRA GPU layout only on little endian CPUs!
+		#if PLATFORM_LITTLE_ENDIAN
+			#define BGRA_AS_FCOLOR_BLUE B
+			#define BGRA_AS_FCOLOR_GREEN G
+			#define BGRA_AS_FCOLOR_RED R
+			#define BGRA_AS_FCOLOR_ALPHA A
+		#else
+			#define BGRA_AS_FCOLOR_BLUE A
+			#define BGRA_AS_FCOLOR_GREEN R
+			#define BGRA_AS_FCOLOR_RED G
+			#define BGRA_AS_FCOLOR_ALPHA B
+		#endif // PLATFORM_LITTLE_ENDIAN
+
 		for (int32 PassIdx = 0; PassIdx < ComponentInfo.NumPasses; PassIdx++)
 		{
 			FColor* SampleData = &Samples[ComponentInfo.PixelOffsetX + PassIdx * ComponentSizeVerts];
@@ -403,14 +448,14 @@ TMap<ULandscapeComponent*, TUniquePtr<FLandscapeComponentGrassData>, TInlineSetA
 						for (int32 x = 0; x < ComponentSizeVerts; x++)
 						{
 							FColor& Sample = SampleData[x + y * TargetSize.X];
-							uint16 Height = (((uint16)Sample.R) << 8) + (uint16)(Sample.G);
+							uint16 Height = (((uint16)Sample.BGRA_AS_FCOLOR_RED) << 8) + (uint16)(Sample.BGRA_AS_FCOLOR_GREEN);
 							HeightData.Add(Height);
 							if (ComponentInfo.RequestedGrassTypes.Num() > 0)
 							{
-								GrassWeightArrays[0]->Add(Sample.B);
+								GrassWeightArrays[0]->Add(Sample.BGRA_AS_FCOLOR_BLUE);
 								if (ComponentInfo.RequestedGrassTypes.Num() > 1)
 								{
-									GrassWeightArrays[1]->Add(Sample.A);
+									GrassWeightArrays[1]->Add(Sample.BGRA_AS_FCOLOR_ALPHA);
 								}
 							}
 						}
@@ -425,16 +470,16 @@ TMap<ULandscapeComponent*, TUniquePtr<FLandscapeComponentGrassData>, TInlineSetA
 							FColor& Sample = SampleData[x + y * TargetSize.X];
 
 							int32 TypeIdx = PassIdx * 4 - 2;
-							GrassWeightArrays[TypeIdx++]->Add(Sample.R);
+							GrassWeightArrays[TypeIdx++]->Add(Sample.BGRA_AS_FCOLOR_RED);
 							if (TypeIdx < ComponentInfo.RequestedGrassTypes.Num())
 							{
-								GrassWeightArrays[TypeIdx++]->Add(Sample.G);
+								GrassWeightArrays[TypeIdx++]->Add(Sample.BGRA_AS_FCOLOR_GREEN);
 								if (TypeIdx < ComponentInfo.RequestedGrassTypes.Num())
 								{
-									GrassWeightArrays[TypeIdx++]->Add(Sample.B);
+									GrassWeightArrays[TypeIdx++]->Add(Sample.BGRA_AS_FCOLOR_BLUE);
 									if (TypeIdx < ComponentInfo.RequestedGrassTypes.Num())
 									{
-										GrassWeightArrays[TypeIdx++]->Add(Sample.A);
+										GrassWeightArrays[TypeIdx++]->Add(Sample.BGRA_AS_FCOLOR_ALPHA);
 									}
 								}
 							}
@@ -444,6 +489,7 @@ TMap<ULandscapeComponent*, TUniquePtr<FLandscapeComponentGrassData>, TInlineSetA
 			}
 			else // PassIdx >= FirstHeightMipsPassIndex
 			{
+#if WITH_EDITORONLY_DATA
 				const int32 Mip = HeightMips[PassIdx - ComponentInfo.FirstHeightMipsPassIndex];
 				int32 MipSizeVerts = NumSubsections * (SubsectionSizeQuads >> Mip);
 				TArray<uint16>& MipHeightData = NewGrassData->HeightMipData.Add(Mip);
@@ -452,13 +498,19 @@ TMap<ULandscapeComponent*, TUniquePtr<FLandscapeComponentGrassData>, TInlineSetA
 					for (int32 x = 0; x < MipSizeVerts; x++)
 					{
 						FColor& Sample = SampleData[x + y * TargetSize.X];
-						uint16 Height = (((uint16)Sample.R) << 8) + (uint16)(Sample.G);
+						uint16 Height = (((uint16)Sample.BGRA_AS_FCOLOR_RED) << 8) + (uint16)(Sample.BGRA_AS_FCOLOR_GREEN);
 						MipHeightData.Add(Height);
 					}
 				}
+#endif // WITH_EDITORONLY_DATA
 			}
 		}
 
+		#undef BGRA_AS_FCOLOR_BLUE
+		#undef BGRA_AS_FCOLOR_GREEN
+		#undef BGRA_AS_FCOLOR_RED
+		#undef BGRA_AS_FCOLOR_ALPHA
+
 		// remove null grass type if we had one (can occur if the node has null entries)
 		WeightData.Remove(nullptr);
 
@@ -490,13 +542,18 @@ void FLandscapeGrassWeightExporter::ApplyResults()
 
 		// Assign the new data (thread-safe)
 		Component->GrassData = MakeShareable(ComponentGrassData);
+
+#if WITH_EDITORONLY_DATA
 		Component->GrassData->bIsDirty = true;
+#endif // WITH_EDITORONLY_DATA
 
+#if WITH_EDITOR
 		if (Proxy->bBakeMaterialPositionOffsetIntoCollision)
 		{
 			Component->DestroyCollisionData();
 			Component->UpdateCollisionData();
 		}
+#endif // WITH_EDITOR
 	}
 }
 
diff --git a/Engine/Source/Runtime/Landscape/Private/LandscapeGrassWeightExporter.h b/Engine/Source/Runtime/Landscape/Private/LandscapeGrassWeightExporter.h
--- a/Engine/Source/Runtime/Landscape/Private/LandscapeGrassWeightExporter.h
+++ b/Engine/Source/Runtime/Landscape/Private/LandscapeGrassWeightExporter.h
@@ -4,6 +4,7 @@
 #include "CoreMinimal.h"
 #include "LandscapeComponent.h"
 #include "MeshMaterialShader.h"
+#include "LandscapeAsyncTextureReadback.h"
 
 class ALandscapeProxy;
 class FLandscapeComponentSceneProxy;
@@ -13,8 +14,9 @@ class UTextureRenderTarget2D;
 // data also accessible by render thread
 class FLandscapeGrassWeightExporter_RenderThread
 {
-	FLandscapeGrassWeightExporter_RenderThread(const TArray<int32>& InHeightMips)
+	FLandscapeGrassWeightExporter_RenderThread(const TArray<int32>& InHeightMips, bool bInUseAsyncReadback)
 		: HeightMips(InHeightMips)
+		, bUseAsyncReadback(bInUseAsyncReadback)
 	{
 	}
 
@@ -65,15 +67,10 @@ public:
 	FMatrix ViewRotationMatrix;
 	FMatrix ProjectionMatrix;
 
-	void RenderLandscapeComponentToTexture_RenderThread(FRHICommandListImmediate& RHICmdList);
-};
-
-class FLandscapeGrassWeightShaderElementData : public FMeshMaterialShaderElementData
-{
-public:
+	bool bUseAsyncReadback;
+	FLandscapeAsyncTextureReadback AsyncReadback;
 
-	int32 OutputPass;
-	FVector2f RenderOffset;
+	void RenderLandscapeComponentToTexture_RenderThread(FRHICommandListImmediate& RHICmdList);
 };
 
 class FLandscapeGrassWeightExporter : public FLandscapeGrassWeightExporter_RenderThread
@@ -86,9 +83,29 @@ class FLandscapeGrassWeightExporter : public FLandscapeGrassWeightExporter_Rende
 	TObjectPtr<UTextureRenderTarget2D> RenderTargetTexture;
 
 public:
-	FLandscapeGrassWeightExporter(ALandscapeProxy* InLandscapeProxy, TArrayView<ULandscapeComponent* const> InLandscapeComponents, bool bInNeedsGrassmap = true, bool bInNeedsHeightmap = true, const TArray<int32>& InHeightMips = {});
+	FLandscapeGrassWeightExporter(ALandscapeProxy* InLandscapeProxy, TArrayView<ULandscapeComponent* const> InLandscapeComponents, bool bInNeedsGrassmap = true, bool bInNeedsHeightmap = true, const TArray<int32>& InHeightMips = {}, bool bUseAsyncReadback = false);
+
+	// If using the async readback path, check its status and update if needed. Return true when the AsyncReadbackResults are available.
+	// You must call this periodically, or the async readback may not complete.
+	bool CheckAndUpdateAsyncReadback()
+	{
+		return AsyncReadback.CheckAndUpdate();
+	}
+
+	// return true if the async readback is complete.  (Does not update the readback state)
+	bool IsAsyncReadbackComplete()
+	{
+		return AsyncReadback.IsComplete();
+	}
+
+	// Fetches the results from the GPU texture and translates them into FLandscapeComponentGrassDatas.
+	// If using async readback, requires AsyncReadback to be complete before calling this.
 	TMap<ULandscapeComponent*, TUniquePtr<FLandscapeComponentGrassData>, TInlineSetAllocator<1>> FetchResults();
+
+	// Fetches the results and applies them to the landscape components
+	// If using async readback, requires AsyncReadback to be complete before calling this.
 	void ApplyResults();
+
 	void AddReferencedObjects(UObject* InThis, FReferenceCollector& Collector);
 };
 
