Optimized FAsyncGrassBuilder to increase chances of early out in exclude boxes

diff --git a/Engine/Source/Runtime/Core/Public/Math/Box.h b/Engine/Source/Runtime/Core/Public/Math/Box.h
--- a/Engine/Source/Runtime/Core/Public/Math/Box.h
+++ b/Engine/Source/Runtime/Core/Public/Math/Box.h
@@ -493,6 +493,21 @@ public:
 		return (IsInside(Other.Min) && IsInside(Other.Max));
 	}
 
+	/** 
+	 * Checks whether a given box is fully encapsulated by this box.
+	 * 
+	 * @param Other The box to test for encapsulation within the bounding volume.
+	 * @return true if box is inside this volume.
+	 * @see IsInsideXY
+	 *
+	 * @note  This function assumes boxes have closed bounds, i.e. boxes with
+	 *        coincident borders on any edge are encapsulated.
+	 */
+	FORCEINLINE bool IsInsideOrOn( const TBox<T>& Other ) const
+	{
+		return (IsInsideOrOn(Other.Min) && IsInsideOrOn(Other.Max));
+	}
+
 	/** 
 	 * Checks whether the given location is inside this box in the XY plane.
 	 * 
diff --git a/Engine/Source/Runtime/Landscape/Classes/LandscapeComponent.h b/Engine/Source/Runtime/Landscape/Classes/LandscapeComponent.h
--- a/Engine/Source/Runtime/Landscape/Classes/LandscapeComponent.h
+++ b/Engine/Source/Runtime/Landscape/Classes/LandscapeComponent.h
@@ -643,7 +643,21 @@ public:
 public:
 	/** Grass data for generation **/
 	TSharedRef<FLandscapeComponentGrassData, ESPMode::ThreadSafe> GrassData;
-	TArray<FBox> ActiveExcludedBoxes;
+	
+	// This wrapper is needed to filter out exclude boxes that are completely inside of another exclude box
+	struct FExcludeBox
+	{
+		FBox Box;
+
+		FExcludeBox() = default;
+		FExcludeBox(const FBox& InBox) : Box(InBox) {}
+
+		bool operator==(const FExcludeBox& Other) const
+		{
+			return Box.IsInsideOrOn(Other.Box);
+		}
+	};
+	TArray<FExcludeBox> ActiveExcludedBoxes;
 	uint32 ChangeTag;
 
 #if WITH_EDITOR
diff --git a/Engine/Source/Runtime/Landscape/Private/LandscapeGrass.cpp b/Engine/Source/Runtime/Landscape/Private/LandscapeGrass.cpp
--- a/Engine/Source/Runtime/Landscape/Private/LandscapeGrass.cpp
+++ b/Engine/Source/Runtime/Landscape/Private/LandscapeGrass.cpp
@@ -2409,9 +2409,9 @@ static uint32 GGrassExclusionChangeTag = 1;
 static uint32 GFrameNumberLastStaleCheck = 0;
 static TMap<FWeakObjectPtr, FBox, FDefaultSetAllocator, TWeakObjectPtrMapKeyFuncs<FWeakObjectPtr, FBox>> GGrassExclusionBoxes;
 
-void ALandscapeProxy::AddExclusionBox(FWeakObjectPtr Owner, const FBox& BoxToRemove)
+void ALandscapeProxy::AddExclusionBox(FWeakObjectPtr Owner, const FBox& Box)
 {
-	GGrassExclusionBoxes.Add(Owner, BoxToRemove);
+	GGrassExclusionBoxes.Add(Owner, Box);
 	GGrassExclusionChangeTag++;
 }
 void ALandscapeProxy::RemoveExclusionBox(FWeakObjectPtr Owner)
@@ -2820,9 +2820,16 @@ void ALandscapeProxy::UpdateGrass(const TArray<FVector>& Cameras, int32& InOutNu
 						{
 							if (Pair.Value.Intersect(WorldBox))
 							{
-								Component->ActiveExcludedBoxes.AddUnique(Pair.Value);
+								// this will also filter out boxes that are completely inside of the other boxes
+								Component->ActiveExcludedBoxes.AddUnique(ULandscapeComponent::FExcludeBox(Pair.Value));
 							}
 						}
+
+						//Sort exclude boxes by their volume size, as the biggest boxes are more likely to exclude points and early out when building grass
+						Component->ActiveExcludedBoxes.Sort([](const ULandscapeComponent::FExcludeBox& A, const ULandscapeComponent::FExcludeBox& B) {
+								return A.Box.GetVolume() > B.Box.GetVolume();
+							}
+						);
 					}
 					Component->ChangeTag = GGrassExclusionChangeTag;
 				}
@@ -2933,11 +2940,11 @@ void ALandscapeProxy::UpdateGrass(const TArray<FVector>& Cameras, int32& InOutNu
 											FCachedLandscapeFoliage::FGrassComp* Existing = FoliageCache.CachedGrassComps.Find(NewComp.Key);
 											if (Existing && !Existing->PreviousFoliage.IsValid() && Existing->ExclusionChangeTag != GGrassExclusionChangeTag && !Existing->PendingRemovalRebuild && !Existing->Pending)
 											{
-												for (const FBox& Box : Component->ActiveExcludedBoxes)
+												for (const ULandscapeComponent::FExcludeBox& BoxWrapper : Component->ActiveExcludedBoxes)
 												{
-													if (Box.Intersect(WorldSubBox))
+													if (BoxWrapper.Box.Intersect(WorldSubBox))
 													{
-														NewComp.ExcludedBoxes.Add(Box);
+														NewComp.ExcludedBoxes.Add(BoxWrapper.Box);
 													}
 												}
 												if (NewComp.ExcludedBoxes != Existing->ExcludedBoxes)
@@ -2971,11 +2978,11 @@ void ALandscapeProxy::UpdateGrass(const TArray<FVector>& Cameras, int32& InOutNu
 										}
 										if (!bRebuildForBoxes)
 										{
-											for (const FBox& Box : Component->ActiveExcludedBoxes)
+											for (const ULandscapeComponent::FExcludeBox& BoxWrapper : Component->ActiveExcludedBoxes)
 											{
-												if (Box.Intersect(WorldSubBox))
+												if (BoxWrapper.Box.Intersect(WorldSubBox))
 												{
-													NewComp.ExcludedBoxes.Add(Box);
+													NewComp.ExcludedBoxes.Add(BoxWrapper.Box);
 												}
 											}
 										}
