Nanite skirts improvements : 	when trying to maintain the slope at the edge by extrapolating the skirt vertex position, we now only do it if the slope is going downwards from the edge vertex to the skirt vertex (otherwise, in case of steep slopes, the skirt depth might not be enough to bring the skirt's vertex underneath the neighboring landscape proxy)

#localization none
#tests none

diff --git a/Engine/Source/Runtime/Landscape/Private/LandscapeEdit.cpp b/Engine/Source/Runtime/Landscape/Private/LandscapeEdit.cpp
--- a/Engine/Source/Runtime/Landscape/Private/LandscapeEdit.cpp
+++ b/Engine/Source/Runtime/Landscape/Private/LandscapeEdit.cpp
@@ -4008,7 +4008,9 @@ bool ALandscapeProxy::ExportToRawMeshDataCopyNew(const FRawMeshExportParams& InE
 				Position.X = CDI.GetScaleFactor() * VertexX;
 				Position.Y = CDI.GetScaleFactor() * VertexY;
 
-				Position -= FVector(0.0f, 0.0f, Diff + SkirtDepth);
+				// Maintain the slope at the edge by extrapolating the skirt vertex position but only if the slope is going downwards from the edge vertex to the skirt vertex
+				//  (otherwise, in case of steep slopes, the skirt depth might not be enough to bring the skirt's vertex underneath the neighboring landscape proxy)
+				Position -= FVector(0.0f, 0.0f, FMath::Max(Diff, 0.0f) + SkirtDepth);
 			}
 			return Position;
 		};
