Fix an issue with Landscape materials where they would not cache Nanite rasterization shaders.

diff --git a/Engine/Source/Runtime/Landscape/Private/LandscapeRender.cpp b/Engine/Source/Runtime/Landscape/Private/LandscapeRender.cpp
--- a/Engine/Source/Runtime/Landscape/Private/LandscapeRender.cpp
+++ b/Engine/Source/Runtime/Landscape/Private/LandscapeRender.cpp
@@ -3646,71 +3646,60 @@ public:
 
 	bool ShouldCache(EShaderPlatform Platform, const FShaderType* ShaderType, const FVertexFactoryType* VertexFactoryType) const override
 	{
-		// Don't compile if this is a mobile shadermap and a desktop MIC, and vice versa, unless it's a tool material
-		if (!(IsPCPlatform(Platform) && bEditorToolUsage) && bMobile != IsMobilePlatform(Platform))
-		{
-			// @todo For some reason this causes this resource to return true for IsCompilationFinished. For now we will needlessly compile this shader until this is fixed.
-			//return false;
-		}
-
-		if (VertexFactoryType)
+		if (bIsLayerThumbnail)
 		{
 			// Always check against FLocalVertexFactory in editor builds as it is required to render thumbnails
 			// Thumbnail MICs are only rendered in the preview scene using a simple LocalVertexFactory
-			if (bIsLayerThumbnail)
+			static const FName LocalVertexFactory = FName(TEXT("FLocalVertexFactory"));
+			if (!IsMobilePlatform(Platform) && (VertexFactoryType != nullptr) && (VertexFactoryType->GetFName() == LocalVertexFactory))
 			{
-				static const FName LocalVertexFactory = FName(TEXT("FLocalVertexFactory"));
-				if (!IsMobilePlatform(Platform) && VertexFactoryType->GetFName() == LocalVertexFactory)
+				if (Algo::Find(GetAllowedShaderTypesInThumbnailRender(), ShaderType->GetFName()))
 				{
-					if (Algo::Find(GetAllowedShaderTypesInThumbnailRender(), ShaderType->GetFName()))
+					return FMaterialResource::ShouldCache(Platform, ShaderType, VertexFactoryType);
+				}
+				else
+				{
+					// No ray tracing on thumbnails : we don't need any variation of ray hit group shaders : 
+					const bool bIsRayHitGroupShader = (ShaderType->GetFrequency() == SF_RayHitGroup);
+					if (bIsRayHitGroupShader
+						|| Algo::Find(GetExcludedShaderTypesInThumbnailRender(), ShaderType->GetFName()))
 					{
-						return FMaterialResource::ShouldCache(Platform, ShaderType, VertexFactoryType);
+						UE_LOG(LogLandscape, VeryVerbose, TEXT("Excluding shader %s from landscape thumbnail material"), ShaderType->GetName());
+						return false;
 					}
 					else
 					{
-						// No ray tracing on thumbnails : we don't need any variation of ray hit group shaders : 
-						const bool bIsRayHitGroupShader = (ShaderType->GetFrequency() == SF_RayHitGroup);
-						if (bIsRayHitGroupShader
-							|| Algo::Find(GetExcludedShaderTypesInThumbnailRender(), ShaderType->GetFName()))
-						{
-							UE_LOG(LogLandscape, VeryVerbose, TEXT("Excluding shader %s from landscape thumbnail material"), ShaderType->GetName());
-							return false;
-						}
-						else
-						{
-							UE_LOG(LogLandscape, Warning, TEXT("Shader %s unknown by landscape thumbnail material, please add to either AllowedShaderTypes or ExcludedShaderTypes"), ShaderType->GetName());
-							return FMaterialResource::ShouldCache(Platform, ShaderType, VertexFactoryType);
-						}
+						UE_LOG(LogLandscape, Warning, TEXT("Shader %s unknown by landscape thumbnail material, please add to either AllowedShaderTypes or ExcludedShaderTypes"), ShaderType->GetName());
+						return FMaterialResource::ShouldCache(Platform, ShaderType, VertexFactoryType);
 					}
 				}
 			}
-			else
+		}
+		else if (FMaterialResource::ShouldCache(Platform, ShaderType, VertexFactoryType))
+		{
+			// Landscape MICs are only for use with the Landscape vertex factories
+			if (VertexFactoryType && VertexFactoryType->SupportsLandscape())
 			{
-				// Landscape MICs are only for use with the Landscape vertex factories
-
 				// For now only compile FLandscapeFixedGridVertexFactory for grass and runtime virtual texture page rendering (can change if we need for other cases)
 				// Todo: only compile LandscapeXYOffsetVertexFactory if we are using it
 				const bool bIsGrassShaderType = Algo::Find(GetGrassShaderTypes(), ShaderType->GetFName()) != nullptr;
 				const bool bIsGPULightmassShaderType = Algo::Find(GetGPULightmassShaderTypes(), ShaderType->GetFName()) != nullptr;
 				const bool bIsRuntimeVirtualTextureShaderType = Algo::Find(GetRuntimeVirtualTextureShaderTypes(), ShaderType->GetFName()) != nullptr;
 				const bool bIsLumen = Algo::Find(GetLumenCardShaderTypes(), ShaderType->GetFName()) != nullptr;
-
 				const bool bIsShaderTypeUsingFixedGrid = bIsGrassShaderType || bIsRuntimeVirtualTextureShaderType || bIsGPULightmassShaderType || bIsLumen;
+				const bool bIsFixedGridVertexFactory = (&FLandscapeFixedGridVertexFactory::StaticType == VertexFactoryType);
 				
 				static const FName RayTracingDynamicGeometryConverterCS = FName(TEXT("FRayTracingDynamicGeometryConverterCS"));
 				const bool bIsRayTracingShaderType = ShaderType->GetFName() == RayTracingDynamicGeometryConverterCS;
 
-				if (VertexFactoryType->SupportsLandscape())
-				{
-					if (&FLandscapeFixedGridVertexFactory::StaticType == VertexFactoryType)
-					{
-						return (bIsRayTracingShaderType || bIsShaderTypeUsingFixedGrid) && FMaterialResource::ShouldCache(Platform, ShaderType, VertexFactoryType);
-					}
-					else
-					{
-						return (bIsRayTracingShaderType || !bIsShaderTypeUsingFixedGrid) && FMaterialResource::ShouldCache(Platform, ShaderType, VertexFactoryType);
-					}
-				}
+				return (bIsRayTracingShaderType 
+					|| (bIsFixedGridVertexFactory == bIsShaderTypeUsingFixedGrid));
+			}
+			else
+			{
+				// Allow the landscape MICs on rasterizer shader types as well, 
+				const bool bIsRasterizerShaderType = Algo::Find(GetRasterizeShaderTypes(), ShaderType->GetFName()) != nullptr;
+				return bIsRasterizerShaderType;
 			}
 		}
 
@@ -3963,6 +3952,18 @@ public:
 		return ShaderTypes;
 	}
 
+	static const TArray<FName>& GetRasterizeShaderTypes()
+	{
+		static const TArray<FName> ShaderTypes =
+		{
+			FName(TEXT("FHWRasterizeVS")),
+			FName(TEXT("FHWRasterizeMS")),
+			FName(TEXT("FHWRasterizePS")),
+			FName(TEXT("FMicropolyRasterizeCS")),
+		};
+		return ShaderTypes;
+	}
+
 };
 
 FMaterialResource* ULandscapeMaterialInstanceConstant::AllocatePermutationResource()
