Remove invalid check for grass exclusion boxes every frame - instead only remove invalid entries after garbage collection.

diff --git a/Engine/Source/Runtime/Landscape/Classes/LandscapeProxy.h b/Engine/Source/Runtime/Landscape/Classes/LandscapeProxy.h
--- a/Engine/Source/Runtime/Landscape/Classes/LandscapeProxy.h
+++ b/Engine/Source/Runtime/Landscape/Classes/LandscapeProxy.h
@@ -1036,6 +1036,8 @@ public:
 	 */
 	LANDSCAPE_API static void RemoveAllExclusionBoxes();
 
+	static void RemoveInvalidExclusionBoxes();
+
 	/* Invalidate the precomputed grass and baked texture data for the specified components */
 	LANDSCAPE_API static void InvalidateGeneratedComponentData(const TSet<ULandscapeComponent*>& Components, bool bInvalidateLightingCache = false);
 	LANDSCAPE_API static void InvalidateGeneratedComponentData(const TArray<ULandscapeComponent*>& Components, bool bInvalidateLightingCache = false);
diff --git a/Engine/Source/Runtime/Landscape/Private/LandscapeGrass.cpp b/Engine/Source/Runtime/Landscape/Private/LandscapeGrass.cpp
--- a/Engine/Source/Runtime/Landscape/Private/LandscapeGrass.cpp
+++ b/Engine/Source/Runtime/Landscape/Private/LandscapeGrass.cpp
@@ -2406,7 +2406,6 @@ void ALandscapeProxy::FlushGrassComponents(const TSet<ULandscapeComponent*>* Onl
 }
 
 static uint32 GGrassExclusionChangeTag = 1;
-static uint32 GFrameNumberLastStaleCheck = 0;
 static TMap<FWeakObjectPtr, FBox, FDefaultSetAllocator, TWeakObjectPtrMapKeyFuncs<FWeakObjectPtr, FBox>> GGrassExclusionBoxes;
 
 void ALandscapeProxy::AddExclusionBox(FWeakObjectPtr Owner, const FBox& Box)
@@ -2428,6 +2427,18 @@ void ALandscapeProxy::RemoveAllExclusionBoxes()
 	}
 }
 
+void ALandscapeProxy::RemoveInvalidExclusionBoxes()
+{
+	for (auto Iter = GGrassExclusionBoxes.CreateIterator(); Iter; ++Iter)
+	{
+		if (!Iter->Key.IsValid())
+		{
+			Iter.RemoveCurrent();
+			GGrassExclusionChangeTag++;
+		}
+	}
+}
+
 #if WITH_EDITOR
 
 void ALandscapeProxy::BuildGrassMaps(FScopedSlowTask* InSlowTask)
@@ -2691,19 +2702,6 @@ void ALandscapeProxy::UpdateGrass(const TArray<FVector>& Cameras, int32& InOutNu
 		return;
 	}
 
-	if (GFrameNumberLastStaleCheck != GFrameNumber && GIgnoreExcludeBoxes == 0)
-	{
-		GFrameNumberLastStaleCheck = GFrameNumber;
-		for (auto Iter = GGrassExclusionBoxes.CreateIterator(); Iter; ++Iter)
-		{
-			if (!Iter->Key.IsValid())
-			{
-				Iter.RemoveCurrent();
-				GGrassExclusionChangeTag++;
-			}
-		}
-	}
-
 	UWorld* World = GetWorld();
 	ULandscapeSubsystem* LandscapeSubsystem = World ? World->GetSubsystem<ULandscapeSubsystem>() : nullptr;
 	const bool bIsGrassCreationPrioritized = LandscapeSubsystem ? LandscapeSubsystem->IsGrassCreationPrioritized() : false;
@@ -2820,8 +2818,11 @@ void ALandscapeProxy::UpdateGrass(const TArray<FVector>& Cameras, int32& InOutNu
 						{
 							if (Pair.Value.Intersect(WorldBox))
 							{
-								// this will also filter out boxes that are completely inside of the other boxes
-								Component->ActiveExcludedBoxes.AddUnique(ULandscapeComponent::FExcludeBox(Pair.Value));
+								if (Pair.Key.IsValid())
+								{
+									// this will also filter out boxes that are completely inside of the other boxes
+									Component->ActiveExcludedBoxes.AddUnique(ULandscapeComponent::FExcludeBox(Pair.Value));
+								}
 							}
 						}
 
diff --git a/Engine/Source/Runtime/Landscape/Private/LandscapeSubsystem.cpp b/Engine/Source/Runtime/Landscape/Private/LandscapeSubsystem.cpp
--- a/Engine/Source/Runtime/Landscape/Private/LandscapeSubsystem.cpp
+++ b/Engine/Source/Runtime/Landscape/Private/LandscapeSubsystem.cpp
@@ -126,10 +126,13 @@ void ULandscapeSubsystem::Initialize(FSubsystemCollectionBase& Collection)
     	}
     });
     	
+	FCoreUObjectDelegates::GetPostGarbageCollect().AddUObject(this, &ULandscapeSubsystem::HandlePostGarbageCollect);
 }
 
 void ULandscapeSubsystem::Deinitialize()
 {
+	FCoreUObjectDelegates::GetPostGarbageCollect().RemoveAll(this);
+
 	if (OnNaniteWorldSettingsChangedHandle.IsValid())
 	{
 		UWorld* World = GetWorld();
@@ -174,6 +177,11 @@ void ULandscapeSubsystem::Deinitialize()
 	Super::Deinitialize();
 }
 
+void ULandscapeSubsystem::HandlePostGarbageCollect()
+{
+	ALandscapeProxy::RemoveInvalidExclusionBoxes();
+}
+
 TStatId ULandscapeSubsystem::GetStatId() const
 {
 	RETURN_QUICK_DECLARE_CYCLE_STAT(ULandscapeSubsystem, STATGROUP_Tickables);
diff --git a/Engine/Source/Runtime/Landscape/Public/LandscapeSubsystem.h b/Engine/Source/Runtime/Landscape/Public/LandscapeSubsystem.h
--- a/Engine/Source/Runtime/Landscape/Public/LandscapeSubsystem.h
+++ b/Engine/Source/Runtime/Landscape/Public/LandscapeSubsystem.h
@@ -131,6 +131,8 @@ private:
 	void OnNaniteWorldSettingsChanged(AWorldSettings* WorldSettings) { RegenerateGrass(true, true); }
 	void OnNaniteEnabledChanged(IConsoleVariable*);
 
+	void HandlePostGarbageCollect();
+
 	bool bIsGrassCreationPrioritized = false;
 	TArray<TWeakObjectPtr<ALandscapeProxy>> Proxies;
 	FDelegateHandle OnNaniteWorldSettingsChangedHandle;
