Removed unused landscape code
#tests editor

diff --git a/Engine/Source/Editor/LandscapeEditor/Private/LandscapeEdMode.h b/Engine/Source/Editor/LandscapeEditor/Private/LandscapeEdMode.h
--- a/Engine/Source/Editor/LandscapeEditor/Private/LandscapeEdMode.h
+++ b/Engine/Source/Editor/LandscapeEditor/Private/LandscapeEdMode.h
@@ -67,7 +67,6 @@ struct FLandscapeTargetListInfo
 	ELandscapeToolTargetType TargetType;
 	TWeakObjectPtr<ULandscapeInfo> LandscapeInfo;
 
-	//Values cloned from FLandscapeLayerStruct LayerStruct;			// ignored for heightmap
 	TWeakObjectPtr<ULandscapeLayerInfoObject> LayerInfoObj;			// ignored for heightmap
 	FName LayerName;												// ignored for heightmap
 	TWeakObjectPtr<class ALandscapeProxy> Owner;					// ignored for heightmap
diff --git a/Engine/Source/Runtime/Landscape/Classes/LandscapeProxy.h b/Engine/Source/Runtime/Landscape/Classes/LandscapeProxy.h
--- a/Engine/Source/Runtime/Landscape/Classes/LandscapeProxy.h
+++ b/Engine/Source/Runtime/Landscape/Classes/LandscapeProxy.h
@@ -103,44 +103,6 @@ struct FLandscapeEditorLayerSettings
 #endif // WITH_EDITORONLY_DATA
 };
 
-USTRUCT()
-struct FLandscapeLayerStruct
-{
-	GENERATED_USTRUCT_BODY()
-
-	UPROPERTY()
-	TObjectPtr<ULandscapeLayerInfoObject> LayerInfoObj;
-
-#if WITH_EDITORONLY_DATA
-	UPROPERTY(transient)
-	TObjectPtr<ULandscapeMaterialInstanceConstant> ThumbnailMIC;
-
-	UPROPERTY()
-	TObjectPtr<ALandscapeProxy> Owner;
-
-	UPROPERTY(transient)
-	int32 DebugColorChannel;
-
-	UPROPERTY(transient)
-	uint32 bSelected:1;
-
-	UPROPERTY()
-	FString SourceFilePath;
-#endif // WITH_EDITORONLY_DATA
-
-	FLandscapeLayerStruct()
-		: LayerInfoObj(nullptr)
-#if WITH_EDITORONLY_DATA
-		, ThumbnailMIC(nullptr)
-		, Owner(nullptr)
-		, DebugColorChannel(0)
-		, bSelected(false)
-		, SourceFilePath()
-#endif // WITH_EDITORONLY_DATA
-	{
-	}
-};
-
 UENUM()
 enum class ELandscapeImportAlphamapType : uint8
 {
diff --git a/Engine/Source/Runtime/Landscape/Private/Landscape.cpp b/Engine/Source/Runtime/Landscape/Private/Landscape.cpp
--- a/Engine/Source/Runtime/Landscape/Private/Landscape.cpp
+++ b/Engine/Source/Runtime/Landscape/Private/Landscape.cpp
@@ -979,25 +979,25 @@ void ULandscapeComponent::GetLayerDebugColorKey(int32& R, int32& G, int32& B) co
 
 		for (auto It = Info->Layers.CreateConstIterator(); It; It++)
 		{
-			const FLandscapeInfoLayerSettings& LayerStruct = *It;
-			if (LayerStruct.DebugColorChannel > 0
-				&& LayerStruct.LayerInfoObj)
+			const FLandscapeInfoLayerSettings& LayerSettings = *It;
+			if (LayerSettings.DebugColorChannel > 0
+				&& LayerSettings.LayerInfoObj)
 			{
 				const TArray<FWeightmapLayerAllocationInfo>& ComponentWeightmapLayerAllocations = GetWeightmapLayerAllocations();
 
 				for (int32 LayerIdx = 0; LayerIdx < ComponentWeightmapLayerAllocations.Num(); LayerIdx++)
 				{
-					if (ComponentWeightmapLayerAllocations[LayerIdx].LayerInfo == LayerStruct.LayerInfoObj)
+					if (ComponentWeightmapLayerAllocations[LayerIdx].LayerInfo == LayerSettings.LayerInfoObj)
 					{
-						if (LayerStruct.DebugColorChannel & 1) // R
+						if (LayerSettings.DebugColorChannel & 1) // R
 						{
 							R = (ComponentWeightmapLayerAllocations[LayerIdx].WeightmapTextureIndex * 4 + ComponentWeightmapLayerAllocations[LayerIdx].WeightmapTextureChannel);
 						}
-						if (LayerStruct.DebugColorChannel & 2) // G
+						if (LayerSettings.DebugColorChannel & 2) // G
 						{
 							G = (ComponentWeightmapLayerAllocations[LayerIdx].WeightmapTextureIndex * 4 + ComponentWeightmapLayerAllocations[LayerIdx].WeightmapTextureChannel);
 						}
-						if (LayerStruct.DebugColorChannel & 4) // B
+						if (LayerSettings.DebugColorChannel & 4) // B
 						{
 							B = (ComponentWeightmapLayerAllocations[LayerIdx].WeightmapTextureIndex * 4 + ComponentWeightmapLayerAllocations[LayerIdx].WeightmapTextureChannel);
 						}
@@ -2652,20 +2652,6 @@ FArchive& operator<<(FArchive& Ar, FLandscapeAddCollision& U)
 }
 #endif // WITH_EDITORONLY_DATA
 
-FArchive& operator<<(FArchive& Ar, FLandscapeLayerStruct*& L)
-{
-	if (L)
-	{
-		Ar << L->LayerInfoObj;
-#if WITH_EDITORONLY_DATA
-		return Ar << L->ThumbnailMIC;
-#else
-		return Ar;
-#endif // WITH_EDITORONLY_DATA
-	}
-	return Ar;
-}
-
 void ULandscapeInfo::Serialize(FArchive& Ar)
 {
 	Super::Serialize(Ar);
