Fixed issues with a landscape culling at certain landscape configuration

diff --git a/Engine/Source/Runtime/Landscape/Private/LandscapeCulling.cpp b/Engine/Source/Runtime/Landscape/Private/LandscapeCulling.cpp
--- a/Engine/Source/Runtime/Landscape/Private/LandscapeCulling.cpp
+++ b/Engine/Source/Runtime/Landscape/Private/LandscapeCulling.cpp
@@ -640,7 +640,7 @@ static void ComputeSectionIntermediateData(FRDGBuilder& GraphBuilder, TArrayView
 					Section.LocalZ = static_cast<float>(SectionLocalBounds.Origin.Z);
 					Section.HalfHeight = static_cast<float>(SectionLocalBounds.BoxExtent.Z);
 					// How many quads to add to each tile extent to compensate for a neighbors LOD
-					Section.NeighborLODExtent = static_cast<float>((1 << NeighborsMaxLOD) - 1); 
+					Section.NeighborLODExtent = FMath::Max(static_cast<float>((1 << NeighborsMaxLOD) - 1), 1.f);
 
 					ViewStateIntermediates.SectionRenderCoords.Add(SectionInfo->RenderCoord);
 				}
