Landscape: blind fix for crash in SaveModifiedLandscapes callback from a mapcheck action which may occur on an invalid landscape subsystem or world.

diff --git a/Engine/Source/Runtime/Landscape/Private/Landscape.cpp b/Engine/Source/Runtime/Landscape/Private/Landscape.cpp
--- a/Engine/Source/Runtime/Landscape/Private/Landscape.cpp
+++ b/Engine/Source/Runtime/Landscape/Private/Landscape.cpp
@@ -4280,6 +4280,36 @@ void ALandscapeProxy::GetSharedProperties(ALandscapeProxy* Landscape)
 
 namespace UE::Landscape::Private
 {
+	static bool HasModifiedLandscapes()
+	{
+		if (GEditor)
+		{
+			if (UWorld* World = GEditor->GetEditorWorldContext().World())
+			{
+				if (ULandscapeSubsystem* LandscapeSubsystem = World->GetSubsystem<ULandscapeSubsystem>())
+				{
+					return  LandscapeSubsystem->HasModifiedLandscapes();
+				}
+			}
+		}
+
+		return false;
+	}
+
+	static void SaveModifiedLandscapes()
+	{
+		if (GEditor)
+		{
+			if (UWorld* World = GEditor->GetEditorWorldContext().World())
+			{
+				if (ULandscapeSubsystem* LandscapeSubsystem = World->GetSubsystem<ULandscapeSubsystem>())
+				{
+					LandscapeSubsystem->SaveModifiedLandscapes();
+				}
+			}
+		}
+	}
+
 	static bool CopyProperty(FProperty* InProperty, UObject* InSourceObject, UObject* InDestinationObject)
 	{
 		void* SrcValuePtr = InProperty->ContainerPtrToValuePtr<void>(InSourceObject);
@@ -4350,8 +4380,8 @@ namespace UE::Landscape::Private
 			->AddToken(FUObjectToken::Create(&InSynchronizedProxy, FText::FromString(InSynchronizedProxy.GetActorNameOrLabel())))
 			->AddToken(FTextToken::Create(LOCTEXT("MapCheck_Message_LandscapeProxy_FixupSharedData", "had some shared properties not in sync with its parent landscape actor. This has been fixed but the proxy needs to be saved in order to ensure cooking behaves as expected. ")))
 			->AddToken(FActionToken::Create(LOCTEXT("MapCheck_SaveFixedUpData", "Save Modified Landscapes"), LOCTEXT("MapCheck_SaveFixedUpData_Desc", "Saves the modified landscape proxy actors"),
-				FOnActionTokenExecuted::CreateUObject(LandscapeSubsystem, &ULandscapeSubsystem::SaveModifiedLandscapes),
-				FCanExecuteActionToken::CreateUObject(LandscapeSubsystem, &ULandscapeSubsystem::HasModifiedLandscapes),
+				FOnActionTokenExecuted::CreateStatic(&SaveModifiedLandscapes),
+				FCanExecuteActionToken::CreateStatic(&HasModifiedLandscapes),
 				/*bInSingleUse = */false))
 			->AddToken(FTextToken::Create(FText::Format(LOCTEXT("MapCheck_Message_LandscapeProxy_FixupSharedData_SharedProperties", "The following properties were synchronized: {0}."), FText::FromString(SynchronizedPropertiesStringBuilder.ToString()))));
 			
