Adds BP exposed DeleteUnusedLayers method.
Removing layer allocation if no data is found for this layer's weightmap.

diff --git a/Engine/Source/Runtime/Landscape/Classes/Landscape.h b/Engine/Source/Runtime/Landscape/Classes/Landscape.h
--- a/Engine/Source/Runtime/Landscape/Classes/Landscape.h
+++ b/Engine/Source/Runtime/Landscape/Classes/Landscape.h
@@ -403,6 +403,8 @@ public:
 	LANDSCAPE_API bool PrepareTextureResources(bool bInWaitForStreaming);
 
 	bool GetVisibilityLayerAllocationIndex() const { return 0; }
+	
+	LANDSCAPE_API virtual void DeleteUnusedLayers() override;
 
 protected:
 	FName GenerateUniqueLayerName(FName InName = NAME_None) const;
diff --git a/Engine/Source/Runtime/Landscape/Classes/LandscapeComponent.h b/Engine/Source/Runtime/Landscape/Classes/LandscapeComponent.h
--- a/Engine/Source/Runtime/Landscape/Classes/LandscapeComponent.h
+++ b/Engine/Source/Runtime/Landscape/Classes/LandscapeComponent.h
@@ -830,6 +830,9 @@ public:
 	/** Recomputes the maximum discard distance across all grass types in the GrassTypes array. */
 	void UpdateGrassTypesMaxDiscardDistance();
 
+	/** Deletes a layer from this component if it does not contain data, calling DeleteLayerAllocation. */
+	bool DeleteLayerIfAllZero(const FGuid& InEditLayerGuid, const uint8* const TexDataPtr, int32 TexSize, int32 LayerIdx, bool bShouldDirtyPackage);
+
 	/** Deletes a material layer from the current edit layer on this component, removing all its data, adjusting other layer's weightmaps if necessary, etc. */
 	LANDSCAPE_API void DeleteLayer(ULandscapeLayerInfoObject* LayerInfo, FLandscapeEditDataInterface& LandscapeEdit);
 	
diff --git a/Engine/Source/Runtime/Landscape/Classes/LandscapeProxy.h b/Engine/Source/Runtime/Landscape/Classes/LandscapeProxy.h
--- a/Engine/Source/Runtime/Landscape/Classes/LandscapeProxy.h
+++ b/Engine/Source/Runtime/Landscape/Classes/LandscapeProxy.h
@@ -1390,6 +1390,13 @@ public:
 	LANDSCAPE_API virtual void UpdateCachedHasLayersContent(bool InCheckComponentDataIntegrity = false);
 
 	FAsyncWorkMonitor& GetAsyncWorkMonitor() { return AsyncWorkMonitor; }
+
+	/**
+	 * Delete all unused layers in components. Warning: any update of the component could re-introduce them.
+	 */
+	UFUNCTION(BlueprintCallable, Category = "Landscape")
+	LANDSCAPE_API virtual void DeleteUnusedLayers();
+
 protected:
 	friend class ALandscape;
 
diff --git a/Engine/Source/Runtime/Landscape/Private/LandscapeEditInterface.cpp b/Engine/Source/Runtime/Landscape/Private/LandscapeEditInterface.cpp
--- a/Engine/Source/Runtime/Landscape/Private/LandscapeEditInterface.cpp
+++ b/Engine/Source/Runtime/Landscape/Private/LandscapeEditInterface.cpp
@@ -1628,6 +1628,35 @@ void FLandscapeEditDataInterface::GetHeightDataFast(const int32 X1, const int32
 	}
 }
 
+bool ULandscapeComponent::DeleteLayerIfAllZero(const FGuid& InEditLayerGuid, const uint8* const TexDataPtr, int32 TexSize, int32 LayerIdx, bool bShouldDirtyPackage)
+{
+	if (TexDataPtr == nullptr)
+	{
+		return false;
+	}
+	
+	// Check the data for the entire component and to see if it's all zero
+	for (int32 TexY = 0; TexY < TexSize; TexY++)
+	{
+		for (int32 TexX = 0; TexX < TexSize; TexX++)
+		{
+			const int32 TexDataIndex = 4 * (TexX + TexY * TexSize);
+
+			// Stop the first time we see any non-zero data
+			uint8 Weight = TexDataPtr[TexDataIndex];
+			if (Weight != 0)
+			{
+				return false;
+			}
+		}
+	}
+
+	DeleteLayerAllocation(InEditLayerGuid, LayerIdx, bShouldDirtyPackage);
+	MarkRenderStateDirty();
+
+	return true;
+}
+
 void ULandscapeComponent::DeleteLayer(ULandscapeLayerInfoObject* LayerInfo, FLandscapeEditDataInterface& LandscapeEdit)
 {
 	FGuid EditLayerGuid = this->GetEditingLayerGUID();
@@ -2484,30 +2513,6 @@ namespace
 	}
 };
 
-bool DeleteLayerIfAllZero(ULandscapeComponent* const Component, const FGuid& InEditLayerGuid, const uint8* const TexDataPtr, int32 TexSize, int32 LayerIdx, bool bShouldDirtyPackage)
-{
-	// Check the data for the entire component and to see if it's all zero
-	for (int32 TexY = 0; TexY < TexSize; TexY++)
-	{
-		for (int32 TexX = 0; TexX < TexSize; TexX++)
-		{
-			const int32 TexDataIndex = 4 * (TexX + TexY * TexSize);
-
-			// Stop the first time we see any non-zero data
-			uint8 Weight = TexDataPtr[TexDataIndex];
-			if (Weight != 0)
-			{
-				return false;
-			}
-		}
-	}
-
-	Component->DeleteLayerAllocation(InEditLayerGuid, LayerIdx, bShouldDirtyPackage);
-
-	return true;
-}
-
-
 inline bool FLandscapeEditDataInterface::IsLayerAllowed(const ULandscapeLayerInfoObject* const LayerInfo, const int32 ComponentIndexX, const int32 SubIndexX, const int32 SubX, const int32 ComponentIndexY, const int32 SubIndexY, const int32 SubY)
 {
 	// left / right
@@ -3193,7 +3198,7 @@ void FLandscapeEditDataInterface::SetAlphaData(ULandscapeLayerInfoObject* const
 			{
 				if (LayerEditDataAllZero[LayerIdx])
 				{
-					bool bLayerDeleted = DeleteLayerIfAllZero(Component, GetEditLayer(), LayerDataPtrs[LayerIdx], TexSize, LayerIdx, GetShouldDirtyPackage());
+					bool bLayerDeleted = Component->DeleteLayerIfAllZero(GetEditLayer(), LayerDataPtrs[LayerIdx], TexSize, LayerIdx, GetShouldDirtyPackage());
 
 					if (bLayerDeleted)
 					{
@@ -3482,7 +3487,7 @@ void FLandscapeEditDataInterface::SetAlphaData(const TSet<ULandscapeLayerInfoObj
 			{
 				if (LayerEditDataAllZero[LayerIdx])
 				{
-					bool bLayerDeleted = DeleteLayerIfAllZero(Component, GetEditLayer(), LayerDataInfos[LayerIdx].TexDataPtr, TexSize, LayerIdx, GetShouldDirtyPackage());
+					bool bLayerDeleted = Component->DeleteLayerIfAllZero(GetEditLayer(), LayerDataInfos[LayerIdx].TexDataPtr, TexSize, LayerIdx, GetShouldDirtyPackage());
 
 					if (bLayerDeleted)
 					{
diff --git a/Engine/Source/Runtime/Landscape/Private/LandscapeEditLayers.cpp b/Engine/Source/Runtime/Landscape/Private/LandscapeEditLayers.cpp
--- a/Engine/Source/Runtime/Landscape/Private/LandscapeEditLayers.cpp
+++ b/Engine/Source/Runtime/Landscape/Private/LandscapeEditLayers.cpp
@@ -3538,6 +3538,28 @@ bool ALandscape::PrepareTextureResources(bool bInWaitForStreaming)
 	return bIsReady;
 }
 
+void ALandscape::DeleteUnusedLayers()
+{
+	ULandscapeInfo* LandscapeInfo = GetLandscapeInfo();
+
+	if (LandscapeInfo == nullptr)
+	{
+		return;
+	}
+
+	for (const TWeakObjectPtr<ALandscapeStreamingProxy>& Proxy : LandscapeInfo->StreamingProxies)
+	{
+		if (!Proxy.IsValid())
+		{
+			continue;
+		}
+
+		Proxy->DeleteUnusedLayers();
+	}
+
+	ALandscapeProxy::DeleteUnusedLayers();
+}
+
 // Note: this approach is generic, because FObjectCacheContextScope is a fast texture->material interface->primitive component lookup. 
 // If FObjectCacheContextScope was available at runtime, it could become an efficient way to automatically invalidate RVT areas corresponding to primitive components that use textures that are being streamed in:
 void ALandscape::InvalidateRVTForTextures(const TSet<UTexture2D*>& InTextures)
@@ -8765,6 +8787,66 @@ void ALandscapeProxy::UpdateCachedHasLayersContent(bool InCheckComponentDataInte
 	}
 }
 
+namespace
+{
+	void DeleteUnusedLayersImpl(ULandscapeComponent* InComponent, const FGuid& InLayerGuid)
+	{
+		TArray<FWeightmapLayerAllocationInfo>& ComponentWeightmapLayerAllocations = InComponent->GetWeightmapLayerAllocations(InLayerGuid);
+
+		for (int32 LayerIdx = 0; LayerIdx < ComponentWeightmapLayerAllocations.Num();)
+		{
+			const FWeightmapLayerAllocationInfo& Allocation = ComponentWeightmapLayerAllocations[LayerIdx];
+			const TArray<TObjectPtr<UTexture2D>>& WeightmapTextures = InComponent->GetWeightmapTextures(InLayerGuid);
+			UTexture2D* Texture = WeightmapTextures[Allocation.WeightmapTextureIndex];
+
+			if (Texture == nullptr)
+			{
+				++LayerIdx;
+				continue;
+			}
+
+			const uint8* MipDataPtr = Texture->Source.LockMipReadOnly(0);
+
+			if (MipDataPtr == nullptr)
+			{
+				++LayerIdx;
+				continue;
+			}
+				
+			const uint8* const TextDataPtr = MipDataPtr + ChannelOffsets[Allocation.WeightmapTextureChannel];
+
+			constexpr bool bShouldDirtyPackage = true;
+
+			// If DeleteLayerIfAllZero returns true, We just removed the current layer allocation, so we need to iterate on the new current index.
+			if (!InComponent->DeleteLayerIfAllZero(InLayerGuid, TextDataPtr, Texture->GetSizeX(), LayerIdx, bShouldDirtyPackage))
+			{
+				++LayerIdx;
+			}
+
+			Texture->Source.UnlockMip(0);
+		}
+	}
+}
+
+void ALandscapeProxy::DeleteUnusedLayers()
+{
+	for (ULandscapeComponent* Component : LandscapeComponents)
+	{
+		if (Component == nullptr)
+		{
+			continue;
+		}
+
+		Component->ForEachLayer([Component](const FGuid& LayerGuid, FLandscapeLayerComponentData& LayerData)
+		{
+			DeleteUnusedLayersImpl(Component, LayerGuid);
+		});
+
+		// Execute ClearUnusedLayersImpl on the final Layer.
+		DeleteUnusedLayersImpl(Component, FGuid());
+	}
+}
+
 bool ALandscapeProxy::RemoveObsoleteLayers(const TSet<FGuid>& InExistingLayers)
 {
 	TSet<TPair<FGuid, FName>> ComponentLayers;
