Add support to enable landscape nanite skirts which can then be called from pure Blueprint functions as necessary.

diff --git a/Engine/Source/Runtime/Landscape/Classes/Landscape.h b/Engine/Source/Runtime/Landscape/Classes/Landscape.h
--- a/Engine/Source/Runtime/Landscape/Classes/Landscape.h
+++ b/Engine/Source/Runtime/Landscape/Classes/Landscape.h
@@ -406,6 +406,8 @@ public:
 	
 	LANDSCAPE_API virtual void DeleteUnusedLayers() override;
 
+	LANDSCAPE_API void EnableNaniteSkirts(bool bInEnable, float InSkirtDepth, bool bInShouldDirtyPackage);
+
 protected:
 	FName GenerateUniqueLayerName(FName InName = NAME_None) const;
 
diff --git a/Engine/Source/Runtime/Landscape/Private/Landscape.cpp b/Engine/Source/Runtime/Landscape/Private/Landscape.cpp
--- a/Engine/Source/Runtime/Landscape/Private/Landscape.cpp
+++ b/Engine/Source/Runtime/Landscape/Private/Landscape.cpp
@@ -3433,6 +3433,33 @@ bool ALandscape::GetUseGeneratedLandscapeSplineMeshesActors() const
 	return bUseGeneratedLandscapeSplineMeshesActors;
 }
 
+void ALandscape::EnableNaniteSkirts(bool bInEnable, float InSkirtDepth, bool bInShouldDirtyPackage)
+{
+	bNaniteSkirtEnabled = bInEnable;
+	NaniteSkirtDepth = InSkirtDepth;
+
+	InvalidateOrUpdateNaniteRepresentation(/*bInCheckContentId*/true, /*InTargetPlatform*/nullptr);
+	UpdateRenderingMethod();
+	MarkComponentsRenderStateDirty();
+	Modify(bInShouldDirtyPackage);
+	
+	if (ULandscapeInfo* LandscapeInfo = GetLandscapeInfo())
+	{
+		LandscapeInfo->ForEachLandscapeProxy([&](ALandscapeProxy* Proxy)
+		{
+			if (Proxy != nullptr)
+			{
+				Proxy->SynchronizeSharedProperties(this);
+				Proxy->InvalidateOrUpdateNaniteRepresentation(/*bInCheckContentId*/true, /*InTargetPlatform*/nullptr);
+				Proxy->UpdateRenderingMethod();
+				Proxy->MarkComponentsRenderStateDirty();
+				Proxy->Modify(bInShouldDirtyPackage);
+			}
+			return true;
+		});
+	}
+}
+
 void ALandscapeProxy::OnFeatureLevelChanged(ERHIFeatureLevel::Type NewFeatureLevel)
 {
 	FlushGrassComponents();
