Force FLandscapeNaniteSceneProxies to be not masked becaues visibility info has already been baked into the geometry

diff --git a/Engine/Source/Runtime/Landscape/Private/LandscapeRender.cpp b/Engine/Source/Runtime/Landscape/Private/LandscapeRender.cpp
--- a/Engine/Source/Runtime/Landscape/Private/LandscapeRender.cpp
+++ b/Engine/Source/Runtime/Landscape/Private/LandscapeRender.cpp
@@ -4523,6 +4523,29 @@ public:
 			bAffectDistanceFieldLighting = false;
 		}
 
+		bool bAnySectionMasked = false;
+		bHasProgrammableRaster = false;
+
+		for (::Nanite::FSceneProxyBase::FMaterialSection& MaterialSection : MaterialSections)
+		{
+			const bool bWasMasked = MaterialSection.MaterialRelevance.bMasked;
+			MaterialSection.MaterialRelevance.bMasked = false;
+
+			if (MaterialSection.IsProgrammableRaster(bEvaluateWorldPositionOffset))
+			{
+				// Don't change bMasked if it is not the sole factor that makes the material section programmable
+				MaterialSection.MaterialRelevance.bMasked = bWasMasked;
+				bAnySectionMasked |= bWasMasked;
+				bHasProgrammableRaster = true;
+			}
+			else
+			{
+				MaterialSection.ResetToDefaultMaterial(false, true);
+			}
+		}
+
+		CombinedMaterialRelevance.bMasked = bAnySectionMasked;
+
 		// Overwrite filter flags to specify landscape instead of static mesh
 		FilterFlags = ::Nanite::EFilterFlags::Landscape;
 	}
