Landscape: Allow for use of texture arrays on mobile platforms

* Allows for material expressions to indicate they can sample from multiple textures
* Material landscape expressions can now sample from texture arrays if enabled (disabled for now)
* Added WeightMapArray place holder texture

#virtualized

diff --git a/Engine/Config/BaseEngine.ini b/Engine/Config/BaseEngine.ini
--- a/Engine/Config/BaseEngine.ini
+++ b/Engine/Config/BaseEngine.ini
@@ -170,6 +170,7 @@ SimpleVolumeTextureName=/Engine/EngineMaterials/Substrate/Volume/SubstrateSimple
 SimpleVolumeEnvTextureName=/Engine/EngineMaterials/Substrate/Volume/SubstrateSimpleVolumeEnvLUT3D.SubstrateSimpleVolumeEnvLUT3D
 MiniFontTextureName=/Engine/EngineMaterials/MiniFont.MiniFont
 WeightMapPlaceholderTextureName=/Engine/EngineMaterials/WeightMapPlaceholderTexture.WeightMapPlaceholderTexture
+WeightMapArrayPlaceholderTextureName=/Engine/EngineMaterials/WeightMapArrayPlaceholderTexture.WeightMapArrayPlaceholderTexture
 LightMapDensityTextureName=/Engine/EngineMaterials/DefaultWhiteGrid.DefaultWhiteGrid
 NaniteHiddenSectionMaterialName=/Engine/EngineMaterials/NaniteHiddenSectionMaterial.NaniteHiddenSectionMaterial
 LevelColorationLitMaterialName=/Engine/EngineDebugMaterials/LevelColorationLitMaterial.LevelColorationLitMaterial
diff --git a/Engine/Source/Runtime/Engine/Classes/Engine/Engine.h b/Engine/Source/Runtime/Engine/Classes/Engine/Engine.h
--- a/Engine/Source/Runtime/Engine/Classes/Engine/Engine.h
+++ b/Engine/Source/Runtime/Engine/Classes/Engine/Engine.h
@@ -1400,10 +1400,16 @@ public:
 	UPROPERTY()
 	TObjectPtr<class UTexture> WeightMapPlaceholderTexture;
 
+	UPROPERTY()
+	TObjectPtr<class UTexture> WeightMapArrayPlaceholderTexture;
+	
 	/** Path of the texture used as a placeholder for terrain weight-maps to give the material the correct texture format. */
 	UPROPERTY(globalconfig)
 	FSoftObjectPath WeightMapPlaceholderTextureName;
 
+	UPROPERTY(globalconfig)
+	FSoftObjectPath WeightMapArrayPlaceholderTextureName;
+
 	/** Texture used to display LightMapDensity */
 	UPROPERTY()
 	TObjectPtr<class UTexture2D> LightMapDensityTexture;
diff --git a/Engine/Source/Runtime/Engine/Classes/Materials/MaterialExpression.h b/Engine/Source/Runtime/Engine/Classes/Materials/MaterialExpression.h
--- a/Engine/Source/Runtime/Engine/Classes/Materials/MaterialExpression.h
+++ b/Engine/Source/Runtime/Engine/Classes/Materials/MaterialExpression.h
@@ -330,13 +330,18 @@ class UMaterialExpression : public UObject
 	*/
 	virtual void GetTexturesForceMaterialRecompile(TArray<UTexture *> &Textures) const { }
 
+
 	/** 
-	 * Callback to get any texture reference this expression emits.
+	 * To get any texture references this expression emits.
 	 * This is used to link the compiled uniform expressions with their default texture values. 
 	 * Any UMaterialExpression whose compilation creates a texture uniform expression (eg Compiler->Texture, Compiler->TextureParameter) must implement this.
 	 */
 	virtual UObject* GetReferencedTexture() const { return nullptr; }
-	/** Returns true if GetReferencedTexture() can ever return a valid pointer. */
+
+	using ReferencedTextureArray = TArray<UObject*, TInlineAllocator<4>>; 
+	virtual ReferencedTextureArray GetReferencedTextures() const { return { GetReferencedTexture() }; }
+	
+	/** Returns true if GetReferencedTexture() / GetReferencedTextures() can ever return a valid pointer(s). */
 	virtual bool CanReferenceTexture() const { return false; }
 
 #if WITH_EDITOR
diff --git a/Engine/Source/Runtime/Engine/Private/Materials/HLSLMaterialTranslator.cpp b/Engine/Source/Runtime/Engine/Private/Materials/HLSLMaterialTranslator.cpp
--- a/Engine/Source/Runtime/Engine/Private/Materials/HLSLMaterialTranslator.cpp
+++ b/Engine/Source/Runtime/Engine/Private/Materials/HLSLMaterialTranslator.cpp
@@ -8047,7 +8047,7 @@ bool FHLSLMaterialTranslator::GetStaticBoolValue(int32 BoolIndex, bool& bSucceed
 	return false;
 }
 
-int32 FHLSLMaterialTranslator::StaticTerrainLayerWeight(FName LayerName,int32 Default)
+int32 FHLSLMaterialTranslator::StaticTerrainLayerWeight(FName LayerName,int32 Default, bool bTextureArray)
 {
 	// Look up the weight-map index for this static parameter.
 	int32 WeightmapCode = INDEX_NONE;
@@ -8075,19 +8075,31 @@ int32 FHLSLMaterialTranslator::StaticTerrainLayerWeight(FName LayerName,int32 De
 			continue;
 		}
 
-		constexpr EMaterialSamplerType SamplerType = SAMPLERTYPE_Masks;
-		FString WeightmapName = FString::Printf(TEXT("Weightmap%d"), WeightmapIndex);
-		int32 TextureReferenceIndex = INDEX_NONE;
-
 		FMaterialParameterInfo GlobalParameterInfo;
 		PushParameterOwner(GlobalParameterInfo);
 
-		int32 TextureCodeIndex = TextureParameter(FName(*WeightmapName), GEngine->WeightMapPlaceholderTexture, TextureReferenceIndex, SamplerType);
-		int32 SampleCodeIndex = TextureSample(TextureCodeIndex, TextureCoordinate(3, false, false), SamplerType, /*MipValue0Index = */INDEX_NONE, /*MipValue1Index = */INDEX_NONE, /*MipValueMode = */TMVM_None, /*SamplerSource = */SSM_TerrainWeightmapGroupSettings);
-
+		int32 UVIndex = TextureCoordinate(3, false, false);
+		constexpr EMaterialSamplerType SamplerType = SAMPLERTYPE_Masks;
+		int32 TextureReferenceIndex = INDEX_NONE;
+		int32 SampleCodeIndex = 0;
+		 
+		if (bTextureArray)
+		{
+			int32 TextureArrayCodeIndex = TextureParameter(TEXT("WeightmapArray"), GEngine->WeightMapArrayPlaceholderTexture, TextureReferenceIndex, SamplerType);
+			int32 ConstantSliceIndex = Constant(WeightmapIndex);
+			int32 UVWIndex = AppendVector(UVIndex, ConstantSliceIndex);
+			SampleCodeIndex = TextureSample(TextureArrayCodeIndex, UVWIndex, SamplerType, /*MipValue0Index = */INDEX_NONE, /*MipValue1Index = */INDEX_NONE, /*MipValueMode = */TMVM_None, /*SamplerSource = */SSM_TerrainWeightmapGroupSettings);
+		}
+		else
+		{
+			FString WeightmapName = FString::Printf(TEXT("Weightmap%d"), WeightmapIndex);
+			int32 TextureCodeIndex = TextureParameter(FName(*WeightmapName), GEngine->WeightMapPlaceholderTexture, TextureReferenceIndex, SamplerType);
+			SampleCodeIndex = TextureSample(TextureCodeIndex, UVIndex, SamplerType, /*MipValue0Index = */INDEX_NONE, /*MipValue1Index = */INDEX_NONE, /*MipValueMode = */TMVM_None, /*SamplerSource = */SSM_TerrainWeightmapGroupSettings);
+		}
+	
 		FString LayerMaskName = FString::Printf(TEXT("LayerMask_%s"), *Parameter.LayerName.ToString());
-		int32 CurrentWeightmapCode = Dot(SampleCodeIndex, VectorParameter(FName(*LayerMaskName), FLinearColor(1.f, 0.f, 0.f, 0.f)));
-
+		int32 CurrentWeightmapCode = Dot(SampleCodeIndex, VectorParameter(FName(*LayerMaskName), FLinearColor(1.f, 0.f, 0.f, 0.f)));	
+		 
 		if(WeightmapCode == INDEX_NONE)
 		{
 			WeightmapCode = CurrentWeightmapCode;
diff --git a/Engine/Source/Runtime/Engine/Private/Materials/HLSLMaterialTranslator.h b/Engine/Source/Runtime/Engine/Private/Materials/HLSLMaterialTranslator.h
--- a/Engine/Source/Runtime/Engine/Private/Materials/HLSLMaterialTranslator.h
+++ b/Engine/Source/Runtime/Engine/Private/Materials/HLSLMaterialTranslator.h
@@ -1048,7 +1048,7 @@ protected:
 
 	virtual bool GetStaticBoolValue(int32 BoolIndex, bool& bSucceeded) override;
 
-	virtual int32 StaticTerrainLayerWeight(FName ParameterName, int32 Default) override;
+	virtual int32 StaticTerrainLayerWeight(FName ParameterName, int32 Default, bool bTextureArray = false) override;
 
 	virtual int32 VertexColor() override;
 
diff --git a/Engine/Source/Runtime/Engine/Private/Materials/MaterialCachedData.cpp b/Engine/Source/Runtime/Engine/Private/Materials/MaterialCachedData.cpp
--- a/Engine/Source/Runtime/Engine/Private/Materials/MaterialCachedData.cpp
+++ b/Engine/Source/Runtime/Engine/Private/Materials/MaterialCachedData.cpp
@@ -360,12 +360,14 @@ void FMaterialCachedExpressionData::UpdateForExpressions(const FMaterialCachedEx
 		const bool bCanReferenceTexture = Expression->CanReferenceTexture();
 		if (!ReferencedTexture && bCanReferenceTexture)
 		{
-			ReferencedTexture = Expression->GetReferencedTexture();
+			const UMaterialExpression::ReferencedTextureArray ExpressionReferencedTextures = Expression->GetReferencedTextures();
+			for (UObject* ExpressionReferencedTexture : ExpressionReferencedTextures)
+			{
+				ReferencedTextures.AddUnique(ExpressionReferencedTexture);
+			}
 		}
-
-		if (ReferencedTexture)
+		else if (ReferencedTexture)
 		{
-			checkf(bCanReferenceTexture, TEXT("CanReferenceTexture() returned false, but found a referenced texture"));
 			ReferencedTextures.AddUnique(ReferencedTexture);
 		}
 
diff --git a/Engine/Source/Runtime/Engine/Private/Materials/MaterialExpressionHLSL.cpp b/Engine/Source/Runtime/Engine/Private/Materials/MaterialExpressionHLSL.cpp
--- a/Engine/Source/Runtime/Engine/Private/Materials/MaterialExpressionHLSL.cpp
+++ b/Engine/Source/Runtime/Engine/Private/Materials/MaterialExpressionHLSL.cpp
@@ -4826,14 +4826,20 @@ bool UMaterialExpressionDistanceCullFade::GenerateHLSLExpression(FMaterialHLSLGe
 	return true;
 }
 
-bool GenerateStaticTerrainLayerWeightExpression(FName LayerName, float PreviewWeight, FMaterialHLSLGenerator& Generator, const UE::HLSLTree::FExpression*& OutExpression)
+bool GenerateStaticTerrainLayerWeightExpression(FName LayerName, float PreviewWeight, bool bUseTextureArray, FMaterialHLSLGenerator& Generator, const UE::HLSLTree::FExpression*& OutExpression)
 {
 	using namespace UE::HLSLTree;
 	const FExpression* TexCoordExpression = Generator.NewExternalInput(Material::EExternalInput::TexCoord3);
-	OutExpression = Generator.GetTree().NewExpression<Material::FExpressionStaticTerrainLayerWeight>(Generator.GetParameterInfo(LayerName), TexCoordExpression, PreviewWeight);
+	OutExpression = Generator.GetTree().NewExpression<Material::FExpressionStaticTerrainLayerWeight>(Generator.GetParameterInfo(LayerName), TexCoordExpression, PreviewWeight, bUseTextureArray);
 	return true;
 }
 
+bool GenerateStaticTerrainLayerWeightExpression(FName LayerName, float PreviewWeight, FMaterialHLSLGenerator& Generator, const UE::HLSLTree::FExpression*& OutExpression)
+{
+	const bool bUseTextureArray = false;
+	return GenerateStaticTerrainLayerWeightExpression(LayerName, PreviewWeight, bUseTextureArray, Generator, OutExpression);
+}
+
 UE::Shader::EValueType UMaterialExpressionNeuralNetworkInput::GetCustomOutputType(int32 OutputIndex) const
 {
 	using namespace UE::Shader;
diff --git a/Engine/Source/Runtime/Engine/Private/Materials/MaterialHLSLTree.cpp b/Engine/Source/Runtime/Engine/Private/Materials/MaterialHLSLTree.cpp
--- a/Engine/Source/Runtime/Engine/Private/Materials/MaterialHLSLTree.cpp
+++ b/Engine/Source/Runtime/Engine/Private/Materials/MaterialHLSLTree.cpp
@@ -1675,7 +1675,8 @@ bool FExpressionStaticTerrainLayerWeight::PrepareValue(FEmitContext& Context, FE
 {
 	bool bFoundMatchingParameter = false;
 	FEmitData& EmitData = Context.FindData<FEmitData>();
-
+	// TODO: revist whether we need to add the parameters to CachedExpressionData
+	const bool bTextureArrayEnabled = UseTextureArraySample(Context);
 	if (EmitData.StaticParameters)
 	{
 		for (int32 ParameterIndex = 0; ParameterIndex < EmitData.StaticParameters->EditorOnly.TerrainLayerWeightParameters.Num(); ++ParameterIndex)
@@ -1706,10 +1707,17 @@ bool FExpressionStaticTerrainLayerWeight::PrepareValue(FEmitContext& Context, FE
 			}
 
 			FMaterialParameterInfo WeightmapParameterInfo = BaseParameterInfo;
-			WeightmapParameterInfo.Name = *FString::Printf(TEXT("Weightmap%d"), WeightmapIndex);
-
 			FMaterialParameterMetadata WeightmapParameterMeta;
-			WeightmapParameterMeta.Value = GEngine->WeightMapPlaceholderTexture;
+			if (bTextureArrayEnabled)
+			{
+				WeightmapParameterInfo.Name = FName(TEXT("WeightmapArray"));
+				WeightmapParameterMeta.Value = GEngine->WeightMapPlaceholderTexture;
+			}
+			else
+			{
+				WeightmapParameterInfo.Name = *FString::Printf(TEXT("Weightmap%d"), WeightmapIndex);
+				WeightmapParameterMeta.Value = GEngine->WeightMapArrayPlaceholderTexture;
+			}
 
 			UObject* UnusedReferencedTexture;
 			EmitData.CachedExpressionData->AddParameter(WeightmapParameterInfo, WeightmapParameterMeta, UnusedReferencedTexture);
@@ -1725,6 +1733,7 @@ bool FExpressionStaticTerrainLayerWeight::PrepareValue(FEmitContext& Context, FE
 
 		if (bFoundMatchingParameter && EmitData.CachedExpressionData)
 		{
+			EmitData.CachedExpressionData->ReferencedTextures.AddUnique(GEngine->WeightMapArrayPlaceholderTexture);
 			EmitData.CachedExpressionData->ReferencedTextures.AddUnique(GEngine->WeightMapPlaceholderTexture);
 		}
 	}
@@ -1771,8 +1780,10 @@ void FExpressionStaticTerrainLayerWeight::EmitValueShader(FEmitContext& Context,
 		OutResult.Code = Context.EmitConstantZero(Scope, Shader::EValueType::Float4);
 		return;
 	}
+	
+	const bool bTextureArrayEnabled = UseTextureArraySample(Context);
 
-	const EMaterialValueType TextureMaterialType = GEngine->WeightMapPlaceholderTexture->GetMaterialType();
+	const EMaterialValueType TextureMaterialType = bTextureArrayEnabled ? GEngine->WeightMapArrayPlaceholderTexture->GetMaterialType() : GEngine->WeightMapPlaceholderTexture->GetMaterialType();
 	const Shader::EValueType TexCoordType = Private::GetTexCoordType(TextureMaterialType);
 	FEmitShaderExpression* EmitTexCoordValue = nullptr;
 	FEmitShaderExpression* EmitResult = nullptr;
@@ -1799,16 +1810,31 @@ void FExpressionStaticTerrainLayerWeight::EmitValueShader(FEmitContext& Context,
 		}
 
 		FMaterialTextureValue TextureValue;
-		TextureValue.Texture = GEngine->WeightMapPlaceholderTexture;
-		TextureValue.SamplerType = SAMPLERTYPE_Masks;
-		TextureValue.ParameterInfo = BaseParameterInfo;
-		TextureValue.ParameterInfo.Name = *FString::Printf(TEXT("Weightmap%d"), WeightmapIndex);
-
+		if (bTextureArrayEnabled)
+		{
+			TextureValue.Texture = GEngine->WeightMapArrayPlaceholderTexture;
+			TextureValue.SamplerType = SAMPLERTYPE_Masks;
+			TextureValue.ParameterInfo = BaseParameterInfo;
+			TextureValue.ParameterInfo.Name = FName(TEXT("WeightmapArray"));
+		}
+		else
+		{
+			TextureValue.Texture = GEngine->WeightMapPlaceholderTexture;
+			TextureValue.SamplerType = SAMPLERTYPE_Masks;
+			TextureValue.ParameterInfo = BaseParameterInfo;
+			TextureValue.ParameterInfo.Name = *FString::Printf(TEXT("Weightmap%d"), WeightmapIndex);	
+		}
+		
 		if (!EmitTexCoordValue)
 		{
 			EmitTexCoordValue = TexCoordExpression->GetValueShader(Context, Scope, TexCoordType);
 		}
 
+		if (bTextureArrayEnabled)
+		{
+			EmitTexCoordValue = Context.EmitExpression(Scope, Shader::EValueType::Float3, TEXT("float3((%).xy,%)"), EmitTexCoordValue, WeightmapIndex);
+		}
+
 		FEmitShaderExpression* EmitWeightmapSampleValue = Private::EmitTextureSampleShader(
 			Context, Scope,
 			TextureValue, EmitTexCoordValue, nullptr, nullptr, nullptr,
@@ -1861,6 +1887,11 @@ void FExpressionStaticTerrainLayerWeight::EmitValueShader(FEmitContext& Context,
 	OutResult.Code = NumWeightmapParameters > 1 ? Context.EmitExpression(Scope, Shader::EValueType::Float1, TEXT("saturate(%)"), EmitResult) : EmitResult;
 }
 
+bool FExpressionStaticTerrainLayerWeight::UseTextureArraySample(const FEmitContext& Context) const
+{
+	return bTextureArray && !Context.TargetParameters.IsGenericTarget() && IsMobilePlatform(Context.TargetParameters.ShaderPlatform);
+}
+
 bool FExpressionTextureProperty::PrepareValue(FEmitContext& Context, FEmitScope& Scope, const FRequestedType& RequestedType, FPrepareValueResult& OutResult) const
 {
 	const FPreparedType& TextureType = Context.PrepareExpression(TextureExpression, Scope, FMaterialTextureValue::GetTypeName());
diff --git a/Engine/Source/Runtime/Engine/Private/Materials/MaterialShared.cpp b/Engine/Source/Runtime/Engine/Private/Materials/MaterialShared.cpp
--- a/Engine/Source/Runtime/Engine/Private/Materials/MaterialShared.cpp
+++ b/Engine/Source/Runtime/Engine/Private/Materials/MaterialShared.cpp
@@ -4783,7 +4783,8 @@ bool UMaterialInterface::IsTextureReferencedByProperty(EMaterialProperty InPrope
 
 		virtual EMaterialExpressionVisitResult Visit(UMaterialExpression* InExpression) override
 		{
-			if (InExpression->GetReferencedTexture() == Texture)
+			const UMaterialExpression::ReferencedTextureArray ReferencedTextures = InExpression->GetReferencedTextures();
+			if (ReferencedTextures.Contains( Texture ))
 			{
 				FoundTexture = true;
 				return MVR_STOP;
diff --git a/Engine/Source/Runtime/Engine/Private/UnrealEngine.cpp b/Engine/Source/Runtime/Engine/Private/UnrealEngine.cpp
--- a/Engine/Source/Runtime/Engine/Private/UnrealEngine.cpp
+++ b/Engine/Source/Runtime/Engine/Private/UnrealEngine.cpp
@@ -3142,6 +3142,7 @@ void UEngine::InitializeObjectReferences()
 	LoadEngineTexture(DefaultBokehTexture, *DefaultBokehTextureName.ToString());
 	LoadEngineTexture(MiniFontTexture, *MiniFontTextureName.ToString());
 	LoadEngineTexture(WeightMapPlaceholderTexture, *WeightMapPlaceholderTextureName.ToString());
+	LoadEngineTexture(WeightMapArrayPlaceholderTexture,  *WeightMapArrayPlaceholderTextureName.ToString());
 	LoadEngineTexture(LightMapDensityTexture, *LightMapDensityTextureName.ToString());
 	ConditionallyLoadPreIntegratedSkinBRDFTexture();
 
diff --git a/Engine/Source/Runtime/Engine/Public/MaterialCompiler.h b/Engine/Source/Runtime/Engine/Public/MaterialCompiler.h
--- a/Engine/Source/Runtime/Engine/Public/MaterialCompiler.h
+++ b/Engine/Source/Runtime/Engine/Public/MaterialCompiler.h
@@ -403,7 +403,7 @@ public:
 	virtual int32 StaticComponentMask(int32 Vector,FName ParameterName,bool bDefaultR,bool bDefaultG,bool bDefaultB,bool bDefaultA) = 0;
 	virtual const FMaterialLayersFunctions* GetMaterialLayers() = 0;
 	virtual bool GetStaticBoolValue(int32 BoolIndex, bool& bSucceeded) = 0;
-	virtual int32 StaticTerrainLayerWeight(FName ParameterName,int32 Default) = 0;
+	virtual int32 StaticTerrainLayerWeight(FName ParameterName,int32 Default, bool bTextureArray = false) = 0;
 
 	virtual int32 VertexColor() = 0;
 
@@ -873,7 +873,7 @@ public:
 	virtual int32 StaticComponentMask(int32 Vector, FName ParameterName, bool bDefaultR, bool bDefaultG, bool bDefaultB, bool bDefaultA) override { return Compiler->StaticComponentMask(Vector, ParameterName, bDefaultR, bDefaultG, bDefaultB, bDefaultA); }
 	virtual const FMaterialLayersFunctions* GetMaterialLayers() override { return Compiler->GetMaterialLayers(); }
 	virtual bool GetStaticBoolValue(int32 BoolIndex, bool& bSucceeded) override { return Compiler->GetStaticBoolValue(BoolIndex, bSucceeded); }
-	virtual int32 StaticTerrainLayerWeight(FName ParameterName, int32 Default) override { return Compiler->StaticTerrainLayerWeight(ParameterName, Default); }
+	virtual int32 StaticTerrainLayerWeight(FName ParameterName, int32 Default, bool bTextureArray = false) override { return Compiler->StaticTerrainLayerWeight(ParameterName, Default, bTextureArray); }
 
 	virtual int32 VertexColor() override { return Compiler->VertexColor(); }
 
diff --git a/Engine/Source/Runtime/Engine/Public/MaterialHLSLGenerator.h b/Engine/Source/Runtime/Engine/Public/MaterialHLSLGenerator.h
--- a/Engine/Source/Runtime/Engine/Public/MaterialHLSLGenerator.h
+++ b/Engine/Source/Runtime/Engine/Public/MaterialHLSLGenerator.h
@@ -333,6 +333,14 @@ private:
 	TMap<FXxHash64, int32> CustomExpressionOutputStructIdMap;
 };
 
+ENGINE_API bool GenerateStaticTerrainLayerWeightExpression(
+	FName LayerName,
+	float PreviewWeight,
+	bool bUseTextureArray, 
+	FMaterialHLSLGenerator& Generator,
+	const UE::HLSLTree::FExpression*& OutExpression);
+
+UE_DEPRECATED(5.4, "GenerateStaticTerrainLayerWeightExpression(FName, float, FMaterialHLSLGenerator&, const UE::HLSLTree::FExpression*&) use version above")
 ENGINE_API bool GenerateStaticTerrainLayerWeightExpression(
 	FName LayerName,
 	float PreviewWeight,
diff --git a/Engine/Source/Runtime/Engine/Public/MaterialHLSLTree.h b/Engine/Source/Runtime/Engine/Public/MaterialHLSLTree.h
--- a/Engine/Source/Runtime/Engine/Public/MaterialHLSLTree.h
+++ b/Engine/Source/Runtime/Engine/Public/MaterialHLSLTree.h
@@ -626,21 +626,35 @@ public:
 class FExpressionStaticTerrainLayerWeight : public FExpression
 {
 public:
+	FExpressionStaticTerrainLayerWeight(const FMaterialParameterInfo& InBaseParameterInfo, const FExpression* InTexCoordExpression, float InDefaultWeight, bool bInTextureArray)
+			: BaseParameterInfo(InBaseParameterInfo)
+			, TexCoordExpression(InTexCoordExpression)
+			, DefaultWeight(InDefaultWeight)
+			, bTextureArray(bInTextureArray)
+	{
+		check(!BaseParameterInfo.Name.IsNone());
+	}
+	
+	UE_DEPRECATED(5.4, "FExpressionStaticTerrainLayerWeight::FExpressionStaticTerrainLayerWeight(const FMaterialParameterInfo& , const FExpression* , float) has been deprecate. Use version above")
 	FExpressionStaticTerrainLayerWeight(const FMaterialParameterInfo& InBaseParameterInfo, const FExpression* InTexCoordExpression, float InDefaultWeight)
-		: BaseParameterInfo(InBaseParameterInfo)
-		, TexCoordExpression(InTexCoordExpression)
-		, DefaultWeight(InDefaultWeight)
+			: BaseParameterInfo(InBaseParameterInfo)
+			, TexCoordExpression(InTexCoordExpression)
+			, DefaultWeight(InDefaultWeight)
+			, bTextureArray(false)
 	{
 		check(!BaseParameterInfo.Name.IsNone());
 	}
-
+	
 	FMaterialParameterInfo BaseParameterInfo;
 	const FExpression* TexCoordExpression;
 	float DefaultWeight;
+	bool bTextureArray;
 
 	virtual bool PrepareValue(FEmitContext& Context, FEmitScope& Scope, const FRequestedType& RequestedType, FPrepareValueResult& OutResult) const override;
 	virtual void EmitValuePreshader(FEmitContext& Context, FEmitScope& Scope, const FRequestedType& RequestedType, FEmitValuePreshaderResult& OutResult) const override;
 	virtual void EmitValueShader(FEmitContext& Context, FEmitScope& Scope, const FRequestedType& RequestedType, FEmitValueShaderResult& OutResult) const override;
+private:
+	bool UseTextureArraySample(const FEmitContext& Context) const;
 };
 
 class FExpressionTextureProperty : public FExpression
diff --git a/Engine/Source/Runtime/Landscape/Classes/LandscapeComponent.h b/Engine/Source/Runtime/Landscape/Classes/LandscapeComponent.h
--- a/Engine/Source/Runtime/Landscape/Classes/LandscapeComponent.h
+++ b/Engine/Source/Runtime/Landscape/Classes/LandscapeComponent.h
@@ -15,7 +15,7 @@
 #include "LandscapeWeightmapUsage.h"
 #include "Containers/ArrayView.h"
 #include "Engine/StreamableRenderAsset.h"
-
+#include "Engine/Texture2DArray.h"
 #include "LandscapeComponent.generated.h"
 
 class ALandscape;
@@ -626,6 +626,9 @@ public:
 	UPROPERTY(NonPIEDuplicateTransient)
 	TArray<TObjectPtr<UTexture2D>> MobileWeightmapTextures;
 
+	UPROPERTY(NonPIEDuplicateTransient)
+	TObjectPtr<UTexture2DArray> MobileWeightmapTextureArray;
+	
 	/** Layer allocations used by mobile.*/
 	UPROPERTY()
 	TArray<FWeightmapLayerAllocationInfo> MobileWeightmapLayerAllocations;
diff --git a/Engine/Source/Runtime/Landscape/Classes/LandscapeProxy.h b/Engine/Source/Runtime/Landscape/Classes/LandscapeProxy.h
--- a/Engine/Source/Runtime/Landscape/Classes/LandscapeProxy.h
+++ b/Engine/Source/Runtime/Landscape/Classes/LandscapeProxy.h
@@ -17,6 +17,7 @@
 #include "VT/RuntimeVirtualTextureEnum.h"
 #include "ActorPartition/PartitionActor.h"
 #include "ILandscapeSplineInterface.h"
+#include "Engine/Texture2DArray.h"
 
 #if UE_ENABLE_INCLUDE_ORDER_DEPRECATED_IN_5_2
 #include "LandscapeComponent.h"
@@ -1328,6 +1329,9 @@ public:
 	/** Creates a Texture2D for use by this landscape proxy or one of it's components. If OptionalOverrideOuter is not specified, the proxy is used. */
 	LANDSCAPE_API UTexture2D* CreateLandscapeTexture(int32 InSizeX, int32 InSizeY, TextureGroup InLODGroup, ETextureSourceFormat InFormat, UObject* OptionalOverrideOuter = nullptr, bool bCompress = false, bool bMipChain = true) const;
 
+	/** Creates a Texture2DArray for use by this landscape proxy or one of it's components. If OptionalOverrideOuter is not specified, the proxy is used. */
+	LANDSCAPE_API UTexture2DArray* CreateLandscapeTextureArray(int32 InSizeX, int32 InSizeY, int32 Slices, TextureGroup InLODGroup, ETextureSourceFormat InFormat, UObject* OptionalOverrideOuter = nullptr);
+	
 	/** Creates a Texture2D for use by this landscape proxy or one of it's components for tools .*/ 
 	LANDSCAPE_API UTexture2D* CreateLandscapeToolTexture(int32 InSizeX, int32 InSizeY, TextureGroup InLODGroup, ETextureSourceFormat InFormat) const;
 
diff --git a/Engine/Source/Runtime/Landscape/Classes/Materials/MaterialExpressionLandscapeLayerBlend.h b/Engine/Source/Runtime/Landscape/Classes/Materials/MaterialExpressionLandscapeLayerBlend.h
--- a/Engine/Source/Runtime/Landscape/Classes/Materials/MaterialExpressionLandscapeLayerBlend.h
+++ b/Engine/Source/Runtime/Landscape/Classes/Materials/MaterialExpressionLandscapeLayerBlend.h
@@ -92,6 +92,7 @@ class UMaterialExpressionLandscapeLayerBlend : public UMaterialExpression
 	LANDSCAPE_API virtual void GetLandscapeLayerNames(TArray<FName>& OutLayers) const override;
 #endif
 	LANDSCAPE_API virtual UObject* GetReferencedTexture() const override;
+	LANDSCAPE_API virtual ReferencedTextureArray GetReferencedTextures() const override;
 	virtual bool CanReferenceTexture() const override { return true; }
 	//~ End UMaterialExpression Interface
 
diff --git a/Engine/Source/Runtime/Landscape/Classes/Materials/MaterialExpressionLandscapeLayerSample.h b/Engine/Source/Runtime/Landscape/Classes/Materials/MaterialExpressionLandscapeLayerSample.h
--- a/Engine/Source/Runtime/Landscape/Classes/Materials/MaterialExpressionLandscapeLayerSample.h
+++ b/Engine/Source/Runtime/Landscape/Classes/Materials/MaterialExpressionLandscapeLayerSample.h
@@ -42,6 +42,7 @@ public:
 	LANDSCAPE_API virtual void GetLandscapeLayerNames(TArray<FName>& OutLayers) const override;
 #endif
 	LANDSCAPE_API virtual UObject* GetReferencedTexture() const override;
+	LANDSCAPE_API virtual ReferencedTextureArray GetReferencedTextures() const override;
 	virtual bool CanReferenceTexture() const override { return true; }
 	//~ End UMaterialExpression Interface
 };
diff --git a/Engine/Source/Runtime/Landscape/Classes/Materials/MaterialExpressionLandscapeLayerSwitch.h b/Engine/Source/Runtime/Landscape/Classes/Materials/MaterialExpressionLandscapeLayerSwitch.h
--- a/Engine/Source/Runtime/Landscape/Classes/Materials/MaterialExpressionLandscapeLayerSwitch.h
+++ b/Engine/Source/Runtime/Landscape/Classes/Materials/MaterialExpressionLandscapeLayerSwitch.h
@@ -48,6 +48,7 @@ public:
 	LANDSCAPE_API virtual bool MatchesSearchQuery(const TCHAR* SearchQuery) override;
 #endif
 	LANDSCAPE_API virtual UObject* GetReferencedTexture() const override;
+	LANDSCAPE_API virtual ReferencedTextureArray GetReferencedTextures() const override;
 	virtual bool CanReferenceTexture() const override { return true; }
 #if WITH_EDITOR
 	virtual uint32 GetInputType(int32 InputIndex) override {return MCT_Unknown;}
diff --git a/Engine/Source/Runtime/Landscape/Classes/Materials/MaterialExpressionLandscapeLayerWeight.h b/Engine/Source/Runtime/Landscape/Classes/Materials/MaterialExpressionLandscapeLayerWeight.h
--- a/Engine/Source/Runtime/Landscape/Classes/Materials/MaterialExpressionLandscapeLayerWeight.h
+++ b/Engine/Source/Runtime/Landscape/Classes/Materials/MaterialExpressionLandscapeLayerWeight.h
@@ -48,6 +48,7 @@ public:
 	LANDSCAPE_API virtual bool MatchesSearchQuery(const TCHAR* SearchQuery) override;
 #endif
 	LANDSCAPE_API virtual UObject* GetReferencedTexture() const override;
+	LANDSCAPE_API virtual ReferencedTextureArray GetReferencedTextures() const override;
 	virtual bool CanReferenceTexture() const override { return true; }
 #if WITH_EDITOR
 	virtual uint32 GetInputType(int32 InputIndex) override {return MCT_Float | MCT_MaterialAttributes;}
diff --git a/Engine/Source/Runtime/Landscape/Classes/Materials/MaterialExpressionLandscapeVisibilityMask.h b/Engine/Source/Runtime/Landscape/Classes/Materials/MaterialExpressionLandscapeVisibilityMask.h
--- a/Engine/Source/Runtime/Landscape/Classes/Materials/MaterialExpressionLandscapeVisibilityMask.h
+++ b/Engine/Source/Runtime/Landscape/Classes/Materials/MaterialExpressionLandscapeVisibilityMask.h
@@ -34,6 +34,7 @@ public:
 	LANDSCAPE_API virtual void GetLandscapeLayerNames(TArray<FName>& OutLayers) const override;
 #endif
 	LANDSCAPE_API virtual UObject* GetReferencedTexture() const override;
+	LANDSCAPE_API virtual ReferencedTextureArray GetReferencedTextures() const override;
 	virtual bool CanReferenceTexture() const override { return true; }
 	//~ End UMaterialExpression Interface
 };
diff --git a/Engine/Source/Runtime/Landscape/Private/Landscape.cpp b/Engine/Source/Runtime/Landscape/Private/Landscape.cpp
--- a/Engine/Source/Runtime/Landscape/Private/Landscape.cpp
+++ b/Engine/Source/Runtime/Landscape/Private/Landscape.cpp
@@ -93,6 +93,7 @@ Landscape.cpp: Terrain rendering
 #include "WorldPartition/WorldPartitionHelpers.h"
 #include "WorldPartition/WorldPartitionHandle.h"
 #include "WorldPartition/Landscape/LandscapeActorDesc.h"
+#include "Engine\Texture2DArray.h"
 
 #if WITH_EDITOR
 #include "Rendering/StaticLightingSystemInterface.h"
@@ -705,6 +706,12 @@ void ULandscapeComponent::Serialize(FArchive& Ar)
 		{
 			TexturesAndMaterials.Add((UObject**)&static_cast<UTexture2D*&>(MobileWeightmapTexture));
 		}
+		
+		if (MobileWeightmapTextureArray)
+		{
+			TexturesAndMaterials.Add((UObject**)&static_cast<UTexture2DArray*&>(MobileWeightmapTextureArray));
+		}
+		
 		for (auto& ItPair : LayersData)
 		{
 			FLandscapeLayerComponentData& LayerComponentData = ItPair.Value;
@@ -776,7 +783,7 @@ void ULandscapeComponent::Serialize(FArchive& Ar)
 
 			Exchange(MobileMaterialInterfaces, BackupMobileMaterialInterfaces);
 			Exchange(MobileWeightmapTextures, BackupMobileWeightmapTextures);
-
+			MobileWeightmapTextureArray = TObjectPtr<UTexture2DArray>(nullptr);
 			Super::Serialize(Ar);
 
 			Exchange(MobileMaterialInterfaces, BackupMobileMaterialInterfaces);
@@ -1245,6 +1252,11 @@ void ULandscapeComponent::PostLoad()
 		ReparentObject(MobileWeightmapTexture);
 	}
 
+	if (MobileWeightmapTextureArray)
+	{
+		ReparentObject(MobileWeightmapTextureArray.Get());
+	}
+
 	for (auto& ItPair : LayersData)
 	{
 		FLandscapeLayerComponentData& LayerComponentData = ItPair.Value;
diff --git a/Engine/Source/Runtime/Landscape/Private/LandscapeEdit.cpp b/Engine/Source/Runtime/Landscape/Private/LandscapeEdit.cpp
--- a/Engine/Source/Runtime/Landscape/Private/LandscapeEdit.cpp
+++ b/Engine/Source/Runtime/Landscape/Private/LandscapeEdit.cpp
@@ -45,6 +45,7 @@ LandscapeEdit.cpp: Landscape editing
 #include "LandscapeSplinesComponent.h"
 #include "Serialization/MemoryWriter.h"
 #include "MaterialCachedData.h"
+#include "Math/UnrealMathUtility.h"
 
 #if WITH_EDITOR
 #include "Engine/World.h"
@@ -89,6 +90,7 @@ LandscapeEdit.cpp: Landscape editing
 #include "Serialization/MemoryWriter.h"
 #include "Engine/Canvas.h"
 #include "Spatial/PointHashGrid2.h"
+#include "Engine/Texture2DArray.h"
 
 DEFINE_LOG_CATEGORY(LogLandscape);
 DEFINE_LOG_CATEGORY(LogLandscapeBP);
@@ -5388,6 +5390,7 @@ ALandscapeProxy* ULandscapeInfo::MoveComponentsToProxy(const TArray<ULandscapeCo
 		Component->MobileDataSourceHash.Invalidate();
 		Component->MobileMaterialInterfaces.Reset();
 		Component->MobileWeightmapTextures.Reset();
+		Component->MobileWeightmapTextureArray = nullptr; 
 
 		Component->UpdateMaterialInstances();
 	}
@@ -7354,49 +7357,91 @@ void ULandscapeComponent::GenerateMobilePlatformPixelData(bool bIsCooking, const
 	UE::Landscape::FBatchTextureCopy CopyRequests;
 
 	MobileWeightmapTextures.Empty();
+	MobileWeightmapTextureArray = nullptr;
+
+	const int32 NumWeightTextures = FMath::DivideAndRoundUp(static_cast<int32>(Algo::CountIf(MobileWeightmapLayerAllocations, [](const FWeightmapLayerAllocationInfo& AllocationInfo) { return AllocationInfo.LayerInfo; })), 4);
+
+	const bool MobileWeightmapTextureArrayEnabled = UE::Landscape::IsMobileWeightmapTextureArrayEnabled();
+	
+	if (MobileWeightmapTextureArrayEnabled && NumWeightTextures > 0)
+	{
+		MobileWeightmapTextureArray = GetLandscapeProxy()->CreateLandscapeTextureArray(WeightmapSize, WeightmapSize, NumWeightTextures, TEXTUREGROUP_Terrain_Weightmap, TSF_BGRA8);
+		MobileWeightmapTextureArray->PostEditChange();
+		MobileWeightmapTextureArray->UpdateResource();
+	}
+	else
+	{
+		MobileWeightmapTextures.SetNum(NumWeightTextures);
+		for (int32 i = 0; i < NumWeightTextures; ++i)
+		{
+			UTexture2D* CurrentWeightmapTexture  = GetLandscapeProxy()->CreateLandscapeTexture(WeightmapSize, WeightmapSize, TEXTUREGROUP_Terrain_Weightmap, TSF_BGRA8);
+			CreateEmptyTextureMips(CurrentWeightmapTexture, true);
+			MobileWeightmapTextures[i] = CurrentWeightmapTexture;
+		}
+	}
+
 	{
 		FLandscapeTextureDataInterface LandscapeData;
 		UTexture2D* CurrentWeightmapTexture = nullptr;
 		int32 CurrentChannel = 0;
 		int32 RemainingChannels = 0;
-
+		int32 index = 0;
 		for (auto& Allocation : MobileWeightmapLayerAllocations)
 		{
-			if (Allocation.LayerInfo)
+			if (!Allocation.LayerInfo)
 			{
-				if (RemainingChannels == 0)
-				{
-					// create a new weightmap texture if we've run out of channels
-					CurrentChannel = 0;
-					RemainingChannels = 4;
-                    CurrentWeightmapTexture = GetLandscapeProxy()->CreateLandscapeTexture(WeightmapSize, WeightmapSize, TEXTUREGROUP_Terrain_Weightmap, TSF_BGRA8);
-					CreateEmptyTextureMips(CurrentWeightmapTexture, true);
-					MobileWeightmapTextures.Add(CurrentWeightmapTexture);
-				}
-
-				CopyRequests.AddWeightmapCopy(CurrentWeightmapTexture, IntCastChecked<int8>(CurrentChannel), this, Allocation.LayerInfo);
+				continue;
+			}
+			
+			if (RemainingChannels == 0)
+			{
+				CurrentChannel = 0;
+				RemainingChannels = 4;
+			}
 
-				// update Allocation
-				Allocation.WeightmapTextureIndex = IntCastChecked<uint8>(MobileWeightmapTextures.Num() - 1);
-				Allocation.WeightmapTextureChannel = IntCastChecked<uint8>(CurrentChannel);
-				CurrentChannel++;
-				RemainingChannels--;
+			int32 Slice = FMath::DivideAndRoundDown(index,4);
+			if (MobileWeightmapTextureArrayEnabled)
+			{
+				CopyRequests.AddWeightmapCopy(MobileWeightmapTextureArray, Slice, IntCastChecked<int8>(CurrentChannel), this, Allocation.LayerInfo);	
 			}
+			else
+			{
+				CopyRequests.AddWeightmapCopy(MobileWeightmapTextures[Slice], 0, IntCastChecked<int8>(CurrentChannel), this, Allocation.LayerInfo);
+			}
+			
+			// update Allocation
+			Allocation.WeightmapTextureIndex = IntCastChecked<uint8>(Slice);
+			Allocation.WeightmapTextureChannel = IntCastChecked<uint8>(CurrentChannel);
+			CurrentChannel++;
+			RemainingChannels--;
+			index++;
 		}
+		
 	}
-
+	
 	CopyRequests.ProcessTextureCopies();
-
-	GDisableAutomaticTextureMaterialUpdateDependencies = true;
-	for (int TextureIdx = 0; TextureIdx < MobileWeightmapTextures.Num(); TextureIdx++)
+	
+	if (MobileWeightmapTextureArray)
 	{
-		UTexture* Texture = MobileWeightmapTextures[TextureIdx];
-		Texture->PostEditChange();
+		GDisableAutomaticTextureMaterialUpdateDependencies = true;
+		MobileWeightmapTextureArray->PostEditChange();
+		MobileWeightmapTextureArray->UpdateResource();
+		MobileWeightmapTextureArray->SetDeterministicLightingGuid();
+		GDisableAutomaticTextureMaterialUpdateDependencies = false;
+	}
+	else
+	{
+		GDisableAutomaticTextureMaterialUpdateDependencies = true;
+		for (int TextureIdx = 0; TextureIdx < MobileWeightmapTextures.Num(); TextureIdx++)
+		{
+			UTexture* Texture = MobileWeightmapTextures[TextureIdx];
+			Texture->PostEditChange();
 
-		// PostEditChange() will assign a random GUID to the texture, which leads to non-deterministic builds.
-		Texture->SetDeterministicLightingGuid();
+			// PostEditChange() will assign a random GUID to the texture, which leads to non-deterministic builds.
+			Texture->SetDeterministicLightingGuid();
+		}
+		GDisableAutomaticTextureMaterialUpdateDependencies = false;
 	}
-	GDisableAutomaticTextureMaterialUpdateDependencies = false;
 
 	FLinearColor Masks[4];
 	Masks[0] = FLinearColor(1, 0, 0, 0);
@@ -7432,11 +7477,18 @@ void ULandscapeComponent::GenerateMobilePlatformPixelData(bool bIsCooking, const
 				}
 			}
 
-			for (int TextureIdx = 0; TextureIdx < MobileWeightmapTextures.Num(); TextureIdx++)
+			if (MobileWeightmapTextureArray)
 			{
-				NewMobileMaterialInstance->SetTextureParameterValue(FName(*FString::Printf(TEXT("Weightmap%d"), TextureIdx)), MobileWeightmapTextures[TextureIdx]);
+				NewMobileMaterialInstance->SetTextureParameterValue(TEXT("WeightmapArray"), MobileWeightmapTextureArray);	
 			}
-
+			else
+			{
+				for (int TextureIdx = 0; TextureIdx < MobileWeightmapTextures.Num(); TextureIdx++)
+				{
+					NewMobileMaterialInstance->SetTextureParameterValue(FName(*FString::Printf(TEXT("Weightmap%d"), TextureIdx)), MobileWeightmapTextures[TextureIdx]);
+				}	
+			}
+			
 			MobileMaterialInterfaces.Add(NewMobileMaterialInstance);
 		}
 	}
@@ -7494,11 +7546,18 @@ void ULandscapeComponent::GenerateMobilePlatformPixelData(bool bIsCooking, const
 				}
 			}
 
-			for (int TextureIdx = 0; TextureIdx < MobileWeightmapTextures.Num(); TextureIdx++)
+			if (MobileWeightmapTextureArray)
 			{
-				NewMobileMaterialInstance->SetTextureParameterValueEditorOnly(FName(*FString::Printf(TEXT("Weightmap%d"), TextureIdx)), MobileWeightmapTextures[TextureIdx]);
+				NewMobileMaterialInstance->SetTextureParameterValueEditorOnly(TEXT("WeightmapArray"), MobileWeightmapTextureArray);
 			}
-
+			else
+			{
+				for (int TextureIdx = 0; TextureIdx < MobileWeightmapTextures.Num(); TextureIdx++)
+				{
+					NewMobileMaterialInstance->SetTextureParameterValueEditorOnly(FName(*FString::Printf(TEXT("Weightmap%d"), TextureIdx)), MobileWeightmapTextures[TextureIdx]);
+				}
+			}
+			
 			NewMobileMaterialInstance->PostEditChange();
 
 			MobileMaterialInterfaces.Add(NewMobileMaterialInstance);
@@ -7545,6 +7604,24 @@ UTexture2D* ALandscapeProxy::CreateLandscapeTexture(int32 InSizeX, int32 InSizeY
 	return NewTexture;
 }
 
+UTexture2DArray* ALandscapeProxy::CreateLandscapeTextureArray(int32 InSizeX, int32 InSizeY, int32 Slices, TextureGroup InLODGroup, ETextureSourceFormat InFormat, UObject* OptionalOverrideOuter )
+{
+	UObject* TexOuter = OptionalOverrideOuter ? OptionalOverrideOuter : const_cast<ALandscapeProxy*>(this);
+	UTexture2DArray* NewTextureArray = NewObject<UTexture2DArray>(TexOuter);
+	
+	const int32 NumMips = FMath::FloorLog2(FMath::Max(InSizeX, InSizeY)) + 1;
+	NewTextureArray->Source.Init(InSizeX, InSizeY, Slices, NumMips, InFormat);
+
+	NewTextureArray->SRGB = false;
+	NewTextureArray->CompressionNone = true;
+	NewTextureArray->MipGenSettings = TMGS_LeaveExistingMips;
+	NewTextureArray->AddressX = TA_Clamp;
+	NewTextureArray->AddressY = TA_Clamp;
+	NewTextureArray->LODGroup = InLODGroup;
+	
+	return NewTextureArray;
+}
+
 UTexture2D* ALandscapeProxy::CreateLandscapeToolTexture(int32 InSizeX, int32 InSizeY, TextureGroup InLODGroup, ETextureSourceFormat InFormat) const
 {
 	UObject* TexOuter = const_cast<ALandscapeProxy*>(this);
diff --git a/Engine/Source/Runtime/Landscape/Private/LandscapeUtils.cpp b/Engine/Source/Runtime/Landscape/Private/LandscapeUtils.cpp
--- a/Engine/Source/Runtime/Landscape/Private/LandscapeUtils.cpp
+++ b/Engine/Source/Runtime/Landscape/Private/LandscapeUtils.cpp
@@ -87,7 +87,7 @@ bool operator==(const FTextureCopyRequest& InEntryA, const FTextureCopyRequest&
 	return (InEntryA.Source == InEntryB.Source) && (InEntryA.Destination == InEntryB.Destination);
 }
 
-bool FBatchTextureCopy::AddWeightmapCopy(UTexture2D* InDestination, int8 InDestinationChannel, const ULandscapeComponent* InComponent, ULandscapeLayerInfoObject* InLayerInfo)
+bool FBatchTextureCopy::AddWeightmapCopy(UTexture* InDestination, int8 InDestinationSlice, int8 InDestinationChannel, const ULandscapeComponent* InComponent, ULandscapeLayerInfoObject* InLayerInfo)
 {
 	FTextureCopyRequest CopyRequest;
 	const TArray<UTexture2D*>& ComponentWeightmapTextures = InComponent->GetWeightmapTextures();
@@ -95,6 +95,7 @@ bool FBatchTextureCopy::AddWeightmapCopy(UTexture2D* InDestination, int8 InDesti
 	int8 SourceChannel = INDEX_NONE;
 
 	CopyRequest.Destination = InDestination;
+	CopyRequest.DestinationSlice = InDestinationSlice;
 
 	// Find the proper Source Texture and channel from Layer Allocations
 	for (const FWeightmapLayerAllocationInfo& ComponentWeightmapLayerAllocation : ComponentWeightmapLayerAllocations)
@@ -137,7 +138,7 @@ bool FBatchTextureCopy::ProcessTextureCopies()
 {
 	TRACE_CPUPROFILER_EVENT_SCOPE(FBatchTextureCopy::ProcessTextureCopyRequest);
 	TMap<UTexture2D*, FSourceDataMipNumber> Sources;
-	TMap<UTexture2D*, FDestinationDataMipNumber> Destinations;
+	TMap<UTexture*, FDestinationDataMipNumber> Destinations;
 
 	if (CopyRequests.Num() == 0)
 	{
@@ -161,7 +162,7 @@ bool FBatchTextureCopy::ProcessTextureCopies()
 	}
 
 	// Lock all destinations mips
-	for (TPair<UTexture2D*, FDestinationDataMipNumber>& Destination : Destinations)
+	for (TPair<UTexture*, FDestinationDataMipNumber>& Destination : Destinations)
 	{
 		int32 MipNumber = Destination.Value.MipNumber;
 		TArray<uint8*>& DestinationDataPtr = Destination.Value.DestinationDataPtr;
@@ -184,13 +185,15 @@ bool FBatchTextureCopy::ProcessTextureCopies()
 
 		for (int32 MipLevel = 0; MipLevel < MipNumber; ++MipLevel)
 		{
+			const int64 MipSizeInBytes = CopyRequest.Key.Source->Source.CalcMipSize(MipLevel);
+			
 			const int32 MipSize = CopyRequest.Key.Destination->Source.GetSizeX() >> MipLevel;
 			check(MipSize == (CopyRequest.Key.Destination->Source.GetSizeY() >> MipLevel));
 
 			int32 MipSizeSquare = FMath::Square(MipSize);
 			FSharedBuffer MipSrcData = SourceDataMipNumber->MipData->GetMipData(0, 0, MipLevel);
 			const uint8* SourceTextureData = static_cast<const uint8*>(MipSrcData.GetData());
-			uint8* DestTextureData = DestinationDataMipNumber->DestinationDataPtr[MipLevel];
+			uint8* DestTextureData = DestinationDataMipNumber->DestinationDataPtr[MipLevel] + CopyRequest.Key.DestinationSlice * MipSizeInBytes;
 
 			check((SourceTextureData != nullptr) && (DestTextureData != nullptr));
 
@@ -217,7 +220,7 @@ bool FBatchTextureCopy::ProcessTextureCopies()
 	// Note that source textures do not need unlocking, data will be released once the FMipData go out of scope
 	
 	// Unlock all destination mips
-	for (TPair<UTexture2D*, FDestinationDataMipNumber>& Destination : Destinations)
+	for (TPair<UTexture*, FDestinationDataMipNumber>& Destination : Destinations)
 	{
 		int32 MipNumber = Destination.Value.MipNumber;
 		
@@ -230,6 +233,22 @@ bool FBatchTextureCopy::ProcessTextureCopies()
 	return true;
 }
 
+int32 LandscapeMobileWeightTextureArray = 0;
+static FAutoConsoleVariableRef CVarLandscapeMobileWeightTextureArray(
+	TEXT("landscape.MobileWeightTextureArray"),
+	LandscapeMobileWeightTextureArray,
+	TEXT("Use Texture Arrays for weights on Mobile platforms"),
+	ECVF_ReadOnly);
+
+bool IsMobileWeightmapTextureArrayEnabled()
+{
+	return LandscapeMobileWeightTextureArray != 0;	
+}
+	
+bool UseWeightmapTextureArray(EShaderPlatform InPlatform)
+{
+	return IsMobilePlatform(InPlatform) && (LandscapeMobileWeightTextureArray != 0);	
+}
 
 #endif //!WITH_EDITOR
 
diff --git a/Engine/Source/Runtime/Landscape/Private/Materials/MaterialExpressionLandscapeLayerBlend.cpp b/Engine/Source/Runtime/Landscape/Private/Materials/MaterialExpressionLandscapeLayerBlend.cpp
--- a/Engine/Source/Runtime/Landscape/Private/Materials/MaterialExpressionLandscapeLayerBlend.cpp
+++ b/Engine/Source/Runtime/Landscape/Private/Materials/MaterialExpressionLandscapeLayerBlend.cpp
@@ -5,6 +5,9 @@
 #include "Engine/Texture.h"
 #include "EngineGlobals.h"
 #include "MaterialCompiler.h"
+#include "DataDrivenShaderPlatformInfo.h"
+#include "LandscapeUtils.h"
+
 #if WITH_EDITOR
 #include "MaterialGraph/MaterialGraphNode.h"
 #include "MaterialHLSLGenerator.h"
@@ -154,6 +157,8 @@ int32 UMaterialExpressionLandscapeLayerBlend::Compile(class FMaterialCompiler* C
 	TArray<int32> WeightCodes;
 	WeightCodes.Empty(Layers.Num());
 
+	const bool bTextureArrayEnabled = UE::Landscape::UseWeightmapTextureArray(Compiler->GetShaderPlatform());
+	
 	for (int32 LayerIdx = 0; LayerIdx<Layers.Num(); LayerIdx++)
 	{
 		WeightCodes.Add(INDEX_NONE);
@@ -167,7 +172,7 @@ int32 UMaterialExpressionLandscapeLayerBlend::Compile(class FMaterialCompiler* C
 			const int32 HeightCode = Layer.HeightInput.Expression ? Layer.HeightInput.Compile(Compiler) : Compiler->Constant(Layer.ConstHeightInput);
 
 			const int32 DefaultWeightCode = Layer.PreviewWeight > 0.0f ? Compiler->Constant(Layer.PreviewWeight) : INDEX_NONE;
-			const int32 WeightCode = Compiler->StaticTerrainLayerWeight(Layer.LayerName, DefaultWeightCode);
+			const int32 WeightCode = Compiler->StaticTerrainLayerWeight(Layer.LayerName, DefaultWeightCode, bTextureArrayEnabled);
 			if (WeightCode != INDEX_NONE)
 			{
 				switch (Layer.BlendType)
@@ -234,7 +239,7 @@ int32 UMaterialExpressionLandscapeLayerBlend::Compile(class FMaterialCompiler* C
 		if (Layer.BlendType == LB_AlphaBlend)
 		{
 			const int32 DefaultWeightCode = Layer.PreviewWeight > 0.0f ? Compiler->Constant(Layer.PreviewWeight) : INDEX_NONE;
-			const int32 WeightCode = Compiler->StaticTerrainLayerWeight(Layer.LayerName, DefaultWeightCode);
+			const int32 WeightCode = Compiler->StaticTerrainLayerWeight(Layer.LayerName, DefaultWeightCode, bTextureArrayEnabled);
 			if (WeightCode != INDEX_NONE)
 			{
 				const int32 LayerCode = Layer.LayerInput.Expression
@@ -305,7 +310,8 @@ bool UMaterialExpressionLandscapeLayerBlend::GenerateHLSLExpression(FMaterialHLS
 		if (Layer.BlendType != LB_AlphaBlend)
 		{
 			const FExpression* WeightExpression = nullptr;
-			verify(GenerateStaticTerrainLayerWeightExpression(Layer.LayerName, Layer.PreviewWeight, Generator, WeightExpression));
+			const bool bTextureArrayEnabled = UE::Landscape::IsMobileWeightmapTextureArrayEnabled();
+			verify(GenerateStaticTerrainLayerWeightExpression(Layer.LayerName, Layer.PreviewWeight, bTextureArrayEnabled, Generator, WeightExpression));
 
 			if (WeightExpression)
 			{
@@ -368,7 +374,8 @@ bool UMaterialExpressionLandscapeLayerBlend::GenerateHLSLExpression(FMaterialHLS
 		if (Layer.BlendType == LB_AlphaBlend)
 		{
 			const FExpression* WeightExpression = nullptr;
-			verify(GenerateStaticTerrainLayerWeightExpression(Layer.LayerName, Layer.PreviewWeight, Generator, WeightExpression));
+			const bool bTextureArrayEnabled = UE::Landscape::IsMobileWeightmapTextureArrayEnabled();
+			verify(GenerateStaticTerrainLayerWeightExpression(Layer.LayerName, Layer.PreviewWeight, bTextureArrayEnabled, Generator, WeightExpression));
 
 			if (WeightExpression)
 			{
@@ -390,6 +397,11 @@ UObject* UMaterialExpressionLandscapeLayerBlend::GetReferencedTexture() const
 	return GEngine->WeightMapPlaceholderTexture;
 }
 
+UMaterialExpression::ReferencedTextureArray UMaterialExpressionLandscapeLayerBlend::GetReferencedTextures() const
+{
+	return { GEngine->WeightMapPlaceholderTexture, GEngine->WeightMapArrayPlaceholderTexture };
+}
+
 #if WITH_EDITOR
 void UMaterialExpressionLandscapeLayerBlend::GetCaption(TArray<FString>& OutCaptions) const
 {
diff --git a/Engine/Source/Runtime/Landscape/Private/Materials/MaterialExpressionLandscapeLayerSample.cpp b/Engine/Source/Runtime/Landscape/Private/Materials/MaterialExpressionLandscapeLayerSample.cpp
--- a/Engine/Source/Runtime/Landscape/Private/Materials/MaterialExpressionLandscapeLayerSample.cpp
+++ b/Engine/Source/Runtime/Landscape/Private/Materials/MaterialExpressionLandscapeLayerSample.cpp
@@ -6,6 +6,8 @@
 #include "EngineGlobals.h"
 #include "MaterialCompiler.h"
 #include "Materials/Material.h"
+#include "LandscapeUtils.h"
+
 #if WITH_EDITOR
 #include "MaterialHLSLGenerator.h"
 #endif
@@ -40,7 +42,8 @@ UMaterialExpressionLandscapeLayerSample::UMaterialExpressionLandscapeLayerSample
 #if WITH_EDITOR
 int32 UMaterialExpressionLandscapeLayerSample::Compile(class FMaterialCompiler* Compiler, int32 OutputIndex)
 {
-	const int32 WeightCode = Compiler->StaticTerrainLayerWeight(ParameterName, Compiler->Constant(PreviewWeight));
+	const bool bTextureArrayEnabled = UE::Landscape::UseWeightmapTextureArray(Compiler->GetShaderPlatform());
+	const int32 WeightCode = Compiler->StaticTerrainLayerWeight(ParameterName, Compiler->Constant(PreviewWeight), bTextureArrayEnabled);
 	if (WeightCode == INDEX_NONE)
 	{
 		// layer is not used in this component, sample value is 0.
@@ -54,7 +57,8 @@ int32 UMaterialExpressionLandscapeLayerSample::Compile(class FMaterialCompiler*
 
 bool UMaterialExpressionLandscapeLayerSample::GenerateHLSLExpression(FMaterialHLSLGenerator& Generator, UE::HLSLTree::FScope& Scope, int32 OutputIndex, UE::HLSLTree::FExpression const*& OutExpression) const
 {
-	return GenerateStaticTerrainLayerWeightExpression(ParameterName, PreviewWeight, Generator, OutExpression);
+	const bool bTextureArrayEnabled = UE::Landscape::IsMobileWeightmapTextureArrayEnabled();
+	return GenerateStaticTerrainLayerWeightExpression(ParameterName, PreviewWeight, bTextureArrayEnabled, Generator, OutExpression);
 }
 
 #endif // WITH_EDITOR
@@ -64,6 +68,11 @@ UObject* UMaterialExpressionLandscapeLayerSample::GetReferencedTexture() const
 	return GEngine->WeightMapPlaceholderTexture;
 }
 
+UMaterialExpression::ReferencedTextureArray UMaterialExpressionLandscapeLayerSample::GetReferencedTextures() const
+{
+	return { GEngine->WeightMapPlaceholderTexture, GEngine->WeightMapArrayPlaceholderTexture };
+}
+
 #if WITH_EDITOR
 FString UMaterialExpressionLandscapeLayerSample::GetEditableName() const
 {
diff --git a/Engine/Source/Runtime/Landscape/Private/Materials/MaterialExpressionLandscapeLayerSwitch.cpp b/Engine/Source/Runtime/Landscape/Private/Materials/MaterialExpressionLandscapeLayerSwitch.cpp
--- a/Engine/Source/Runtime/Landscape/Private/Materials/MaterialExpressionLandscapeLayerSwitch.cpp
+++ b/Engine/Source/Runtime/Landscape/Private/Materials/MaterialExpressionLandscapeLayerSwitch.cpp
@@ -6,6 +6,7 @@
 #include "EngineGlobals.h"
 #include "MaterialCompiler.h"
 #include "Materials/Material.h"
+#include "LandscapeUtils.h"
 
 #include UE_INLINE_GENERATED_CPP_BY_NAME(MaterialExpressionLandscapeLayerSwitch)
 
@@ -49,9 +50,11 @@ bool UMaterialExpressionLandscapeLayerSwitch::IsResultMaterialAttributes(int32 O
 
 int32 UMaterialExpressionLandscapeLayerSwitch::Compile(class FMaterialCompiler* Compiler, int32 OutputIndex)
 {
+	const bool bTextureArrayEnabled = UE::Landscape::UseWeightmapTextureArray(Compiler->GetShaderPlatform());
 	const int32 WeightCode = Compiler->StaticTerrainLayerWeight(
 		ParameterName,
-		PreviewUsed ? Compiler->Constant(1.0f) : INDEX_NONE
+		PreviewUsed ? Compiler->Constant(1.0f) : INDEX_NONE,
+		bTextureArrayEnabled
 		);
 
 	int32 ReturnCode = INDEX_NONE;
@@ -80,6 +83,11 @@ UObject* UMaterialExpressionLandscapeLayerSwitch::GetReferencedTexture() const
 	return GEngine->WeightMapPlaceholderTexture;
 }
 
+UMaterialExpression::ReferencedTextureArray UMaterialExpressionLandscapeLayerSwitch::GetReferencedTextures() const
+{
+	return { GEngine->WeightMapPlaceholderTexture, GEngine->WeightMapArrayPlaceholderTexture };
+}
+
 #if WITH_EDITOR
 FString UMaterialExpressionLandscapeLayerSwitch::GetEditableName() const
 {
diff --git a/Engine/Source/Runtime/Landscape/Private/Materials/MaterialExpressionLandscapeLayerWeight.cpp b/Engine/Source/Runtime/Landscape/Private/Materials/MaterialExpressionLandscapeLayerWeight.cpp
--- a/Engine/Source/Runtime/Landscape/Private/Materials/MaterialExpressionLandscapeLayerWeight.cpp
+++ b/Engine/Source/Runtime/Landscape/Private/Materials/MaterialExpressionLandscapeLayerWeight.cpp
@@ -6,6 +6,7 @@
 #include "EngineGlobals.h"
 #include "MaterialCompiler.h"
 #include "Materials/Material.h"
+#include "LandscapeUtils.h"
 
 #include UE_INLINE_GENERATED_CPP_BY_NAME(MaterialExpressionLandscapeLayerWeight)
 
@@ -58,10 +59,11 @@ bool UMaterialExpressionLandscapeLayerWeight::IsResultMaterialAttributes(int32 O
 
 int32 UMaterialExpressionLandscapeLayerWeight::Compile(class FMaterialCompiler* Compiler, int32 OutputIndex)
 {
+	const bool bTextureArrayEnabled = UE::Landscape::UseWeightmapTextureArray(Compiler->GetShaderPlatform());
 	const int32 BaseCode = Base.Expression
 		                       ? Base.Compile(Compiler)
 		                       : Compiler->Constant3(static_cast<float>(ConstBase.X), static_cast<float>(ConstBase.Y), static_cast<float>(ConstBase.Z));
-	const int32 WeightCode = Compiler->StaticTerrainLayerWeight(ParameterName, Compiler->Constant(PreviewWeight));
+	const int32 WeightCode = Compiler->StaticTerrainLayerWeight(ParameterName, Compiler->Constant(PreviewWeight), bTextureArrayEnabled);
 
 	int32 ReturnCode = INDEX_NONE;
 	if (WeightCode == INDEX_NONE)
@@ -90,6 +92,11 @@ UObject* UMaterialExpressionLandscapeLayerWeight::GetReferencedTexture() const
 	return GEngine->WeightMapPlaceholderTexture;
 }
 
+UMaterialExpression::ReferencedTextureArray UMaterialExpressionLandscapeLayerWeight::GetReferencedTextures() const
+{
+	return { GEngine->WeightMapPlaceholderTexture, GEngine->WeightMapArrayPlaceholderTexture };
+}
+
 #if WITH_EDITOR
 FString UMaterialExpressionLandscapeLayerWeight::GetEditableName() const
 {
diff --git a/Engine/Source/Runtime/Landscape/Private/Materials/MaterialExpressionLandscapeVisibilityMask.cpp b/Engine/Source/Runtime/Landscape/Private/Materials/MaterialExpressionLandscapeVisibilityMask.cpp
--- a/Engine/Source/Runtime/Landscape/Private/Materials/MaterialExpressionLandscapeVisibilityMask.cpp
+++ b/Engine/Source/Runtime/Landscape/Private/Materials/MaterialExpressionLandscapeVisibilityMask.cpp
@@ -6,6 +6,8 @@
 #include "EngineGlobals.h"
 #include "MaterialCompiler.h"
 #include "Materials/Material.h"
+#include "LandscapeUtils.h"
+
 #if WITH_EDITOR
 #include "MaterialHLSLGenerator.h"
 #endif
@@ -43,7 +45,8 @@ UMaterialExpressionLandscapeVisibilityMask::UMaterialExpressionLandscapeVisibili
 #if WITH_EDITOR
 int32 UMaterialExpressionLandscapeVisibilityMask::Compile(class FMaterialCompiler* Compiler, int32 OutputIndex)
 {
-	int32 MaskLayerCode = Compiler->StaticTerrainLayerWeight(ParameterName, Compiler->Constant(0.f));
+	const bool bTextureArrayEnabled = UE::Landscape::UseWeightmapTextureArray(Compiler->GetShaderPlatform());
+	int32 MaskLayerCode = Compiler->StaticTerrainLayerWeight(ParameterName, Compiler->Constant(0.f), bTextureArrayEnabled);
 	return MaskLayerCode == INDEX_NONE ? Compiler->Constant(1.f) : Compiler->Sub(Compiler->Constant(1.f), MaskLayerCode);
 }
 
@@ -52,7 +55,8 @@ bool UMaterialExpressionLandscapeVisibilityMask::GenerateHLSLExpression(FMateria
 	using namespace UE::HLSLTree;
 
 	const FExpression* MaskLayerExpression = nullptr;
-	verify(GenerateStaticTerrainLayerWeightExpression(ParameterName, 0.f, Generator, MaskLayerExpression));
+	const bool bTextureArrayEnabled = UE::Landscape::IsMobileWeightmapTextureArrayEnabled();
+	verify(GenerateStaticTerrainLayerWeightExpression(ParameterName, 0.f, bTextureArrayEnabled, Generator, MaskLayerExpression));
 
 	FTree& Tree = Generator.GetTree();
 	const FExpression* ConstantOne = Tree.NewConstant(1.f);
@@ -66,6 +70,11 @@ UObject* UMaterialExpressionLandscapeVisibilityMask::GetReferencedTexture() cons
 	return GEngine->WeightMapPlaceholderTexture;
 }
 
+UMaterialExpression::ReferencedTextureArray UMaterialExpressionLandscapeVisibilityMask::GetReferencedTextures() const
+{
+	return { GEngine->WeightMapPlaceholderTexture, GEngine->WeightMapArrayPlaceholderTexture };
+}
+
 #if WITH_EDITOR
 void UMaterialExpressionLandscapeVisibilityMask::GetLandscapeLayerNames(TArray<FName>& OutLayers) const
 {
diff --git a/Engine/Source/Runtime/Landscape/Public/LandscapeUtils.h b/Engine/Source/Runtime/Landscape/Public/LandscapeUtils.h
--- a/Engine/Source/Runtime/Landscape/Public/LandscapeUtils.h
+++ b/Engine/Source/Runtime/Landscape/Public/LandscapeUtils.h
@@ -14,6 +14,7 @@ enum EShaderPlatform : uint16;
 class ULandscapeComponent;
 class ULandscapeLayerInfoObject;
 class UTexture2D;
+class UTexture;
 
 namespace UE::Landscape
 {
@@ -29,7 +30,8 @@ LANDSCAPE_API bool DoesPlatformSupportEditLayers(EShaderPlatform InShaderPlatfor
 struct FTextureCopyRequest
 {
 	UTexture2D* Source = nullptr;
-	UTexture2D* Destination = nullptr;
+	UTexture* Destination = nullptr;
+	int8 DestinationSlice = 0;
 };
 
 uint32 GetTypeHash(const FTextureCopyRequest& InKey);
@@ -57,11 +59,12 @@ public:
 	* Uses the provided arguments to add proper source/destination entries to internal copy requests.
 	* @param	InDestination	The texture used as a destination for the copy.
 	* @param	InDestinationChannel	The channel used as a destination for the copy.
+	* @param	InDestinationSlice	The Texture array slice to write to (use 0 for a Texture2D)
 	* @param	InComponent		The component containing the wanted source weightmap.
 	* @param	InLayerInfo		The layer info used to retrieve the proper source weightmap and channel.
 	* @return True if the copy has been successfully added.
 	*/
-	LANDSCAPE_API bool AddWeightmapCopy(UTexture2D* InDestination, int8 InDestinationChannel, const ULandscapeComponent* InComponent, ULandscapeLayerInfoObject* InLayerInfo);
+	LANDSCAPE_API bool AddWeightmapCopy(UTexture* InDestination, int8 InDestinationChannel, int8 InDestinationSlice, const ULandscapeComponent* InComponent, ULandscapeLayerInfoObject* InLayerInfo);
 
 	/** Process pending internal copy requests. */
 	LANDSCAPE_API bool ProcessTextureCopies();
@@ -98,6 +101,12 @@ LANDSCAPE_API FString GetLayerInfoObjectPackageName(const ULevel* InLevel, const
 /** Returns true if the provided layer info object is the current visibility layer. */
 LANDSCAPE_API bool IsVisibilityLayer(const ULandscapeLayerInfoObject* InLayerInfoObject);
 
+/** Returns true if InPlatform is a mobile platform and the Landscape.MobileWeightTextureArray CVar set */
+bool UseWeightmapTextureArray(EShaderPlatform InPlatform);	
+
+/** Check if Landscape.MobileWeightTextureArray CVar set and we should attempt to use texture arrays for weight maps on mobile platforms  */
+bool IsMobileWeightmapTextureArrayEnabled();
+	
 #endif //!WITH_EDITOR
 
 } // end namespace UE::Landscape
