Async Runtime Grass Map Generation

Intended to save disk space and runtime memory, at the cost of generating grass density maps at runtime.

Enabled per platform via cvar "grass.GrassMap.UseRuntimeGeneration" (disabled by default), and enabled per landscape via the "bUseRuntimeGrassMapGeneration" object property (also disabled by default).  Both must be enabled at cook time for disk and memory savings.

Original review in Main:  https://p4-swarm.epicgames.net/reviews/29027747

#tests console testers tested grass in cooked builds on all platforms, I ran and profiled game projects in editor and cooked standalone, and tested grass editing in game projects in editor

diff --git a/Engine/Shaders/Private/LandscapeVertexFactory.ush b/Engine/Shaders/Private/LandscapeVertexFactory.ush
--- a/Engine/Shaders/Private/LandscapeVertexFactory.ush
+++ b/Engine/Shaders/Private/LandscapeVertexFactory.ush
@@ -4,6 +4,8 @@
 	LandscapeVertexFactory.usf: Landscape vertex factory.
 =============================================================================*/
 
+// Compile Version 2
+
 #include "VertexFactoryCommon.ush"
 
 #define VERTEX_FACTORY_MODIFIES_INTERPOLATION 1
diff --git a/Engine/Source/Editor/LandscapeEditor/Private/LandscapeEdModeTools.h b/Engine/Source/Editor/LandscapeEditor/Private/LandscapeEdModeTools.h
--- a/Engine/Source/Editor/LandscapeEditor/Private/LandscapeEdModeTools.h
+++ b/Engine/Source/Editor/LandscapeEditor/Private/LandscapeEdModeTools.h
@@ -1479,15 +1479,15 @@ public:
 	{
 		TRACE_BOOKMARK(TEXT("BeginTool - %s"), GetToolName());
 
-		if (ShouldUpdateEditingLayer())
+		ALandscape* Landscape = this->EdMode->GetLandscape();
+		if (Landscape)
 		{
-			ALandscape* Landscape = this->EdMode->GetLandscape();
-			if (Landscape)
+			if (ShouldUpdateEditingLayer())
 			{
 				Landscape->RequestLayersContentUpdate(GetBeginToolContentUpdateFlag());
 				Landscape->SetEditingLayer(this->EdMode->GetCurrentLayerGuid());	// legacy way to set the edit layer, via Landscape state
-				Landscape->SetGrassUpdateEnabled(false);
 			}
+			Landscape->SetGrassUpdateEnabled(false);
 		}
 
 		if (!ensure(InteractorPositions.Num() == 0))
@@ -1555,15 +1555,15 @@ public:
 		EdMode->CurrentBrush->EndStroke();
 		EdMode->UpdateLayerUsageInformation(&EdMode->CurrentToolTarget.LayerInfo);
 
-		if (ShouldUpdateEditingLayer())
+		ALandscape* Landscape = this->EdMode->GetLandscape();
+		if (Landscape)
 		{
-			ALandscape* Landscape = this->EdMode->GetLandscape();
-			if (Landscape)
+			if (ShouldUpdateEditingLayer())
 			{
 				Landscape->RequestLayersContentUpdate(GetEndToolContentUpdateFlag());
 				Landscape->SetEditingLayer();
-				Landscape->SetGrassUpdateEnabled(true);
 			}
+			Landscape->SetGrassUpdateEnabled(true);
 		}
 
 		TRACE_BOOKMARK(TEXT("EndTool - %s"), GetToolName());
diff --git a/Engine/Source/Runtime/Engine/Classes/Materials/Material.h b/Engine/Source/Runtime/Engine/Classes/Materials/Material.h
--- a/Engine/Source/Runtime/Engine/Classes/Materials/Material.h
+++ b/Engine/Source/Runtime/Engine/Classes/Materials/Material.h
@@ -1719,11 +1719,12 @@ public:
 	 *	@Param	InQuality							Optional quality switch - if supplied, only walk QualitySwitch branches according to it.
 	 *	@Param	InShadingPath						Optional shading path switch - if supplied, only walk ShadingPathSwitch branches according to it.
 	 *  @Param	bInRecurseIntoMaterialFunctions		Optional enable recursion into Material Functions - if true, will also return expressions from inside of MaterialFunctions.
+	 *  @Param	InMobileCustomOutputExpressionTypesToQuery	Optional a list of material expression custom output types to include in the query on mobile devices.  On non-mobile devices, all custom output types are queried.
 	 *
 	 *	@return	bool					true if successful, false if not.
 	 */
 	ENGINE_API virtual bool GetAllReferencedExpressions(TArray<UMaterialExpression*>& OutExpressions, struct FStaticParameterSet* InStaticParameterSet,
-		ERHIFeatureLevel::Type InFeatureLevel = ERHIFeatureLevel::Num, EMaterialQualityLevel::Type InQuality = EMaterialQualityLevel::Num, ERHIShadingPath::Type InShadingPath = ERHIShadingPath::Num, const bool bInRecurseIntoMaterialFunctions = false);
+		ERHIFeatureLevel::Type InFeatureLevel = ERHIFeatureLevel::Num, EMaterialQualityLevel::Type InQuality = EMaterialQualityLevel::Num, ERHIShadingPath::Type InShadingPath = ERHIShadingPath::Num, const bool bInRecurseIntoMaterialFunctions = false, TSet<UClass*>* InMobileCustomOutputExpressionTypesToQuery = nullptr);
 
 
 	/**
diff --git a/Engine/Source/Runtime/Engine/Private/Materials/Material.cpp b/Engine/Source/Runtime/Engine/Private/Materials/Material.cpp
--- a/Engine/Source/Runtime/Engine/Private/Materials/Material.cpp
+++ b/Engine/Source/Runtime/Engine/Private/Materials/Material.cpp
@@ -5899,7 +5899,7 @@ void UMaterial::GetAllExpressionsForCustomInterpolators(TArray<class UMaterialEx
 
 #if WITH_EDITOR
 bool UMaterial::GetAllReferencedExpressions(TArray<UMaterialExpression*>& OutExpressions, struct FStaticParameterSet* InStaticParameterSet,
-	ERHIFeatureLevel::Type InFeatureLevel, EMaterialQualityLevel::Type InQuality, ERHIShadingPath::Type InShadingPath, const bool bInRecurseIntoMaterialFunctions)
+	ERHIFeatureLevel::Type InFeatureLevel, EMaterialQualityLevel::Type InQuality, ERHIShadingPath::Type InShadingPath, const bool bInRecurseIntoMaterialFunctions, TSet<UClass*>* InMobileCustomOutputExpressionTypesToQuery)
 {
 	using namespace MaterialImpl;
 
@@ -5935,16 +5935,18 @@ bool UMaterial::GetAllReferencedExpressions(TArray<UMaterialExpression*>& OutExp
 			}
 		}
 
-		// TODO: Need an actual ShaderPlatform for a more precise result
+		// TODO: Need an actual ShaderPlatform for a more precise result. This just grabs the handiest one on the current machine.
 		EShaderPlatform ShaderPlatform = GetFeatureLevelShaderPlatform(ERHIFeatureLevel::ES3_1);
+		
 		bool bMobileUseVirtualTexturing = UseVirtualTexturing(ShaderPlatform);
-		if (bMobileUseVirtualTexturing)
+		if (bMobileUseVirtualTexturing || (InMobileCustomOutputExpressionTypesToQuery && InMobileCustomOutputExpressionTypesToQuery->Num() > 0))
 		{
 			TArray<class UMaterialExpressionCustomOutput*> CustomOutputExpressions;
 			GetAllCustomOutputExpressions(CustomOutputExpressions);
 			for (UMaterialExpressionCustomOutput* Expression : CustomOutputExpressions)
 			{
-				if (Expression->IsA<UMaterialExpressionRuntimeVirtualTextureOutput>())
+				if ((bMobileUseVirtualTexturing && Expression->IsA<UMaterialExpressionRuntimeVirtualTextureOutput>()) ||
+					(InMobileCustomOutputExpressionTypesToQuery && InMobileCustomOutputExpressionTypesToQuery->Contains(Expression->GetClass())))
 				{
 					TArray<FExpressionInput*> ProcessedInputs;
 					RecursiveGetExpressionChain(Expression, ProcessedInputs, OutExpressions, InStaticParameterSet, InFeatureLevel, InQuality, InShadingPath, SF_NumFrequencies, MP_MAX, bInRecurseIntoMaterialFunctions);
diff --git a/Engine/Source/Runtime/Landscape/Classes/Landscape.h b/Engine/Source/Runtime/Landscape/Classes/Landscape.h
--- a/Engine/Source/Runtime/Landscape/Classes/Landscape.h
+++ b/Engine/Source/Runtime/Landscape/Classes/Landscape.h
@@ -492,7 +492,6 @@ private:
 	void UpdateHeightDirtyData(ULandscapeComponent* InLandscapeComponent, UTexture2D const* InHeightmap, FColor const* InOldData, FColor const* InNewData);
 	void OnDirtyHeightmap(FTextureToComponentHelper const& MapHelper, UTexture2D const* InWeightmap, FColor const* InOldData, FColor const* InNewData, int32 InMipLevel);
 
-	static bool IsTextureReady(UTexture2D* InTexture, bool bInWaitForStreaming);
 	static bool IsMaterialResourceCompiled(FMaterialResource* InMaterialResource, bool bInWaitForCompilation);
 #endif // WITH_EDITOR
 
diff --git a/Engine/Source/Runtime/Landscape/Classes/LandscapeComponent.h b/Engine/Source/Runtime/Landscape/Classes/LandscapeComponent.h
--- a/Engine/Source/Runtime/Landscape/Classes/LandscapeComponent.h
+++ b/Engine/Source/Runtime/Landscape/Classes/LandscapeComponent.h
@@ -184,6 +184,15 @@ inline uint32 GetTypeHash(const FWeightmapLayerAllocationInfo& InAllocInfo)
 	return InAllocInfo.GetHash();
 }
 
+template<typename T>
+struct IBuffer2DView
+{
+	// copy up to Count elements to Dest, in X then Y order (standard image order)
+	virtual void CopyTo(T* Dest, int32 Count) const = 0;
+
+	// return the total number of elements
+	virtual int32 Num() const = 0;
+};
 
 struct FLandscapeComponentGrassData
 {
@@ -207,7 +216,7 @@ struct FLandscapeComponentGrassData
 
 	// Grass data was updated but not saved yet
 	bool bIsDirty = false;
-#endif
+#endif // WITH_EDITORONLY_DATA
 	
 	static constexpr int32 UnknownNumElements = -1;
 	// Elements per contiguous array: for validation and also to indicate whether the grass data is valid (NumElements >= 0, meaning 0 elements is valid but the grass data is all zero and 
@@ -219,9 +228,7 @@ struct FLandscapeComponentGrassData
 
 	FLandscapeComponentGrassData() = default;
 
-#if WITH_EDITOR
 	FLandscapeComponentGrassData(ULandscapeComponent* Component);
-#endif
 
 	// Returns whether grass data has been computed (or serialized) yet. Returns true even if the data is completely empty (e.g. all-zero weightmap data)
 	bool HasValidData() const;
@@ -230,6 +237,7 @@ struct FLandscapeComponentGrassData
 	bool HasData() const;
 
 	void InitializeFrom(const TArray<uint16>& HeightData, const TMap<ULandscapeGrassType*, TArray<uint8>>& WeightData);
+	void InitializeFrom(IBuffer2DView<uint16>* HeightData, TMap<ULandscapeGrassType*, IBuffer2DView<uint8>*>& WeightData);
 
 	bool HasWeightData() const;
 	TArrayView<uint8> GetWeightData(const ULandscapeGrassType* GrassType);
@@ -514,18 +522,31 @@ private:
 	UPROPERTY(EditAnywhere, Category = LandscapeComponent)
 	TArray<FLandscapePerLODMaterialOverride> PerLODOverrideMaterials;
 
+#if WITH_EDITORONLY_DATA
 	/** The value of the landscape material AllStateCRC the last time the GrassTypes array was updated from it */
 	uint32 LastLandscapeMaterialAllStateCRCWhenGrassTypesBuilt = 0;
+#endif // WITH_EDITORONLY_DATA
 
-	/** Cached list of grass types supported by the component's material. Call UpdateGrassTypes() to ensure this array is up to date */
+	/** Cached list of grass types supported by the component's material.
+	  *	This is needed in a cooked build, as the grass types list is not available
+	  * on the cooked material.
+	  * Call UpdateGrassTypes() to ensure this array is up to date */
 	UPROPERTY()
 	TArray<TObjectPtr<ULandscapeGrassType>> GrassTypes;
 
-	/** Cached max discard distance of all grass types supported by the component's material*/
-	UPROPERTY()
-	float GrassTypesMaxDiscardDistance = 0.0f;
-
 public:
+	// Non-serialized runtime cache of values derived from the assigned grass types.
+	// Call ALandscapeProxy::UpdateGrassTypeSummary() to update.
+	struct FGrassTypeSummary
+	{
+		bool bInvalid = true;
+		double MaxInstanceDiscardDistance = DBL_MAX;
+	};
+	FGrassTypeSummary GrassTypeSummary;
+	inline bool IsGrassTypeSummaryValid() { return GrassTypeSummary.bInvalid; }
+
+	/** Invalidate the grass type summary.  Call whenever grass types are changed to indicate that the summary values are out of date. */
+	LANDSCAPE_API void InvalidateGrassTypeSummary();
 
 	/** Uniquely identifies this component's built map data. */
 	UPROPERTY()
@@ -673,6 +694,20 @@ public:
 	 * Returns false if there are no non-null physical materials. (We probably don't want to use if no physical material connections are bound.)
 	 */
 	bool GetRenderPhysicalMaterials(TArray<UPhysicalMaterial*>& OutPhysicalMaterials) const;
+
+	// Non-serialized data used to implement BeginCacheForCookedPlatformData() / ClearAllCachedCookedPlatformData()
+	struct FPlatformCook
+	{
+		// -1 indicates not currently cached for a cooked platform
+		int32 CurrentCookedPlatformOrdinal = -1;
+
+		// flags indicating if the current cooked platform stripped grass or collision
+		bool bStripGrassData = false;
+
+		// if Landscape data is stripped by a platform cook, it is stored here so it can be restored
+		TSharedPtr<FLandscapeComponentGrassData> StrippedGrassData;
+	};
+	FPlatformCook PlatformCook;
 #endif // WITH_EDITOR
 
 	//~ Begin UObject Interface.	
@@ -690,6 +725,8 @@ public:
 
 #if WITH_EDITOR
 	virtual void BeginCacheForCookedPlatformData(const ITargetPlatform* TargetPlatform) override;
+	virtual bool IsCachedCookedPlatformDataLoaded(const ITargetPlatform* TargetPlatform) override;
+	virtual void ClearAllCachedCookedPlatformData() override;
 	virtual void PreEditUndo() override;
 	virtual void PostEditUndo() override;
 	virtual void PreEditChange(FProperty* PropertyThatWillChange) override;
@@ -732,6 +769,7 @@ public:
 	LANDSCAPE_API const TArray<UTexture2D*>& GetWeightmapTextures(bool InReturnEditingWeightmap = false) const;
 	LANDSCAPE_API TArray<TObjectPtr<UTexture2D>>& GetWeightmapTextures(const FGuid& InLayerGuid);
 	LANDSCAPE_API const TArray<UTexture2D*>& GetWeightmapTextures(const FGuid& InLayerGuid) const;
+	const TArray<UTexture2D*>& GetRenderedWeightmapTexturesForFeatureLevel(ERHIFeatureLevel::Type FeatureLevel) const;
 
 	LANDSCAPE_API TArray<FWeightmapLayerAllocationInfo>& GetWeightmapLayerAllocations(bool InReturnEditingWeightmap = false);
 	LANDSCAPE_API const TArray<FWeightmapLayerAllocationInfo>& GetWeightmapLayerAllocations(bool InReturnEditingWeightmap = false) const;
@@ -819,20 +857,21 @@ public:
 	const TArray<TObjectPtr<ULandscapeGrassType>>& GetGrassTypes() const { return GrassTypes; }
 
 	/** Temporarily sets the grass type for this component. Any call to UpdateGrassTypes may override what has been set using this method. */
-	void SetGrassTypes(const TArray<TObjectPtr<ULandscapeGrassType>>& InGrassTypes)	{ GrassTypes = InGrassTypes; }
+	void SetGrassTypes(const TArray<TObjectPtr<ULandscapeGrassType>>& InGrassTypes)
+	{
+		GrassTypes = InGrassTypes;
+		InvalidateGrassTypeSummary();
+	}
 	
 	bool MaterialHasGrass() const { return !GetGrassTypes().IsEmpty(); }
-	
-	float GetGrassTypesMaxDiscardDistance() const { return GrassTypesMaxDiscardDistance; }
-	void SetGrassTypesMaxDiscardDistance(const float InGrassTypesMaxDiscardDistance) { GrassTypesMaxDiscardDistance = InGrassTypesMaxDiscardDistance; }
 
-#if WITH_EDITOR
-	/** If the LandscapeMaterial has changed, updates the GrassTypes array. Returns true if the GrassTypes array was updated. */
-	bool UpdateGrassTypes(bool bForceUpdate = false);
+	float GetGrassTypesMaxDiscardDistance() const { return GrassTypeSummary.MaxInstanceDiscardDistance; }
+	void SetGrassTypesMaxDiscardDistance(const float InGrassTypesMaxDiscardDistance) { GrassTypeSummary.MaxInstanceDiscardDistance = InGrassTypesMaxDiscardDistance; GrassTypeSummary.bInvalid = false; }
 
-	/** Recomputes the maximum discard distance across all grass types in the GrassTypes array. */
-	void UpdateGrassTypesMaxDiscardDistance();
+	/** If the LandscapeMaterial has changed, updates the GrassTypes array. Returns true if the GrassTypes array was updated. */
+	LANDSCAPE_API bool UpdateGrassTypes(bool bForceUpdate = false);
 
+#if WITH_EDITOR
 	/** Deletes a layer from this component if it does not contain data, calling DeleteLayerAllocation. */
 	bool DeleteLayerIfAllZero(const FGuid& InEditLayerGuid, const uint8* const TexDataPtr, int32 TexSize, int32 LayerIdx, bool bShouldDirtyPackage);
 
@@ -852,16 +891,15 @@ public:
 	LANDSCAPE_API void ReplaceLayer(ULandscapeLayerInfoObject* FromLayerInfo, ULandscapeLayerInfoObject* ToLayerInfo, FLandscapeEditDataInterface& LandscapeEdit);
 	void ReplaceLayerInternal(ULandscapeLayerInfoObject* FromLayerInfo, ULandscapeLayerInfoObject* ToLayerInfo, FLandscapeEditDataInterface& LandscapeEdit, const FGuid& InEditLayerGUID);
 
-	/** Creates and destroys cooked grass data stored in the map */
-	void RenderGrassMap();
+#endif // WITH_EDITOR
+
+	/** Destroys grass map data stored on the component */
 	void RemoveGrassMap();
 
 	/* Could a grassmap currently be generated, disregarding whether our textures are streamed in? */
 	bool CanRenderGrassMap() const;
 
-	/* Are the textures we need to render a grassmap currently streamed in? */
-	bool AreTexturesStreamedForGrassMapRender() const;
-
+#if WITH_EDITOR
 	/** Computes a hash representing the state of the material and grasstypes used by this component. */
 	LANDSCAPE_API uint32 ComputeGrassMapGenerationHash() const;
 
@@ -910,7 +948,7 @@ public:
 	LANDSCAPE_API void GetGeneratedTexturesAndMaterialInstances(TArray<UObject*>& OutTexturesAndMaterials) const;
 	LANDSCAPE_API TArray<UTexture*> GetGeneratedTextures() const;
 	LANDSCAPE_API TArray<UMaterialInstance*> GetGeneratedMaterialInstances() const;
-#endif
+#endif // WITH_EDITOR
 
 	/** Gets the landscape proxy actor which owns this component */
 	LANDSCAPE_API ALandscapeProxy* GetLandscapeProxy() const;
diff --git a/Engine/Source/Runtime/Landscape/Classes/LandscapeProxy.h b/Engine/Source/Runtime/Landscape/Classes/LandscapeProxy.h
--- a/Engine/Source/Runtime/Landscape/Classes/LandscapeProxy.h
+++ b/Engine/Source/Runtime/Landscape/Classes/LandscapeProxy.h
@@ -65,6 +65,7 @@ namespace UE::Landscape
 
 #if WITH_EDITOR
 LANDSCAPE_API extern bool GLandscapeEditModeActive;
+extern int32 GGrassMapUseRuntimeGeneration;
 
 DECLARE_MULTICAST_DELEGATE_TwoParams(FOnLandscapeProxyComponentDataChanged, ALandscapeProxy*, const FLandscapeProxyComponentDataChangedParams&);
 #endif // WITH_EDITOR
@@ -599,16 +600,28 @@ public:
 	/** Frame offset for tick interval*/
 	uint32 FrameOffsetForTickInterval;
 
-	// Only used outside of the editor (e.g. in cooked builds)
-	// Disables landscape grass processing entirely if no landscape components have landscape grass configured
-	UPROPERTY()
-	bool bHasLandscapeGrass;
-
-	// Only used outside of the editor (e.g. in cooked builds)
+	// Only used outside of the editor (e.g. in cooked builds) - this value is no longer authoritative TODO [chris.tchou] remove
 	// Cached grass max discard distance for all grass types in all landscape components with landscape grass configured
 	UPROPERTY()
 	float GrassTypesMaxDiscardDistance = 0.0f;
 
+	// Non-serialized runtime cache of values derived from the assigned grass types.
+	// Call ALandscapeProxy::UpdateGrassTypeSummary() to update.
+	struct FGrassTypeSummary
+	{
+		// Used to track validity of these values, as it automatically invalidates if you add or remove a component.
+		// If you do both add AND remove, then the add should also trigger UpdateGrassTypes() which will invalidate this cache.
+		// Negative is the invalid state.
+		int32 LandscapeComponentCount = -1;
+
+		bool bHasAnyGrass = true;
+		double MaxInstanceDiscardDistance = DBL_MAX;
+	};
+	FGrassTypeSummary GrassTypeSummary;
+	inline bool IsGrassTypeSummaryValid() { return LandscapeComponents.Num() == GrassTypeSummary.LandscapeComponentCount; }
+	inline void InvalidateGrassTypeSummary() { GrassTypeSummary.LandscapeComponentCount = -1; }
+	void UpdateGrassTypeSummary();
+
 	/**
 	 *	The resolution to cache lighting at, in texels/quad in one axis
 	 *  Total resolution would be changed by StaticLightingResolution*StaticLightingResolution
@@ -794,6 +807,11 @@ public:
 	UPROPERTY(EditAnywhere, Category = Landscape, AdvancedDisplay, meta = (LandscapeInherited))
 	bool bUseCompressedHeightmapStorage = false;
 
+	/** Enable runtime grass data generation to save memory, when grass.GrassMap.UseRuntimeGeneration is true (this is checked per platform).
+		When enabled grass data is not serialized during cook and will be regenerated at runtime when the landscape is loaded. */
+	UPROPERTY(EditAnywhere, Category = Landscape, AdvancedDisplay, meta = (LandscapeInherited))
+	bool bUseRuntimeGrassMapGeneration = false;
+
 	/** Strip Physics/collision components when cooked for client */
 	UPROPERTY(EditAnywhere, Category = Landscape, AdvancedDisplay, meta = (LandscapeOverridable))
 	bool bStripPhysicsWhenCookedClient = false;
@@ -945,12 +963,11 @@ public:
 
 	static void SetGrassUpdateInterval(int32 Interval) { GrassUpdateInterval = Interval; }
 
-	/* Per-frame call to update dynamic grass placement and render grassmaps */
+	/* Determine whether we should update dynamic grass instances this update tick */
 	FORCEINLINE bool ShouldTickGrass() const
 	{
-		// At runtime if we don't have grass we will never have any so avoid ticking it
-		// In editor we might have a material that didn't have grass and now does so we can't rely on bHasLandscapeGrass.
-		if (!GIsEditor && !bHasLandscapeGrass)
+		// don't tick grass if there's no grass to tick
+		if (!GrassTypeSummary.bHasAnyGrass)
 		{
 			return false;
 		}
@@ -967,15 +984,18 @@ public:
 		return true;
 	}
 	void TickGrass(const TArray<FVector>& Cameras, int32& InOutNumCompsCreated);
+	void ProcessAsyncGrassInstanceTasks(bool bWaitAsyncTasks, bool bForceSync, const TSet<UHierarchicalInstancedStaticMeshComponent*>& StillUsed);
 
-	/** Flush the grass cache */
+	/** Flush the grass cache, removing grass instances on the given components (or all proxy components if the component set is not specified).
+	*     bFlushGrassMaps will delete the grass data / density maps on the components as well, but only in editor mode, and only if the grass maps are renderable (i.e. they can be regenerated).
+	*/
 	LANDSCAPE_API void FlushGrassComponents(const TSet<ULandscapeComponent*>* OnlyForComponents = nullptr, bool bFlushGrassMaps = true);
 
 	/**
-		Update Grass 
+		Update Grass -- Builds grass instances given the camera locations
 		* @param Cameras to use for culling, if empty, then NO culling
 		* @param InOutNumComponentsCreated, value can increase if components were created, it is also used internally to limit the number of creations
-		* @param bForceSync if true, block and finish all work
+		* @param bForceSync if true, block and finish all work so that grass is fully populated for the given cameras
 	*/
 	LANDSCAPE_API void UpdateGrass(const TArray<FVector>& Cameras, int32& InOutNumComponentsCreated, bool bForceSync = false);
 	LANDSCAPE_API void UpdateGrass(const TArray<FVector>& Cameras, bool bForceSync = false);
@@ -1011,16 +1031,10 @@ public:
 	LANDSCAPE_API void UpdateRenderingMethod();
 
 #if WITH_EDITOR
-	/** Update Grass maps */
-	void UpdateGrassData(bool bInShouldMarkDirty = false, struct FScopedSlowTask* InSlowTask = nullptr);
-
 	/** Render grass maps for the specified components */
-	UE_DEPRECATED(5.4, "This version of RenderGrassMaps is deprecated to account for landscape components with their own grass types, use the other version instead.")
+	UE_DEPRECATED(5.4, "This version of RenderGrassMaps is deprecated.  Use BuildGrassMaps() instead.")
 	void RenderGrassMaps(const TArray<ULandscapeComponent*>& InLandscapeComponents, const TArray<ULandscapeGrassType*>& InGrassTypes) {}
 
-	/** Render grass maps for the specified components */
-	void RenderGrassMaps(TArrayView<ULandscapeComponent* const> InLandscapeComponents);
-
 	struct UE_DEPRECATED(5.3, "FGIBakedTextureState is officially deprecated now and nothing updates it anymore") FGIBakedTextureState
 	{
 		FGuid CombinedStateId;
@@ -1338,13 +1352,6 @@ public:
 	/** Creates a LandscapeWeightMapUsage object outered to this proxy. */
 	LANDSCAPE_API ULandscapeWeightmapUsage* CreateWeightmapUsage();
 
-	/* For the grassmap rendering notification */
-	int32 NumComponentsNeedingGrassMapRender;
-
-	/* To throttle texture streaming when we're trying to render a grassmap */
-	int32 NumTexturesToStreamForVisibleGrassMapRender;
-	LANDSCAPE_API static int32 TotalTexturesToStreamForVisibleGrassMapRender;
-
 	UE_DEPRECATED(5.3, "NumComponentsNeedingTextureBaking is officially deprecated now and nothing updates it anymore")
 	int32 NumComponentsNeedingTextureBaking;
 
@@ -1456,8 +1463,6 @@ private:
 	FName GenerateUniqueLandscapeTextureName(UObject* InOuter, TextureGroup InLODGroup) const;
 
 #if WITH_EDITOR
-	void UpdateGrassDataStatus(TSet<UTexture2D*>* OutCurrentForcedStreamedTextures, TSet<UTexture2D*>* OutDesiredForcedStreamedTextures, TSet<ULandscapeComponent*>* OutComponentsNeedingGrassMapRender, TSet<ULandscapeComponent*>* OutOutdatedComponents, bool bInEnableForceResidentFlag, int32* OutOutdatedGrassMaps = nullptr) const;
-
 	/** Create Blank Nanite Component */
 	void CreateNaniteComponents(int32 NumComponents);
 
diff --git a/Engine/Source/Runtime/Landscape/Private/Landscape.cpp b/Engine/Source/Runtime/Landscape/Private/Landscape.cpp
--- a/Engine/Source/Runtime/Landscape/Private/Landscape.cpp
+++ b/Engine/Source/Runtime/Landscape/Private/Landscape.cpp
@@ -63,6 +63,7 @@ Landscape.cpp: Terrain rendering
 #include "ComponentRecreateRenderStateContext.h"
 #include "LandscapeWeightmapUsage.h"
 #include "LandscapeSubsystem.h"
+#include "LandscapeGrassMapsBuilder.h"
 #include "LandscapeCulling.h"
 #include "ContentStreaming.h"
 #include "UObject/ObjectSaveContext.h"
@@ -211,6 +212,10 @@ FAutoConsoleVariableRef CVarRenderNaniteLandscape(
 	ECVF_Scalability | ECVF_RenderThreadSafe
 );
 
+extern int32 GGrassEnable;
+extern int32 GGrassMapUseRuntimeGeneration;
+extern FAutoConsoleVariableRef CVarGrassMapUseRuntimeGeneration;
+
 struct FCompareULandscapeComponentClosest
 {
 	FCompareULandscapeComponentClosest(const FIntPoint& InCenter) : Center(InCenter) {}
@@ -325,15 +330,86 @@ void ULandscapeComponent::BeginCacheForCookedPlatformData(const ITargetPlatform*
 {
 	Super::BeginCacheForCookedPlatformData(TargetPlatform);
 
+	// first check if there is a cooked state cached currently, and reset back to non cooked state if so
+	if (PlatformCook.CurrentCookedPlatformOrdinal != -1)
+	{
+		ClearAllCachedCookedPlatformData();
+	}
+
+	// record the platform we are cooking for, so we know what state we are in
+	PlatformCook.CurrentCookedPlatformOrdinal = TargetPlatform->GetPlatformOrdinal();
+	PlatformCook.bStripGrassData = false;
+	PlatformCook.StrippedGrassData = nullptr;
+
 	if (!HasAnyFlags(RF_ClassDefaultObject))
 	{
 		if (TargetPlatform->SupportsFeature(ETargetPlatformFeatures::MobileRendering))
 		{
 			CheckGenerateMobilePlatformData(/*bIsCooking = */ true, TargetPlatform);
 		}
+
+		// determine whether our target platform is going to need serialized grass data
+		bool bNeedSerializedGrassData = false;
+		{
+			TSharedPtr<IConsoleVariable> TargetPlatformUseRuntimeGeneration =
+				CVarGrassMapUseRuntimeGeneration->GetPlatformValueVariable(*TargetPlatform->IniPlatformName());
+			check(TargetPlatformUseRuntimeGeneration.IsValid());
+
+			if (!TargetPlatformUseRuntimeGeneration->GetBool())
+			{
+				bNeedSerializedGrassData = true;
+			}
+		}
+		PlatformCook.bStripGrassData = !bNeedSerializedGrassData;
+
+		if (ALandscapeProxy* Proxy = GetLandscapeProxy())
+		{
+			// Also strip grass data according to Proxy flags (when not cooking for editor)
+			if (!TargetPlatform->AllowsEditorObjects())
+			{
+				if (CVarAllowGrassStripping->GetBool() &&
+					((Proxy->bStripGrassWhenCookedClient && Proxy->bStripGrassWhenCookedServer) ||
+					(Proxy->bStripGrassWhenCookedClient && TargetPlatform->IsClientOnly()) ||
+					(Proxy->bStripGrassWhenCookedServer && TargetPlatform->IsServerOnly())))
+				{
+					PlatformCook.bStripGrassData = true;
+				}
+			}
+		}
+		
+		if (PlatformCook.bStripGrassData)
+		{
+			// save existing value before we strip it, so we can restore later
+			PlatformCook.StrippedGrassData = GrassData;
+			TUniquePtr<FLandscapeComponentGrassData> NewGrassData = MakeUnique<FLandscapeComponentGrassData>();
+			GrassData = MakeShareable(NewGrassData.Release());
+			GrassData->NumElements = 0;
+		}
 	}
 }
 
+bool ULandscapeComponent::IsCachedCookedPlatformDataLoaded(const ITargetPlatform* TargetPlatform)
+{
+	return (PlatformCook.CurrentCookedPlatformOrdinal == TargetPlatform->GetPlatformOrdinal());
+}
+
+void ULandscapeComponent::ClearAllCachedCookedPlatformData()
+{
+	// restore back to "non-cooked" state
+	if (PlatformCook.CurrentCookedPlatformOrdinal >= 0 && PlatformCook.bStripGrassData)
+	{
+		if (PlatformCook.StrippedGrassData.IsValid())
+		{
+			GrassData = PlatformCook.StrippedGrassData.ToSharedRef();
+		}
+	}
+
+	PlatformCook.CurrentCookedPlatformOrdinal = -1;
+	PlatformCook.bStripGrassData = false;
+	PlatformCook.StrippedGrassData = nullptr;
+}
+
+
 void ALandscapeProxy::CheckGenerateMobilePlatformData(bool bIsCooking, const ITargetPlatform* TargetPlatform)
 {
 	for (ULandscapeComponent* Component : LandscapeComponents)
@@ -1362,7 +1438,6 @@ void ULandscapeComponent::PostLoad()
 	{
 		UpdateGrassTypes();
 	}
-	UpdateGrassTypesMaxDiscardDistance();
 
 #if !UE_BUILD_SHIPPING
 	if (MobileCombinationMaterialInstances.Num() == 0)
@@ -1437,7 +1512,6 @@ ALandscapeProxy::ALandscapeProxy(const FObjectInitializer& ObjectInitializer)
 #if WITH_EDITORONLY_DATA
 	, TargetDisplayOrder(ELandscapeLayerDisplayMode::Default)
 #endif // WITH_EDITORONLY_DATA
-	, bHasLandscapeGrass(true)
 {
 	bReplicates = false;
 	NetUpdateFrequency = 10.0f;
@@ -1478,9 +1552,6 @@ ALandscapeProxy::ALandscapeProxy(const FObjectInitializer& ObjectInitializer)
 #endif
 
 #if WITH_EDITOR
-	NumComponentsNeedingGrassMapRender = 0;
-	NumTexturesToStreamForVisibleGrassMapRender = 0;
-
 	if (VisibilityLayer == nullptr)
 	{
 		// Structure to hold one-time initialization
@@ -1987,11 +2058,14 @@ void ULandscapeComponent::OnRegister()
 		UWorld* World = GetLandscapeProxy()->GetWorld();
 		if (World)
 		{
-			ULandscapeInfo* Info = GetLandscapeInfo();
-			if (Info)
+			if (ULandscapeInfo* Info = GetLandscapeInfo())
 			{
 				Info->RegisterActorComponent(this);
 			}
+			if (ULandscapeSubsystem *Subsystem = World->GetSubsystem<ULandscapeSubsystem>())
+			{
+				Subsystem->RegisterComponent(this);
+			}
 		}
 	}
 }
@@ -2017,11 +2091,14 @@ void ULandscapeComponent::OnUnregister()
 
 		if (World)
 		{
-			ULandscapeInfo* Info = GetLandscapeInfo();
-			if (Info)
+			if (ULandscapeInfo* Info = GetLandscapeInfo())
 			{
 				Info->UnregisterActorComponent(this);
 			}
+			if (ULandscapeSubsystem* Subsystem = World->GetSubsystem<ULandscapeSubsystem>())
+			{
+				Subsystem->UnregisterComponent(this);
+			}
 		}
 	}
 }
@@ -2116,6 +2193,18 @@ TArray<TObjectPtr<UTexture2D>>& ULandscapeComponent::GetWeightmapTextures(const
 	return WeightmapTextures;
 }
 
+const TArray<UTexture2D*>& ULandscapeComponent::GetRenderedWeightmapTexturesForFeatureLevel(ERHIFeatureLevel::Type FeatureLevel) const
+{
+	if (FeatureLevel == ERHIFeatureLevel::ES3_1)
+	{
+		return MobileWeightmapTextures;
+	}
+	else
+	{
+		return WeightmapTextures;
+	}
+}
+
 const TArray<FWeightmapLayerAllocationInfo>& ULandscapeComponent::GetWeightmapLayerAllocations(bool InReturnEditingWeightmap) const
 {
 #if WITH_EDITORONLY_DATA
@@ -2695,15 +2784,18 @@ void ALandscape::PostInitProperties()
 				{
 					// Create conditional GrassRenderingNotification
 					GrassRenderingNotification = MakeShared<FLandscapeNotification>(this, FLandscapeNotification::EType::GrassRendering,
-						[this]() { return NumComponentsNeedingGrassMapRender > 0 && !(FSlateApplication::Get().HasAnyMouseCaptor() || GUnrealEd->IsUserInteracting()); },
-						[this](FText& InText)
+						// display condition
+						[this, LandscapeSubSystem]()
+						{
+							return GGrassEnable && !GGrassMapUseRuntimeGeneration &&
+								(LandscapeSubSystem->GetGrassMapBuilder()->GetTotalGrassMapsWaitingToRender() > 0) &&
+								!(FSlateApplication::Get().HasAnyMouseCaptor() || GUnrealEd->IsUserInteracting());
+						},
+						// retrieve text
+						[this, LandscapeSubSystem](FText& InText)
 						{
 							// Accumulating all outstanding grass maps that need to be rendered for ALL landscapes because only one such notification can be displayed at a time
-							int ComponentsNeedingGrassMapRenderAllLandscapes = 0;
-							for (const ALandscape* Landscape : TObjectRange<ALandscape>(RF_ClassDefaultObject | RF_ArchetypeObject, true, EInternalObjectFlags::Garbage))
-							{
-								ComponentsNeedingGrassMapRenderAllLandscapes += Landscape->NumComponentsNeedingGrassMapRender;
-							}
+							int ComponentsNeedingGrassMapRenderAllLandscapes = LandscapeSubSystem->GetGrassMapBuilder()->GetTotalGrassMapsWaitingToRender();
 							
 							FFormatNamedArguments Args;
 							Args.Add(TEXT("OutstandingGrassMaps"), FText::AsNumber(ComponentsNeedingGrassMapRenderAllLandscapes));
@@ -3498,16 +3590,24 @@ void ALandscapeProxy::PreSave(FObjectPreSaveContext ObjectSaveContext)
 	Super::PreSave(ObjectSaveContext);
 
 #if WITH_EDITOR
-	// Work out whether we have grass or not for the next game run
-	BuildGrassMaps();
-
-	for (ULandscapeComponent* Component : LandscapeComponents)
+	if (!ObjectSaveContext.IsProceduralSave()) // only finalize grass in a true editor save (when a GPU is available).
 	{
-		// Reset flag
-		Component->GrassData->bIsDirty = false;
+		// It would be nice to strip grass data at editor save time to reduce asset size on disk.
+		// Unfortunately we can't easily know if there is a platform out there that may need to use the serialized grass map path.
+		// And future cook processes may not have a GPU available to build the grass data themselves.
+		// So for now, we always build all grass maps on editor save, just in case.
+		// The grass maps will get stripped later in BeginCacheForCookedPlatformData for cooked builds that don't need them.
+		{
+			// generate all of the grass data
+			BuildGrassMaps();
+			for (ULandscapeComponent* Component : LandscapeComponents)
+			{
+				// Manually reset dirty flag (for post save)
+				Component->GrassData->bIsDirty = false;
+			}
+		}
 	}
 
-	
 	if (ULandscapeInfo* LandscapeInfo = GetLandscapeInfo())
 	{
 		LandscapeInfo->UpdateNanite(ObjectSaveContext.GetTargetPlatform());
@@ -4100,6 +4200,23 @@ void ALandscapeProxy::PostLoad()
 			continue;
 		}
 
+#if !WITH_EDITOR
+		// if using runtime grass gen, it should have been cleared out in PreSave
+		if (GGrassMapUseRuntimeGeneration)
+		{
+			if (Comp->GrassData->HasValidData() && bUseRuntimeGrassMapGeneration)
+			{
+				UE_LOG(LogGrass, Warning, TEXT("grass.GrassMap.UseRuntimeGeneration is enabled, but component %s on landscape %s has unnecessary grass data saved.  Ensure grass.GrassMap.UseRuntimeGeneration is enabled at cook time to reduce cooked data size."),
+					*Comp->GetName(),
+					*GetName());
+
+				// Free the memory, so at least we will save the space at runtime.
+				TUniquePtr<FLandscapeComponentGrassData> NewGrassData = MakeUnique<FLandscapeComponentGrassData>();
+				Comp->GrassData = MakeShareable(NewGrassData.Release());
+			}
+		}
+#endif // !WITH_EDITOR
+
 		// Validate the layer combination and store it in the MaterialInstanceConstantMap
 		UMaterialInstance* MaterialInstance = Comp->GetMaterialInstance(0, false);
 
@@ -4210,10 +4327,6 @@ void ALandscapeProxy::Destroyed()
 		{
 			SplineComponent->ModifySplines();
 		}
-
-		NumComponentsNeedingGrassMapRender = 0;
-		TotalTexturesToStreamForVisibleGrassMapRender -= NumTexturesToStreamForVisibleGrassMapRender;
-		NumTexturesToStreamForVisibleGrassMapRender = 0;
 	}
 
 	// Destroy the Nanite component when we get destroyed so that we don't restore a garbage Nanite component (it's non-transactional and will get regenerated anyway)
@@ -5996,12 +6109,6 @@ ALandscapeProxy::~ALandscapeProxy()
 	}
 	AsyncFoliageTasks.Empty();
 
-#if WITH_EDITOR
-	NumComponentsNeedingGrassMapRender = 0;
-	TotalTexturesToStreamForVisibleGrassMapRender -= NumTexturesToStreamForVisibleGrassMapRender;
-	NumTexturesToStreamForVisibleGrassMapRender = 0;
-#endif
-
 #if WITH_EDITORONLY_DATA
 	LandscapeProxies.Remove(this);
 #endif
diff --git a/Engine/Source/Runtime/Landscape/Private/LandscapeCollision.cpp b/Engine/Source/Runtime/Landscape/Private/LandscapeCollision.cpp
--- a/Engine/Source/Runtime/Landscape/Private/LandscapeCollision.cpp
+++ b/Engine/Source/Runtime/Landscape/Private/LandscapeCollision.cpp
@@ -20,6 +20,7 @@
 #include "CollisionQueryParams.h"
 #include "Engine/World.h"
 #include "LandscapeSubsystem.h"
+#include "LandscapeGrassMapsBuilder.h"
 #include "LandscapeRender.h"
 #include "LandscapeProxy.h"
 #include "LandscapeInfo.h"
@@ -2517,7 +2518,10 @@ void ULandscapeHeightfieldCollisionComponent::PreSave(FObjectPreSaveContext Obje
 				{
 					RenderComponent->GetMaterialInstance(0, false)->GetMaterialResource(GetWorld()->GetFeatureLevel())->FinishCompilation();
 				}
-				RenderComponent->RenderGrassMap();
+
+				ULandscapeSubsystem* LandscapeSubsystem = GetWorld()->GetSubsystem<ULandscapeSubsystem>();
+				TArray<TObjectPtr<ULandscapeComponent>> Components = { RenderComponent };
+				LandscapeSubsystem->GetGrassMapBuilder()->BuildGrassMapsNowForComponents(Components, /* SlowTask= */ nullptr, /* bMarkDirty= */ false);
 			}
 		}
 #endif// WITH_EDITOR
@@ -2527,7 +2531,6 @@ void ULandscapeHeightfieldCollisionComponent::PreSave(FObjectPreSaveContext Obje
 #if WITH_EDITOR
 void ULandscapeInfo::UpdateAllAddCollisions()
 {
-
 	TRACE_CPUPROFILER_EVENT_SCOPE(ULandscapeInfo::UpdateAllAddCollisions);
 	XYtoAddCollisionMap.Reset();
 
diff --git a/Engine/Source/Runtime/Landscape/Private/LandscapeEdit.cpp b/Engine/Source/Runtime/Landscape/Private/LandscapeEdit.cpp
--- a/Engine/Source/Runtime/Landscape/Private/LandscapeEdit.cpp
+++ b/Engine/Source/Runtime/Landscape/Private/LandscapeEdit.cpp
@@ -83,6 +83,9 @@ LandscapeEdit.cpp: Landscape editing
 #include "ActorPartition/ActorPartitionSubsystem.h"
 #include "LandscapeUtils.h"
 #include "LandscapeSplineActor.h"
+#include "Materials/MaterialExpressionLandscapeGrassOutput.h"
+#include "ShaderPlatformCachedIniValue.h"
+#include "DataDrivenShaderPlatformInfo.h"
 #endif
 #include "Algo/Count.h"
 #include "Algo/Transform.h"
@@ -7293,9 +7296,30 @@ static void GetAllMobileRelevantLayerNames(TSet<FName>& OutLayerNames, UMaterial
 {
 	TArray<FName> LayerNames;
 
+	bool bMobileUseRuntimeGrassMapGeneration = false;
+	{
+		// if ANY mobile platform has runtime generation (because we don't calculate mobile weightmaps per platform, just mobile/non-mobile)
+		static FShaderPlatformCachedIniValue<bool> UseRuntimeGenerationCVar(TEXT("grass.GrassMaps.UseRuntimeGeneration"));
+		for (int32 SPIndex = 0; SPIndex < SP_NumPlatforms; SPIndex++)
+		{
+			EShaderPlatform SP = static_cast<EShaderPlatform>(SPIndex);
+			if (FGenericDataDrivenShaderPlatformInfo::IsValid(SP) && IsMobilePlatform(SP) && UseRuntimeGenerationCVar.Get(SP))
+			{
+				bMobileUseRuntimeGrassMapGeneration = true;
+				break;
+			}
+		}
+	}
+
+	TSet<UClass*> MobileCustomOutputExpressionTypesToQuery;
+	if (bMobileUseRuntimeGrassMapGeneration)
+	{
+		MobileCustomOutputExpressionTypesToQuery.Add(UMaterialExpressionLandscapeGrassOutput::StaticClass());
+	}
+
 	const bool bRecurseIntoMaterialFunctions = true;
 	TArray<UMaterialExpression*> ES31Expressions;
-	InMaterial->GetAllReferencedExpressions(ES31Expressions, nullptr, ERHIFeatureLevel::ES3_1, EMaterialQualityLevel::Num, ERHIShadingPath::Num, bRecurseIntoMaterialFunctions );
+	InMaterial->GetAllReferencedExpressions(ES31Expressions, nullptr, ERHIFeatureLevel::ES3_1, EMaterialQualityLevel::Num, ERHIShadingPath::Num, bRecurseIntoMaterialFunctions, &MobileCustomOutputExpressionTypesToQuery);
 
 	TArray<UMaterialExpression*> MobileExpressions = MoveTemp(ES31Expressions);
 	for (UMaterialExpression* Expression : MobileExpressions)
diff --git a/Engine/Source/Runtime/Landscape/Private/LandscapeEditLayers.cpp b/Engine/Source/Runtime/Landscape/Private/LandscapeEditLayers.cpp
--- a/Engine/Source/Runtime/Landscape/Private/LandscapeEditLayers.cpp
+++ b/Engine/Source/Runtime/Landscape/Private/LandscapeEditLayers.cpp
@@ -18,6 +18,7 @@ LandscapeEditLayers.cpp: Landscape editing layers mode
 #include "LandscapeEditTypes.h"
 #include "LandscapeUtils.h"
 #include "LandscapeSubsystem.h"
+#include "LandscapeTextureStreamingManager.h"
 
 #include "Application/SlateApplicationBase.h"
 #include "Shader.h"
@@ -1745,18 +1746,6 @@ typedef FLandscapeLayersRender_RenderThread<FLandscapeLayersWeightmapShaderParam
 
 #if WITH_EDITOR
 
-bool ALandscape::IsTextureReady(UTexture2D* InTexture, bool bInWaitForStreaming)
-{
-	check(InTexture);
-	InTexture->bForceMiplevelsToBeResident = true;
-	if (bInWaitForStreaming && !InTexture->IsFullyStreamedIn())
-	{
-		TRACE_CPUPROFILER_EVENT_SCOPE(Landscape_WaitForTextureStreaming);
-		InTexture->WaitForStreaming();
-	}
-	return !InTexture->IsDefaultTexture() && InTexture->IsFullyStreamedIn();
-}
-
 bool ALandscape::IsMaterialResourceCompiled(FMaterialResource* InMaterialResource, bool bInWaitForCompilation)
 {
 	check(InMaterialResource);
@@ -3460,6 +3449,8 @@ bool ALandscape::PrepareTextureResources(bool bInWaitForStreaming)
 	// All components containing heightmaps that have just completed streaming in (filled out below)
 	TSet<ULandscapeComponent*> StreamedInHeightmapComponents;
 
+	FLandscapeTextureStreamingManager* TextureStreamingManager = GetWorld()->GetSubsystem<ULandscapeSubsystem>()->GetTextureStreamingManager();
+
 	bool bIsReady = true;
 	Info->ForEachLandscapeProxy([&](ALandscapeProxy* Proxy)
 	{
@@ -3468,7 +3459,7 @@ bool ALandscape::PrepareTextureResources(bool bInWaitForStreaming)
 			UTexture2D* ComponentHeightmap = Component->GetHeightmap();
 
 			{
-				bool bIsTextureReady = IsTextureReady(ComponentHeightmap, bInWaitForStreaming);
+				bool bIsTextureReady = TextureStreamingManager->RequestTextureFullyStreamedInForever(ComponentHeightmap, bInWaitForStreaming);
 				if (!bIsTextureReady)
 				{
 					StreamingInTexturesAfter.Add(ComponentHeightmap);
@@ -3487,7 +3478,7 @@ bool ALandscape::PrepareTextureResources(bool bInWaitForStreaming)
 
 			for (UTexture2D* ComponentWeightmap : Component->GetWeightmapTextures())
 			{
-				bool bIsTextureReady = IsTextureReady(ComponentWeightmap, bInWaitForStreaming);
+				bool bIsTextureReady = TextureStreamingManager->RequestTextureFullyStreamedInForever(ComponentWeightmap, bInWaitForStreaming);
 				// If the texture is not ready, start tracking its state changes to be notified when it's fully streamed in : 
 				if (!bIsTextureReady)
 				{
@@ -3629,8 +3620,10 @@ bool ALandscape::PrepareLayersTextureResources(const TArray<FLandscapeLayer>& In
 		return false;
 	}
 
+	FLandscapeTextureStreamingManager* TextureStreamingManager = GetWorld()->GetSubsystem<ULandscapeSubsystem>()->GetTextureStreamingManager();
+
 	bool bIsReady = true;
-	Info->ForEachLandscapeProxy([&](ALandscapeProxy* Proxy)
+	Info->ForEachLandscapeProxy([&, TextureStreamingManager](ALandscapeProxy* Proxy)
 	{
 		for (const FLandscapeLayer& Layer : InLayers)
 		{
@@ -3638,11 +3631,11 @@ bool ALandscape::PrepareLayersTextureResources(const TArray<FLandscapeLayer>& In
 			{
 				if (FLandscapeLayerComponentData* ComponentLayerData = Component->GetLayerData(Layer.Guid))
 				{
-					bIsReady &= IsTextureReady(ComponentLayerData->HeightmapData.Texture, bInWaitForStreaming);
+					bIsReady &= TextureStreamingManager->RequestTextureFullyStreamedInForever(ComponentLayerData->HeightmapData.Texture, bInWaitForStreaming);
 
 					for (UTexture2D* LayerWeightmap : ComponentLayerData->WeightmapData.Textures)
 					{
-						bIsReady &= IsTextureReady(LayerWeightmap, bInWaitForStreaming);
+						bIsReady &= TextureStreamingManager->RequestTextureFullyStreamedInForever(LayerWeightmap, bInWaitForStreaming);
 					}
 				}
 			}
@@ -3671,13 +3664,15 @@ bool ALandscape::PrepareLayersBrushResources(ERHIFeatureLevel::Type InFeatureLev
 		}
 	}
 
+	FLandscapeTextureStreamingManager* TextureStreamingManager = GetWorld()->GetSubsystem<ULandscapeSubsystem>()->GetTextureStreamingManager();
+
 	bool bIsReady = true;
 	for (UObject* Dependency : Dependencies)
 	{
 		// Streamable textures need to be fully streamed in : 
 		if (UTexture2D* Texture = Cast<UTexture2D>(Dependency))
 		{
-			bIsReady &= IsTextureReady(Texture, bInWaitForStreaming);
+			bIsReady &= TextureStreamingManager->RequestTextureFullyStreamedInForever(Texture, bInWaitForStreaming);
 		}
 
 		// Material shaders need to be fully compiled : 
@@ -5678,6 +5673,10 @@ bool ALandscape::ResolveLayersTexture(
 			}
 		}
 
+		// change lighting guid to be the hash of the source data (so we can use lighting guid to detect when it actually changes)
+		// TODO [chris.tchou] we should check(InOutputTexture->Source.bGuidIsHash); .. but it's not a public value currently.
+ 		InOutputTexture->SetLightingGuid(InOutputTexture->Source.GetId());
+
 		// Find out whether some channels from this weightmap are now all zeros : 
 		static constexpr uint32 AllChannelsAllZerosMask = 15;
 		uint32 AllZerosTextureChannelMask = AllChannelsAllZerosMask;
@@ -6081,10 +6080,10 @@ void ALandscape::ReallocateLayersWeightmaps(FUpdateLayersContentContext& InUpdat
 	//GDisableAutomaticTextureMaterialUpdateDependencies = true;
 
 	FTextureCompilingManager::Get().FinishCompilation(NewCreatedTextures);
+	FLandscapeTextureStreamingManager* TextureStreamingManager = GetWorld()->GetSubsystem<ULandscapeSubsystem>()->GetTextureStreamingManager();
 	for (UTexture* Texture : NewCreatedTextures)
 	{
-		Texture->bForceMiplevelsToBeResident = true;
-		Texture->WaitForStreaming();
+		TextureStreamingManager->RequestTextureFullyStreamedInForever(Texture, /* bWaitForStreaming= */ true);
 	}
 
 	//GDisableAutomaticTextureMaterialUpdateDependencies = false;
diff --git a/Engine/Source/Runtime/Landscape/Private/LandscapeGrass.cpp b/Engine/Source/Runtime/Landscape/Private/LandscapeGrass.cpp
--- a/Engine/Source/Runtime/Landscape/Private/LandscapeGrass.cpp
+++ b/Engine/Source/Runtime/Landscape/Private/LandscapeGrass.cpp
@@ -26,6 +26,7 @@
 #include "SceneView.h"
 #include "Shader.h"
 #include "Landscape.h"
+#include "LandscapePrivate.h"
 #include "LandscapeProxy.h"
 #include "LandscapeInfo.h"
 #include "LightMap.h"
@@ -33,6 +34,7 @@
 #include "ShadowMap.h"
 #include "LandscapeComponent.h"
 #include "LandscapeSubsystem.h"
+#include "LandscapeGrassMapsBuilder.h"
 #include "LandscapeVersion.h"
 #include "MaterialShaderType.h"
 #include "MeshMaterialShaderType.h"
@@ -78,7 +80,7 @@
 
 #define LOCTEXT_NAMESPACE "Landscape"
 
-DEFINE_LOG_CATEGORY_STATIC(LogGrass, Log, All);
+DEFINE_LOG_CATEGORY(LogGrass);
 
 static float GGuardBandMultiplier = 1.3f;
 static FAutoConsoleVariableRef CVarGuardBandMultiplier(
@@ -129,14 +131,14 @@ static FAutoConsoleVariableRef CVarGrassDensityScale(
 	TEXT("Multiplier on all grass densities."),
 	ECVF_Scalability);
 
-static float GGrassCullDistanceScale = 1;
+float GGrassCullDistanceScale = 1;
 static FAutoConsoleVariableRef CVarGrassCullDistanceScale(
 	TEXT("grass.CullDistanceScale"),
 	GGrassCullDistanceScale,
 	TEXT("Multiplier on all grass cull distances."),
 	ECVF_Scalability);
 
-static int32 GGrassEnable = 1;
+int32 GGrassEnable = 1;
 static FAutoConsoleVariableRef CVarGrassEnable(
 	TEXT("grass.Enable"),
 	GGrassEnable,
@@ -189,7 +191,7 @@ static int32 GGrassCreationPrioritizedMultipler = 4;
 static FAutoConsoleVariableRef CVarGrassCreationPrioritizedMultipler(
 	TEXT("grass.GrassCreationPrioritizedMultipler"),
 	GGrassCreationPrioritizedMultipler,
-	TEXT("Multipler applied to MaxCreatePerFrame and MaxAsyncTasks when grass creation is prioritized."));
+	TEXT("Multiplier applied to MaxCreatePerFrame and MaxAsyncTasks when grass creation is prioritized."));
 
 static int32 GGrassUpdateAllOnRebuild = 0;
 static FAutoConsoleVariableRef CVarUpdateAllOnRebuild(
@@ -197,13 +199,6 @@ static FAutoConsoleVariableRef CVarUpdateAllOnRebuild(
 	GGrassUpdateAllOnRebuild,
 	TEXT(""));
 
-
-static int32 GCaptureNextGrassUpdate = 0;
-static FAutoConsoleVariableRef CVarCaptureNextGrassUpdate(
-	TEXT("grass.CaptureNextGrassUpdate"),
-	GCaptureNextGrassUpdate,
-	TEXT("Trigger a renderdoc capture for the next X grass updates (calls to RenderGrassMap or RenderGrassMaps"));
-
 #if RHI_RAYTRACING
 static TAutoConsoleVariable<int32> CVarRayTracingLandscapeGrass(
 	TEXT("r.RayTracing.Geometry.LandscapeGrass"),
@@ -810,6 +805,7 @@ DECLARE_CYCLE_STAT(TEXT("Grass Start Comp"), STAT_FoliageGrassStartComp, STATGRO
 DECLARE_CYCLE_STAT(TEXT("Grass End Comp"), STAT_FoliageGrassEndComp, STATGROUP_Foliage);
 DECLARE_CYCLE_STAT(TEXT("Grass Destroy Comps"), STAT_FoliageGrassDestoryComp, STATGROUP_Foliage);
 DECLARE_CYCLE_STAT(TEXT("Grass Update"), STAT_GrassUpdate, STATGROUP_Foliage);
+DECLARE_CYCLE_STAT(TEXT("Grass Type Update Summary"), STAT_UpdateGrassTypeSummary, STATGROUP_Foliage);
 
 int32 ALandscapeProxy::GrassUpdateInterval = 1;
 
@@ -863,71 +859,59 @@ static FAutoConsoleVariableSink CVarGrassSink(FConsoleCommandDelegate::CreateSta
 // Grass weightmap rendering
 //
 
+FLandscapeComponentGrassData::FLandscapeComponentGrassData(ULandscapeComponent* Component)
+{
 #if WITH_EDITOR
+	GenerationHash = Component->ComputeGrassMapGenerationHash();
+#endif // WITH_EDITOR
+}
 
-FLandscapeComponentGrassData::FLandscapeComponentGrassData(ULandscapeComponent* Component)
-	: GenerationHash(Component->ComputeGrassMapGenerationHash())
+void ULandscapeComponent::InvalidateGrassTypeSummary()
 {
+	GrassTypeSummary.bInvalid = true;
+	GetLandscapeProxy()->InvalidateGrassTypeSummary();
 }
 
 bool ULandscapeComponent::UpdateGrassTypes(bool bForceUpdate)
 {
 	bool bChanged = false;
+#if WITH_EDITOR
 	if (UMaterialInterface* Material = GetLandscapeMaterial())
 	{
 		uint32 CurMaterialAllStateCRC = Material->ComputeAllStateCRC();
 		if (bForceUpdate || (LastLandscapeMaterialAllStateCRCWhenGrassTypesBuilt != CurMaterialAllStateCRC))
 		{
 			LastLandscapeMaterialAllStateCRCWhenGrassTypesBuilt = CurMaterialAllStateCRC;
-			bChanged = true;
 
-			GrassTypes = Material->GetMaterial()->GetCachedExpressionData().GrassTypes;
-			UpdateGrassTypesMaxDiscardDistance();
+			SetGrassTypes(Material->GetMaterial()->GetCachedExpressionData().GrassTypes);
+			bChanged = true;
 		}
 	}
-	return bChanged;
-}
-
-void ULandscapeComponent::UpdateGrassTypesMaxDiscardDistance()
-{
-	GrassTypesMaxDiscardDistance = 0.0f;
-	for (const ULandscapeGrassType* GrassType : GrassTypes)
+#else
+	// In cooked builds, we don't automatically respond to material changes.
+	// The cooked value of the component's GrassTypes array is considered authoritative.
+	if (bForceUpdate)
 	{
-		if (GrassType != nullptr)
+		if (UMaterialInterface* Material = GetLandscapeMaterial())
 		{
-			for (const FGrassVariety& GrassVariety : GrassType->GrassVarieties)
-			{
-				GrassTypesMaxDiscardDistance = FMath::Max((float)GrassVariety.EndCullDistance.GetValue(), GrassTypesMaxDiscardDistance);
-			}
+			SetGrassTypes(Material->GetMaterial()->GetCachedExpressionData().GrassTypes);
+			bChanged = true;
 		}
 	}
+#endif // WITH_EDITOR
+	return bChanged;
 }
 
-uint32 ULandscapeComponent::ComputeGrassMapGenerationHash() const
-{
-	TRACE_CPUPROFILER_EVENT_SCOPE(ULandscapeComponent::ComputeGrassMapGenerationHash);
-	uint32 Hash = 0;
-
-	if (UMaterialInterface* Material = GetLandscapeMaterial())
-	{
-		// Bump the generation hash key to invalidate all cached grass density maps.
-		static FGuid GrassGenerationHashKey("216D95C7651D4095ADC6A8459B4F181D");
-
-		Hash = GetTypeHash(GrassGenerationHashKey);
-		// Take into account any material state change : (excluding texture state)
-		Hash = FCrc::TypeCrc32(Material->ComputeAllStateCRC(), Hash);
+#if WITH_EDITOR
 
-		// If anything changes in the grass types, we should take that into account as well :
-		for (ULandscapeGrassType* GrassType : GrassTypes)
-		{
-			if (GrassType != nullptr)
-			{
-				Hash = FCrc::TypeCrc32(GrassType->StateHash, Hash);
-			}
-		}
-	}
+namespace UE::Landscape
+{
+	extern uint32 ComputeGrassMapGenerationHash(const ULandscapeComponent* Component, UMaterialInterface* Material);
+}
 
-	return Hash;
+uint32 ULandscapeComponent::ComputeGrassMapGenerationHash() const
+{
+	return UE::Landscape::ComputeGrassMapGenerationHash(this, GetLandscapeMaterial());
 }
 
 bool ULandscapeComponent::IsGrassMapOutdated() const
@@ -966,59 +950,6 @@ bool ULandscapeComponent::CanRenderGrassMap() const
 	return true;
 }
 
-static bool IsTextureStreamedForGrassMapRender(UTexture2D* InTexture)
-{
-	return InTexture && !InTexture->IsDefaultTexture() && !InTexture->HasPendingInitOrStreaming() && InTexture->IsFullyStreamedIn();
-}
-
-bool ULandscapeComponent::AreTexturesStreamedForGrassMapRender() const
-{
-	// Check for valid heightmap that is fully streamed in
-	if (!IsTextureStreamedForGrassMapRender(HeightmapTexture))
-	{
-		return false;
-	}
-	
-	// Check for valid weightmaps that is fully streamed in
-	for (auto WeightmapTexture : WeightmapTextures)
-	{
-		if (!IsTextureStreamedForGrassMapRender(WeightmapTexture))
-		{
-			return false;
-		}
-	}
-
-	return true;
-}
-
-void ULandscapeComponent::RenderGrassMap()
-{
-	UMaterialInterface* Material = GetLandscapeMaterial();
-	if (ensure(CanRenderGrassMap()))
-	{
-		const bool bBakeMaterialPositionOffsetIntoCollision = (GetLandscapeProxy() && GetLandscapeProxy()->bBakeMaterialPositionOffsetIntoCollision);
-
-		TArray<int32> HeightMips;
-		if (bBakeMaterialPositionOffsetIntoCollision)
-		{
-			if (CollisionMipLevel > 0)
-			{
-				HeightMips.Add(CollisionMipLevel);
-			}
-			if (SimpleCollisionMipLevel > CollisionMipLevel)
-			{
-				HeightMips.Add(SimpleCollisionMipLevel);
-			}
-		}
-
-		if (GrassTypes.Num() > 0 || bBakeMaterialPositionOffsetIntoCollision)
-		{
-			FLandscapeGrassWeightExporter Exporter(GetLandscapeProxy(), { this }, /*bInNeedsGrassmap = */ true, /*bInNeedsHeightmap =*/ true, MoveTemp(HeightMips));
-			Exporter.ApplyResults();
-		}
-	}
-}
-
 TArray<uint16> ULandscapeComponent::RenderWPOHeightmap(int32 LOD)
 {
 	TArray<uint16> Results;
@@ -1050,37 +981,67 @@ TArray<uint16> ULandscapeComponent::RenderWPOHeightmap(int32 LOD)
 
 	return Results;
 }
+#endif // WITH_EDITOR
 
 void ULandscapeComponent::RemoveGrassMap()
 {
 	*GrassData = FLandscapeComponentGrassData();
 
+#if WITH_EDITOR
 	GrassData->bIsDirty = true;
-	if (!MaterialHasGrass())
+#endif // WITH_EDITOR
+
+	if (GrassTypes.IsEmpty())
 	{
 		// Mark grass data as valid but empty if it just doesn't support grass because nothing will ever trigger a RenderGrassMaps on it, which would leave NumElements unset (which is considered as invalid data) :
+		// TODO [chris.tchou]: this is not technically true with bake position offset into collision...
 		GrassData->NumElements = 0;
 	}
 }
 
-void ALandscapeProxy::RenderGrassMaps(TArrayView<ULandscapeComponent* const> InLandscapeComponents)
+void ALandscapeProxy::UpdateGrassTypeSummary()
 {
-	TArray<int32> HeightMips;
-	if (CollisionMipLevel > 0)
-	{
-		HeightMips.Add(CollisionMipLevel);
-	}
-	if (SimpleCollisionMipLevel > CollisionMipLevel)
+	SCOPE_CYCLE_COUNTER(STAT_UpdateGrassTypeSummary);
+
+	bool bGrassQualityLevelEnabled = GEngine && GEngine->UseGrassVarityPerQualityLevels;
+
+	bool bAnyGrass = false;
+	double ProxyMaxGrassCullDistance = 0.0;
+	for (ULandscapeComponent* Component : LandscapeComponents)
 	{
-		HeightMips.Add(SimpleCollisionMipLevel);
+		double ComponentMaxGrassCullDistance;
+		if (Component->GrassTypeSummary.bInvalid)
+		{
+			// Update the component cache
+			ComponentMaxGrassCullDistance = 0.0;
+			for (ULandscapeGrassType* GrassType : Component->GetGrassTypes())
+			{
+				if (GrassType)
+				{
+					for (FGrassVariety& GrassVariety : GrassType->GrassVarieties)
+					{
+						int32 EndCullDistance = bGrassQualityLevelEnabled ? GrassVariety.EndCullDistanceQuality.GetValue(GGrassQualityLevel) : GrassVariety.EndCullDistance.GetValue();
+						ComponentMaxGrassCullDistance = FMath::Max(ComponentMaxGrassCullDistance, EndCullDistance);
+					}
+				}
+			}
+			Component->GrassTypeSummary.MaxInstanceDiscardDistance = ComponentMaxGrassCullDistance;
+			Component->GrassTypeSummary.bInvalid = false;
+		}
+		else
+		{
+			ComponentMaxGrassCullDistance = Component->GrassTypeSummary.MaxInstanceDiscardDistance;
+		}
+
+		bAnyGrass |= (Component->GetGrassTypes().Num() > 0);
+		ProxyMaxGrassCullDistance = FMath::Max(ProxyMaxGrassCullDistance, ComponentMaxGrassCullDistance);
 	}
 
-	FLandscapeGrassWeightExporter Exporter(this, InLandscapeComponents, /*bInNeedsGrassmap = */ true, /*bInNeedsHeightmap =*/ true, MoveTemp(HeightMips));
-	Exporter.ApplyResults();
+	GrassTypeSummary.bHasAnyGrass = bAnyGrass;
+	GrassTypeSummary.MaxInstanceDiscardDistance = ProxyMaxGrassCullDistance;
+	GrassTypeSummary.LandscapeComponentCount = LandscapeComponents.Num();
 }
 
-#endif //WITH_EDITOR
-
 // the purpose of this class is to copy the lightmap from the terrain, and set the CoordinateScale and CoordinateBias to zero.
 // we re-use the same texture references, so the memory cost is relatively minimal.
 class FLandscapeGrassLightMap : public FLightMap2D
@@ -1378,10 +1339,7 @@ void ULandscapeGrassType::PostEditChangeProperty(FPropertyChangedEvent& Property
 {
 	Super::PostEditChangeProperty(PropertyChangedEvent);
 
-	// Look for landscape components using this grass type, and flush their grass data
-	// This is not required to trigger a grass update (the state hash changing will do that for us)
-	// This is only necessary to immediately delete the instances while dragging in the UI.
-	// (otherwise the instances stick around until we actually rebuild them)
+	// Look for landscape components using this grass type, and clear their grass type caches
 	if (GIsEditor)
 	{
 		for (TObjectIterator<ALandscapeProxy> It(/*AdditionalExclusionFlags = */RF_ClassDefaultObject, /*bIncludeDerivedClasses = */true, /*InInternalExclusionFlags = */EInternalObjectFlags::Garbage); It; ++It)
@@ -1391,14 +1349,9 @@ void ULandscapeGrassType::PostEditChangeProperty(FPropertyChangedEvent& Property
 			{
 				for (ULandscapeComponent* Component : Proxy->LandscapeComponents)
 				{
-					for (ULandscapeGrassType* GrassType : Component->GetGrassTypes())
+					if (Component->GetGrassTypes().Contains(this))
 					{
-						if (GrassType == this)
-						{
-							TSet<ULandscapeComponent*> Components = { Component };
-							Proxy->FlushGrassComponents(&Components, /* bFlushGrassMaps= */ true);
-							break;
-						}
+						Component->InvalidateGrassTypeSummary();
 					}
 				}
 			}
@@ -1501,7 +1454,6 @@ TArrayView<uint16> FLandscapeComponentGrassData::GetHeightData()
 
 void FLandscapeComponentGrassData::InitializeFrom(const TArray<uint16>& HeightData, const TMap<ULandscapeGrassType*, TArray<uint8>>& WeightData)
 {
-#if WITH_EDITORONLY_DATA
 	WeightOffsets.Empty(WeightData.Num());
 
 	// If weight data is empty make sure we don't have any memory allocated to grass
@@ -1535,9 +1487,45 @@ void FLandscapeComponentGrassData::InitializeFrom(const TArray<uint16>& HeightDa
 		FMemory::Memcpy(&CopyDest[CopyOffset], Pair.Value.GetData(), CopySize);
 		CopyOffset += CopySize;
 	}
-#endif
 }
 
+// Note that this is a relatively expensive operation (in the runtime sense), as it's making a copy of the buffers
+void FLandscapeComponentGrassData::InitializeFrom(IBuffer2DView<uint16>* HeightData, TMap<ULandscapeGrassType*, IBuffer2DView<uint8>*>& WeightData)
+{
+	WeightOffsets.Empty(WeightData.Num());
+
+	// If weight data is empty make sure we don't have any memory allocated to grass
+	if (WeightData.Num() == 0)
+	{
+		NumElements = 0;
+		HeightWeightData.Empty();
+		return;
+	}
+
+	NumElements = HeightData->Num();
+	HeightWeightData.SetNumUninitialized(NumElements * sizeof(uint16) + NumElements * WeightData.Num() * sizeof(uint8));
+
+	uint8* CopyDest = HeightWeightData.GetData();
+	int32 CopyOffset = 0;
+	int32 CopySize = NumElements * sizeof(uint16);
+
+	check((CopyOffset + CopySize) <= HeightWeightData.Num());
+
+	HeightData->CopyTo((uint16*)&CopyDest[CopyOffset], CopySize / sizeof(uint16));
+
+	CopyOffset += CopySize;
+	CopySize = NumElements * sizeof(uint8);
+
+	for (const TPair<ULandscapeGrassType*, IBuffer2DView<uint8>*>& Pair : WeightData)
+	{
+		WeightOffsets.Add(Pair.Key, CopyOffset);
+		check(Pair.Value->Num() == NumElements);
+		check((CopyOffset + CopySize) <= HeightWeightData.Num());
+
+		Pair.Value->CopyTo(&CopyDest[CopyOffset], CopySize);
+		CopyOffset += CopySize;
+	}
+}
 
 FArchive& operator<<(FArchive& Ar, FLandscapeComponentGrassData& Data)
 {
@@ -1549,23 +1537,23 @@ FArchive& operator<<(FArchive& Ar, FLandscapeComponentGrassData& Data)
 	{
 		if (Ar.CustomVer(FFortniteMainBranchObjectVersion::GUID) < FFortniteMainBranchObjectVersion::LandscapeSupportPerComponentGrassTypes)
 		{
-		if (Ar.CustomVer(FLandscapeCustomVersion::GUID) >= FLandscapeCustomVersion::GrassMaterialInstanceFix)
-		{
-				Ar << Data.MaterialStateIds_DEPRECATED;
-		}
-		else
-		{
-				Data.MaterialStateIds_DEPRECATED.Empty(1);
-			if (Ar.UEVer() >= VER_UE4_SERIALIZE_LANDSCAPE_GRASS_DATA_MATERIAL_GUID)
+			if (Ar.CustomVer(FLandscapeCustomVersion::GUID) >= FLandscapeCustomVersion::GrassMaterialInstanceFix)
 			{
-				FGuid MaterialStateId;
-				Ar << MaterialStateId;
+					Ar << Data.MaterialStateIds_DEPRECATED;
+			}
+			else
+			{
+					Data.MaterialStateIds_DEPRECATED.Empty(1);
+				if (Ar.UEVer() >= VER_UE4_SERIALIZE_LANDSCAPE_GRASS_DATA_MATERIAL_GUID)
+				{
+					FGuid MaterialStateId;
+					Ar << MaterialStateId;
 					Data.MaterialStateIds_DEPRECATED.Add(MaterialStateId);
+				}
 			}
-		}
 
-		if (Ar.CustomVer(FLandscapeCustomVersion::GUID) >= FLandscapeCustomVersion::GrassMaterialWPO)
-		{
+			if (Ar.CustomVer(FLandscapeCustomVersion::GUID) >= FLandscapeCustomVersion::GrassMaterialWPO)
+			{
 				Ar << Data.RotationForWPO_DEPRECATED;
 			}
 		}
@@ -2338,18 +2326,14 @@ void ALandscapeProxy::FlushGrassComponents(const TSet<ULandscapeComponent*>* Onl
 				Iter.RemoveCurrent();
 			}
 		}
-#if WITH_EDITOR
+
 		if (bFlushGrassMaps)
 		{
 			for (ULandscapeComponent* Component : *OnlyForComponents)
 			{
-				if (Component->CanRenderGrassMap())
-				{
-					Component->RemoveGrassMap();
-				}
+				Component->RemoveGrassMap();
 			}
 		}
-#endif
 	}
 	else
 	{
@@ -2394,10 +2378,7 @@ void ALandscapeProxy::FlushGrassComponents(const TSet<ULandscapeComponent*>* Onl
 			{
 				if (ULandscapeComponent* LandscapeComp = Cast<ULandscapeComponent>(Component))
 				{
-					if (LandscapeComp->CanRenderGrassMap())
-					{
-						LandscapeComp->RemoveGrassMap();
-					}
+					LandscapeComp->RemoveGrassMap();
 				}
 			}
 		}
@@ -2443,254 +2424,48 @@ void ALandscapeProxy::RemoveInvalidExclusionBoxes()
 
 void ALandscapeProxy::BuildGrassMaps(FScopedSlowTask* InSlowTask)
 {
-	if (!HasAnyFlags(RF_ClassDefaultObject))
-	{
-		// TODO [jonathan.bard], todo(luc.eygasier): do this part any time UpdateGrassTypes is called so that it's always in sync :
-		bool bHasLandscapeGrassBefore = bHasLandscapeGrass;
-		float GrassTypesMaxDiscardDistanceBefore = GrassTypesMaxDiscardDistance;
-
-		GrassTypesMaxDiscardDistance = 0.0f;
-		bHasLandscapeGrass = false;
-		for (ULandscapeComponent* Component : LandscapeComponents)
-		{
-			Component->UpdateGrassTypes(/* bForceUpdate = */ true);
-			bHasLandscapeGrass |= (Component->GetGrassTypes().Num() > 0);
-			GrassTypesMaxDiscardDistance = FMath::Max(Component->GetGrassTypesMaxDiscardDistance(), GrassTypesMaxDiscardDistance);
-		}
+	check(!HasAnyFlags(RF_ClassDefaultObject));
 
-		if ((bHasLandscapeGrass != bHasLandscapeGrassBefore)
-			|| (GrassTypesMaxDiscardDistance != GrassTypesMaxDiscardDistanceBefore))
+	// Update all grass types
+	for (ULandscapeComponent* Component : LandscapeComponents)
+	{
+		if (Component->UpdateGrassTypes())
 		{
-			MarkPackageDirty();
+			Component->MarkPackageDirty();
 		}
-
-		// Generate an mark dirty if changed
-		UpdateGrassData(true, InSlowTask);
 	}
-}
+	UpdateGrassTypeSummary();
 
-int32 ALandscapeProxy::GetOutdatedGrassMapCount() const
-{
-	int32 OutdatedGrassMaps = 0;
-	if (GGrassEnable > 0)
+	// now ensure that all of our components have up-to-date grass maps
+	ULandscapeSubsystem *Subsystem = GetWorld()->GetSubsystem<ULandscapeSubsystem>();
+	if (Subsystem)
 	{
-		UpdateGrassDataStatus(nullptr, nullptr, nullptr, nullptr, false, &OutdatedGrassMaps);
+		const bool bMarkDirty = true;
+		Subsystem->GetGrassMapBuilder()->BuildGrassMapsNowForComponents(LandscapeComponents, InSlowTask, bMarkDirty);
 	}
-	return OutdatedGrassMaps;
 }
 
-int32 ALandscapeProxy::TotalTexturesToStreamForVisibleGrassMapRender = 0;
 
-void ALandscapeProxy::UpdateGrassDataStatus(TSet<UTexture2D*>* OutCurrentForcedStreamedTextures, TSet<UTexture2D*>* OutDesiredForcedStreamedTextures, TSet<ULandscapeComponent*>* OutComponentsNeedingGrassMapRender, TSet<ULandscapeComponent*>* OutOutdatedComponents, bool bInEnableForceResidentFlag, int32* OutOutdatedGrassMaps) const
+int32 ALandscapeProxy::GetOutdatedGrassMapCount() const
 {
-	if (OutCurrentForcedStreamedTextures)
-	{
-		OutCurrentForcedStreamedTextures->Empty();
-	}
-	if (OutDesiredForcedStreamedTextures)
-	{
-		OutDesiredForcedStreamedTextures->Empty();
-	}
-
-	if (OutComponentsNeedingGrassMapRender)
-	{
-		OutComponentsNeedingGrassMapRender->Empty();
-	}
-
-	if (OutOutdatedComponents)
-	{
-		OutOutdatedComponents->Empty();
-	}
-
-	if (OutOutdatedGrassMaps)
-	{
-		*OutOutdatedGrassMaps = 0;
-	}
-		
-	const UWorld* World = GetWorld();
-	if (!World || World->IsGameWorld())
-	{
-		return;
-	}
-
-	// In either case we want to check the Grass textures stream state
-	const bool bCheckStreamingState = OutDesiredForcedStreamedTextures || bInEnableForceResidentFlag;
-
-	const bool bIsOutermostPackageDirty = GetOutermost()->IsDirty();
-
-	int32 OutdatedGrassMaps = 0;
-	for (auto Component : LandscapeComponents)
-	{
-		if (Component != nullptr)
-		{
-			UTexture2D* Heightmap = Component->GetHeightmap();
-			const TArray<UTexture2D*>& ComponentWeightmapTextures = Component->GetWeightmapTextures();
-
-			if (OutCurrentForcedStreamedTextures)
-			{
-				// check textures currently needing force streaming
-				if (Heightmap->bForceMiplevelsToBeResident)
-				{
-					OutCurrentForcedStreamedTextures->Add(Heightmap);
-				}
-
-				for (auto WeightmapTexture : ComponentWeightmapTextures)
-				{
-					if (WeightmapTexture->bForceMiplevelsToBeResident)
-					{
-						OutCurrentForcedStreamedTextures->Add(WeightmapTexture);
-					}
-				}
-			}
-
-			// if the landscape material has changed, update the grass types (needed to ensure ComputeGrassMapGenerationHash is correct)
-			Component->UpdateGrassTypes();
-			bool bGrassMapGenerationHashChanged = (Component->ComputeGrassMapGenerationHash() != Component->GrassData->GenerationHash);
-
-			// outdated meaning: it has grass data, but it is not up to date
-			bool bIsGrassMapOutdated = Component->GrassData->HasValidData() && bGrassMapGenerationHashChanged;
-
-			if (bIsGrassMapOutdated && OutOutdatedComponents)
-			{
-				OutOutdatedComponents->Add(Component);
-			}
-
-			// Needs to be called after UpdateGrassTypes
-			const TArray<TObjectPtr<ULandscapeGrassType>>& GrassTypes = Component->GetGrassTypes();
-			bool bHasGrassTypes = (GrassTypes.Num() > 0);
-			if (bHasGrassTypes || bBakeMaterialPositionOffsetIntoCollision)
-			{
-				if (bIsGrassMapOutdated || !Component->GrassData->HasValidData() || (GGrassUpdateAllOnRebuild != 0))
-				{
-					if (OutComponentsNeedingGrassMapRender)
-					{
-						OutComponentsNeedingGrassMapRender->Add(Component);
-					}
-					++OutdatedGrassMaps;
-
-					if (bCheckStreamingState && !Component->AreTexturesStreamedForGrassMapRender())
-					{
-						if (OutDesiredForcedStreamedTextures)
-						{
-							OutDesiredForcedStreamedTextures->Add(Heightmap);
-						}
-
-						if (bInEnableForceResidentFlag)
-						{
-							Heightmap->bForceMiplevelsToBeResident = true;
-						}
-
-						for (auto WeightmapTexture : ComponentWeightmapTextures)
-						{
-							if (OutDesiredForcedStreamedTextures)
-							{
-								OutDesiredForcedStreamedTextures->Add(WeightmapTexture);
-							}
-
-							if (bInEnableForceResidentFlag)
-							{
-								WeightmapTexture->bForceMiplevelsToBeResident = true;
-							}
-						}
-					}
-				}
-				// Don't count dirty component's if package is already dirty. 
-				// Saving the dirty package will also save the updated grass maps.
-				// This counter is used to know if grass maps need to be rebuilt.
-				else if (Component->GrassData->bIsDirty && !bIsOutermostPackageDirty)
-				{
-					++OutdatedGrassMaps;
-				}
-			}
-		}
-	}
-
-	if (OutOutdatedGrassMaps)
+	int32 GrassMapsNeedingRender = 0;
+	if (GGrassEnable > 0)
 	{
-		*OutOutdatedGrassMaps = OutdatedGrassMaps;
+		GrassMapsNeedingRender = GetWorld()->GetSubsystem<ULandscapeSubsystem>()->GetGrassMapBuilder()->CountOutdatedGrassMaps(LandscapeComponents);
 	}
+	return GrassMapsNeedingRender;
 }
 
-void ALandscapeProxy::UpdateGrassData(bool bInShouldMarkDirty, FScopedSlowTask* InSlowTask)
-{
-	if (!GGrassEnable || !GetWorld() || GetWorld()->IsGameWorld())
-	{
-		return;
-	}
-
-	// see if we need to flush grass for any components
-	TSet<UTexture2D*> DesiredForcedStreamedTextures;
-	TSet<UTexture2D*> CurrentForcedStreamedTextures;
-	TSet<ULandscapeComponent*> ComponentsNeedingGrassMapRender;
-	TSet<ULandscapeComponent*> OutdatedComponents;
-	int32 TotalOutdatedGrassMaps = 0;
-	const bool bEnableForceResidentFlag = true;
-	UpdateGrassDataStatus(&CurrentForcedStreamedTextures, &DesiredForcedStreamedTextures, &ComponentsNeedingGrassMapRender, &OutdatedComponents, bEnableForceResidentFlag, &TotalOutdatedGrassMaps);
-
-	if (OutdatedComponents.Num())
-	{
-		FlushGrassComponents(&OutdatedComponents);
-	}
-
-	// Remove local count from global count
-	TotalTexturesToStreamForVisibleGrassMapRender -= NumTexturesToStreamForVisibleGrassMapRender;
-	NumTexturesToStreamForVisibleGrassMapRender = 0;
-
-	// Remove local count from global count
-	NumComponentsNeedingGrassMapRender = 0;
-
-	// Wait for Texture Streaming
-	for (UTexture2D* TextureToStream : DesiredForcedStreamedTextures)
-	{
-		FTextureCompilingManager::Get().FinishCompilation({TextureToStream});
-		TextureToStream->WaitForStreaming();
-		CurrentForcedStreamedTextures.Add(TextureToStream);
-	}
-	
-	auto UpdateProgress = [InSlowTask](int Increment)
-	{
-		if (InSlowTask && Increment && ((InSlowTask->CompletedWork + Increment) <= InSlowTask->TotalAmountOfWork))
-		{
-			InSlowTask->EnterProgressFrame(static_cast<float>(Increment), FText::Format(LOCTEXT("GrassMaps_BuildGrassMapsProgress", "Building Grass Map {0} of  {1})"), FText::AsNumber(InSlowTask->CompletedWork), FText::AsNumber(InSlowTask->TotalAmountOfWork)));
-		}
-	};
-
-	// Take into consideration already dirty/rendered grass maps to update progress accordingly
-	int32 AlreadyRenderedGrassMap = TotalOutdatedGrassMaps - ComponentsNeedingGrassMapRender.Num();
-	UpdateProgress(AlreadyRenderedGrassMap);
-
-	// Render grass data
-	for(ULandscapeComponent* Component : ComponentsNeedingGrassMapRender)
-	{
-		if (Component->CanRenderGrassMap())
-		{
-			Component->RenderGrassMap();
-			UpdateProgress(1);
-		}
-	}
-
-	if (bInShouldMarkDirty && GetOutdatedGrassMapCount() > 0)
-	{
-		MarkPackageDirty();
-	}
-
-	// In Edit Layers, bForceMiplevelsToBeResident needs to remain true
-	if (!HasLayersContent())
-	{
-		for (UTexture2D* TextureToStream : CurrentForcedStreamedTextures)
-		{
-			TextureToStream->bForceMiplevelsToBeResident = false;
-		}
-	}
-}
 #endif
 
+// amortized update path (unless bForceSync is true)
 void ALandscapeProxy::UpdateGrass(const TArray<FVector>& Cameras, bool bForceSync /* = false */)
 {
 	int32 InOutNumCompsCreated = 0;
 	UpdateGrass(Cameras, InOutNumCompsCreated, bForceSync);
 }
 
+// amortized update path (unless bForceSync is true)
 void ALandscapeProxy::UpdateGrass(const TArray<FVector>& Cameras, int32& InOutNumCompsCreated, bool bForceSync /* = false */)
 {
 	SCOPE_CYCLE_COUNTER(STAT_GrassUpdate);
@@ -2728,27 +2503,11 @@ void ALandscapeProxy::UpdateGrass(const TArray<FVector>& Cameras, int32& InOutNu
 		if (World)
 		{
 			// Cull grass max distance based on Cull Distance scale factor and on max GuardBand factor.
-			float GrassMaxCulledDiscardDistance = GrassTypesMaxDiscardDistance * GGrassCullDistanceScale * FMath::Max(GGuardBandDiscardMultiplier, GGuardBandMultiplier);
-			float GrassMaxSquareDiscardDistance = GrassMaxCulledDiscardDistance * GrassMaxCulledDiscardDistance;
-
-#if WITH_EDITOR
+			// Note this is the proxy combined distance, not the per-component distance...
+			float GrassCullDistanceScalar = GGrassCullDistanceScale * FMath::Max(GGuardBandDiscardMultiplier, GGuardBandMultiplier);
+			float ProxyGrassMaxCulledDiscardDistance = GrassTypeSummary.MaxInstanceDiscardDistance * GrassCullDistanceScalar;
+			float ProxyGrassMaxSquareDiscardDistance = ProxyGrassMaxCulledDiscardDistance * ProxyGrassMaxCulledDiscardDistance;
 
-
-			int32 RequiredTexturesNotStreamedIn = 0;
-			TSet<ULandscapeComponent*> ComponentsNeedingGrassMapRender;
-			TSet<ULandscapeComponent*> OutdatedComponents;
-			TSet<UTexture2D*> CurrentForcedStreamedTextures;
-			TSet<UTexture2D*> DesiredForceStreamedTextures;
-			const bool bEnableForceResidentFlag = false;
-			
-			// Do not pass in DesiredForceStreamedTextures because we build our own list
-			UpdateGrassDataStatus(&CurrentForcedStreamedTextures, nullptr, &ComponentsNeedingGrassMapRender, &OutdatedComponents, bEnableForceResidentFlag);
-			
-			if(OutdatedComponents.Num())
-			{
-				FlushGrassComponents(&OutdatedComponents);
-			}
-#endif
 			ERHIFeatureLevel::Type FeatureLevel = World->Scene->GetFeatureLevel();
 
 			struct SortedLandscapeElement
@@ -2770,20 +2529,19 @@ void ALandscapeProxy::UpdateGrass(const TArray<FVector>& Cameras, int32& InOutNu
 			for (ULandscapeComponent* Component : LandscapeComponents)
 			{
 				// skip if we have no data and no way to generate it
-				if (Component == nullptr || (World->IsGameWorld() && !Component->GrassData->HasWeightData()))
+				if (Component == nullptr || !Component->GrassData->HasValidData() || !Component->GrassData->HasWeightData())
 				{
 					continue;
 				}
-				FBoxSphereBounds WorldBounds = Component->CalcBounds(Component->GetComponentTransform());
-			
-				float MinSqrDistanceToComponent = Cameras.Num() ? MAX_flt : 0.0f;
-				
+
+				FBoxSphereBounds WorldBounds = Component->CalcBounds(Component->GetComponentTransform());			
+				float MinSqrDistanceToComponent = Cameras.Num() ? MAX_flt : 0.0f;				
 				for (const FVector& CameraPos : Cameras)
 				{
 					MinSqrDistanceToComponent = FMath::Min<float>(MinSqrDistanceToComponent, static_cast<float>(WorldBounds.ComputeSquaredDistanceFromBoxToPoint(CameraPos)));
 				}
 
-				if ((GrassMaxSquareDiscardDistance > 0.f) && (MinSqrDistanceToComponent > GrassMaxSquareDiscardDistance))
+				if ((ProxyGrassMaxSquareDiscardDistance > 0.f) && (MinSqrDistanceToComponent > ProxyGrassMaxSquareDiscardDistance))
 				{
 					continue;
 				}
@@ -2802,11 +2560,15 @@ void ALandscapeProxy::UpdateGrass(const TArray<FVector>& Cameras, int32& InOutNu
 				Algo::Sort(SortedLandscapeComponents, [](const SortedLandscapeElement& A, const SortedLandscapeElement& B) { return A.MinDistance < B.MinDistance; });
 			}
 
+			// for every component that passes above, we iterate all grass types, grass varieties, and subsections (culling along the way)
+			// every subsection gets looked up in the cache and touched (to indicate it is still valid)
+			// if the subsection doesn't exist we try to start creating it (hence why we sorted the list in priority order)
 			for (const SortedLandscapeElement& SortedLandscapeComponent : SortedLandscapeComponents)
 			{
 				ULandscapeComponent* Component = SortedLandscapeComponent.Component;
 				float MinDistanceToComp = SortedLandscapeComponent.MinDistance;
 
+				// update the component's grass exclusion boxes if necessary
 				if (Component->ChangeTag != GGrassExclusionChangeTag)
 				{
 					Component->ActiveExcludedBoxes.Empty();
@@ -2840,7 +2602,7 @@ void ALandscapeProxy::UpdateGrass(const TArray<FVector>& Cameras, int32& InOutNu
 				{
 					if (GrassType)
 					{
-						if (World->IsGameWorld() && !Component->GrassData->Contains(GrassType))
+						if (World->IsGameWorld() && !Component->GrassData->Contains(GrassType))		// TODO [chris.tchou] why do we only do this in game worlds?
 						{
 							continue;
 						}
@@ -2860,13 +2622,11 @@ void ALandscapeProxy::UpdateGrass(const TArray<FVector>& Cameras, int32& InOutNu
 
 								if (!bUseHalton && MinDistanceToComp > DiscardDistance)
 								{
-									continue;
+									continue; // we can early out if we're not using halton sequence (otherwise we have to at least update the HaltonBaseIndex)
 								}
 
 								FGrassBuilderBase ForSubsectionMath(this, Component, GrassVariety, FeatureLevel);
-
 								int32 SqrtSubsections = 1;
-
 								if (ForSubsectionMath.bHaveValidData && ForSubsectionMath.SqrtMaxInstances > 0)
 								{
 									SqrtSubsections = FMath::Clamp<int32>(FMath::CeilToInt(float(ForSubsectionMath.SqrtMaxInstances) / FMath::Sqrt((float)MaxInstancesPerComponent)), 1, 16);
@@ -2876,7 +2636,7 @@ void ALandscapeProxy::UpdateGrass(const TArray<FVector>& Cameras, int32& InOutNu
 								if (bUseHalton && MinDistanceToComp > DiscardDistance)
 								{
 									HaltonBaseIndex += MaxInstancesSub * SqrtSubsections * SqrtSubsections;
-									continue;
+									continue; // now that we have updated the halton base index, it's ok to skip the rest
 								}
 
 								FBox LocalBox = Component->CachedLocalBox;
@@ -2989,37 +2749,7 @@ void ALandscapeProxy::UpdateGrass(const TArray<FVector>& Cameras, int32& InOutNu
 										}
 										NewComp.ExclusionChangeTag = GGrassExclusionChangeTag;
 
-#if WITH_EDITOR
-										// render grass data if we don't have any
-										if (!Component->GrassData->HasValidData())
-										{
-											if (!Component->CanRenderGrassMap())
-											{
-												// we can't currently render grassmaps (eg shaders not compiled)
-												continue;
-											}
-											else if (!Component->AreTexturesStreamedForGrassMapRender())
-											{
-												// we're ready to generate but our textures need streaming in
-												DesiredForceStreamedTextures.Add(Component->GetHeightmap());
-												const TArray<UTexture2D*>& ComponentWeightmapTextures = Component->GetWeightmapTextures();
-												
-												for (UTexture2D* WeightmapTexture : ComponentWeightmapTextures)
-												{
-													DesiredForceStreamedTextures.Add(WeightmapTexture);
-												}
-												RequiredTexturesNotStreamedIn++;
-												continue;
-											}
-
-											RenderCaptureInterface::FScopedCapture RenderCapture((GCaptureNextGrassUpdate != 0), TEXT("RenderGrassMap"));
-											GCaptureNextGrassUpdate = FMath::Max(GCaptureNextGrassUpdate - 1, 0);
-
-											QUICK_SCOPE_CYCLE_COUNTER(STAT_GrassRenderToTexture);
-											Component->RenderGrassMap();
-											ComponentsNeedingGrassMapRender.Remove(Component);
-										}
-#endif
+										check(Component->GrassData->HasValidData());
 
 										InOutNumCompsCreated++;
 
@@ -3154,79 +2884,6 @@ void ALandscapeProxy::UpdateGrass(const TArray<FVector>& Cameras, int32& InOutNu
 					}
 				}
 			}
-
-#if WITH_EDITOR
-
-			TotalTexturesToStreamForVisibleGrassMapRender -= NumTexturesToStreamForVisibleGrassMapRender;
-			NumTexturesToStreamForVisibleGrassMapRender = RequiredTexturesNotStreamedIn;
-			TotalTexturesToStreamForVisibleGrassMapRender += NumTexturesToStreamForVisibleGrassMapRender;
-
-			{
-				int32 NumComponentsRendered = 0;
-				int32 NumComponentsUnableToRender = 0;
-				if ((GPrerenderGrassmaps > 0) || bBakeMaterialPositionOffsetIntoCollision)
-				{
-					// try to render some grassmaps
-					TArray<ULandscapeComponent*> ComponentsToRender;
-					for (auto Component : ComponentsNeedingGrassMapRender)
-					{
-						if (Component->CanRenderGrassMap())
-						{
-							if (Component->AreTexturesStreamedForGrassMapRender())
-							{
-								// We really want to throttle the number based on component size.
-								if (NumComponentsRendered <= 4)
-								{
-									ComponentsToRender.Add(Component);
-									NumComponentsRendered++;
-								}
-							}
-							else
-							if (TotalTexturesToStreamForVisibleGrassMapRender == 0)
-							{
-								// Force stream in other heightmaps but only if we're not waiting for the textures 
-								// near the camera to stream in
-								DesiredForceStreamedTextures.Add(Component->GetHeightmap());
-								const TArray<UTexture2D*>& ComponentWeightmapTextures = Component->GetWeightmapTextures();
-								
-								for (UTexture2D* WeightmapTexture : ComponentWeightmapTextures)
-								{
-									DesiredForceStreamedTextures.Add(WeightmapTexture);
-								}
-							}
-						}
-						else
-						{
-							NumComponentsUnableToRender++;
-						}
-					}
-					if (ComponentsToRender.Num())
-					{
-						RenderCaptureInterface::FScopedCapture RenderCapture((GCaptureNextGrassUpdate != 0), TEXT("RenderGrassMaps"));
-						GCaptureNextGrassUpdate = FMath::Max(GCaptureNextGrassUpdate - 1, 0);
-
-						RenderGrassMaps(ComponentsToRender);
-					}
-				}
-
-				NumComponentsNeedingGrassMapRender = ComponentsNeedingGrassMapRender.Num() - NumComponentsRendered - NumComponentsUnableToRender;
-
-				// Update resident flags
-				for (auto Texture : DesiredForceStreamedTextures.Difference(CurrentForcedStreamedTextures))
-				{
-					Texture->bForceMiplevelsToBeResident = true;
-				}
-
-				// In Edit Layers, bForceMiplevelsToBeResident needs to remain true
-				if (!HasLayersContent())
-				{
-					for (auto Texture : CurrentForcedStreamedTextures.Difference(DesiredForceStreamedTextures))
-					{
-						Texture->bForceMiplevelsToBeResident = false;
-					}
-				}
-			}
-#endif
 		}
 	}
 
@@ -3271,6 +2928,8 @@ void ALandscapeProxy::UpdateGrass(const TArray<FVector>& Cameras, int32& InOutNu
 			}
 		}
 	}
+
+	// remove components that are no longer used from the components list
 	if (StillUsed.Num() < FoliageComponents.Num())
 	{
 		QUICK_SCOPE_CYCLE_COUNTER(STAT_Grass_DelComps);
@@ -3298,63 +2957,68 @@ void ALandscapeProxy::UpdateGrass(const TArray<FVector>& Cameras, int32& InOutNu
 			}
 		}
 	}
+	
+	// check if any async tasks have completed, and finalize them
+	ProcessAsyncGrassInstanceTasks(bWaitAsyncTasks, bForceSync, StillUsed);
+}
+
+void ALandscapeProxy::ProcessAsyncGrassInstanceTasks(bool bWaitAsyncTasks, bool bForceSync, const TSet<UHierarchicalInstancedStaticMeshComponent*>& StillUsed)
+{
+	QUICK_SCOPE_CYCLE_COUNTER(STAT_Grass_FinishAsync);
+	// finish async tasks
+	for (int32 Index = 0; Index < AsyncFoliageTasks.Num(); Index++)
 	{
-		QUICK_SCOPE_CYCLE_COUNTER(STAT_Grass_FinishAsync);
-		// finish async tasks
-		for (int32 Index = 0; Index < AsyncFoliageTasks.Num(); Index++)
+		FAsyncTask<FAsyncGrassTask>* Task = AsyncFoliageTasks[Index];
+		if (bWaitAsyncTasks)
 		{
-			FAsyncTask<FAsyncGrassTask>* Task = AsyncFoliageTasks[Index];
-			if (bWaitAsyncTasks)
-			{
-				Task->EnsureCompletion();
-			}
-			if (Task->IsDone())
+			Task->EnsureCompletion();
+		}
+		if (Task->IsDone())
+		{
+			SCOPE_CYCLE_COUNTER(STAT_FoliageGrassEndComp);
+			FAsyncGrassTask& Inner = Task->GetTask();
+			AsyncFoliageTasks.RemoveAtSwap(Index--);
+			UGrassInstancedStaticMeshComponent* GrassISMComponent = Cast<UGrassInstancedStaticMeshComponent>(Inner.Foliage.Get());
+			int32 NumBuiltRenderInstances = Inner.Builder->InstanceBuffer.GetNumInstances();
+			//UE_LOG(LogCore, Display, TEXT("%d instances in %4.0fms     %6.0f instances / sec"), NumBuiltRenderInstances, 1000.0f * float(Inner.Builder->BuildTime), float(NumBuiltRenderInstances) / float(Inner.Builder->BuildTime));
+
+			if (GrassISMComponent && StillUsed.Contains(GrassISMComponent))
 			{
-				SCOPE_CYCLE_COUNTER(STAT_FoliageGrassEndComp);
-				FAsyncGrassTask& Inner = Task->GetTask();
-				AsyncFoliageTasks.RemoveAtSwap(Index--);
-				UGrassInstancedStaticMeshComponent* GrassISMComponent = Cast<UGrassInstancedStaticMeshComponent>(Inner.Foliage.Get());
-				int32 NumBuiltRenderInstances = Inner.Builder->InstanceBuffer.GetNumInstances();
-				//UE_LOG(LogCore, Display, TEXT("%d instances in %4.0fms     %6.0f instances / sec"), NumBuiltRenderInstances, 1000.0f * float(Inner.Builder->BuildTime), float(NumBuiltRenderInstances) / float(Inner.Builder->BuildTime));
-
-				if (GrassISMComponent && StillUsed.Contains(GrassISMComponent))
+				if (NumBuiltRenderInstances > 0)
 				{
-					if (NumBuiltRenderInstances > 0)
+					QUICK_SCOPE_CYCLE_COUNTER(STAT_FoliageGrassEndComp_AcceptPrebuiltTree);
+					GrassISMComponent->AcceptPrebuiltTree(Inner.Builder->ClusterTree, Inner.Builder->OutOcclusionLayerNum, NumBuiltRenderInstances, &Inner.Builder->InstanceBuffer);
+					if (bForceSync && GetWorld())
 					{
-						QUICK_SCOPE_CYCLE_COUNTER(STAT_FoliageGrassEndComp_AcceptPrebuiltTree);
-						GrassISMComponent->AcceptPrebuiltTree(Inner.Builder->ClusterTree, Inner.Builder->OutOcclusionLayerNum, NumBuiltRenderInstances, &Inner.Builder->InstanceBuffer);
-						if (bForceSync && GetWorld())
-						{
-							QUICK_SCOPE_CYCLE_COUNTER(STAT_FoliageGrassEndComp_SyncUpdate);
-							GrassISMComponent->RecreateRenderState_Concurrent();
-						}
+						QUICK_SCOPE_CYCLE_COUNTER(STAT_FoliageGrassEndComp_SyncUpdate);
+						GrassISMComponent->RecreateRenderState_Concurrent();
 					}
 				}
-				FCachedLandscapeFoliage::FGrassComp* Existing = FoliageCache.CachedGrassComps.Find(Inner.Key);
-				if (Existing)
+			}
+			FCachedLandscapeFoliage::FGrassComp* Existing = FoliageCache.CachedGrassComps.Find(Inner.Key);
+			if (Existing)
+			{
+				Existing->Pending = false;
+				if (Existing->PreviousFoliage.IsValid())
 				{
-					Existing->Pending = false;
-					if (Existing->PreviousFoliage.IsValid())
+					SCOPE_CYCLE_COUNTER(STAT_FoliageGrassDestoryComp);
+					UHierarchicalInstancedStaticMeshComponent* HComponent = Existing->PreviousFoliage.Get();
+					if (HComponent)
 					{
-						SCOPE_CYCLE_COUNTER(STAT_FoliageGrassDestoryComp);
-						UHierarchicalInstancedStaticMeshComponent* HComponent = Existing->PreviousFoliage.Get();
-						if (HComponent)
-						{
-							HComponent->ClearInstances();
-							HComponent->DetachFromComponent(FDetachmentTransformRules(EDetachmentRule::KeepRelative, false));
-							HComponent->DestroyComponent();
-						}
-						FoliageComponents.RemoveSwap(HComponent);
-						Existing->PreviousFoliage = nullptr;
+						HComponent->ClearInstances();
+						HComponent->DetachFromComponent(FDetachmentTransformRules(EDetachmentRule::KeepRelative, false));
+						HComponent->DestroyComponent();
 					}
-
-					Existing->Touch();
-				}
-				delete Task;
-				if (!bWaitAsyncTasks)
-				{
-					break; // one per frame is fine
+					FoliageComponents.RemoveSwap(HComponent);
+					Existing->PreviousFoliage = nullptr;
 				}
+
+				Existing->Touch();
+			}
+			delete Task;
+			if (!bWaitAsyncTasks)
+			{
+				break; // one per frame is fine
 			}
 		}
 	}
diff --git a/Engine/Source/Runtime/Landscape/Private/LandscapeGrassMapsBuilder.cpp b/Engine/Source/Runtime/Landscape/Private/LandscapeGrassMapsBuilder.cpp
--- a/Engine/Source/Runtime/Landscape/Private/LandscapeGrassMapsBuilder.cpp
+++ b/Engine/Source/Runtime/Landscape/Private/LandscapeGrassMapsBuilder.cpp
@@ -3,36 +3,1205 @@
 #include "LandscapeGrassMapsBuilder.h"
 #include "Misc/ScopedSlowTask.h"
 #include "Landscape.h"
+#include "LandscapePrivate.h"
 #include "LandscapeProxy.h"
 #include "LandscapeGrassWeightExporter.h"
+#include "MaterialCachedData.h"
+#include "LandscapeGrassType.h"
+#include "LandscapeSubsystem.h"
 #include "EngineUtils.h"
+#include "Materials/MaterialInstance.h"
+#include "Stats/Stats.h"
+#include "Components/InstancedStaticMeshComponent.h"
+#include "SceneInterface.h"
 
 #define LOCTEXT_NAMESPACE "Landscape"
 
+#define GRASS_DEBUG_LOG(...) UE_LOG(LogGrass, Verbose, __VA_ARGS__)
+
 extern int32 GGrassEnable;
+extern float GGrassCullDistanceScale;
 
+int32 GGrassMapUseRuntimeGeneration = 0;
+FAutoConsoleVariableRef CVarGrassMapUseRuntimeGeneration(
+	TEXT("grass.GrassMap.UseRuntimeGeneration"),
+	GGrassMapUseRuntimeGeneration,
+	TEXT("Enable runtime grass map generation to save disk space and runtime memory.  When enabled the grass density maps are not serialized and are built on the fly at runtime."));
 
-#if WITH_EDITOR
+int32 GGrassMapAlwaysBuildRuntimeGenerationResources = 0;
+static FAutoConsoleVariableRef CVarGrassMapAlwaysBuildRuntimeGenerationResources(
+	TEXT("grass.GrassMap.AlwaysBuildRuntimeGenerationResources"),
+	GGrassMapAlwaysBuildRuntimeGenerationResources,
+	TEXT("By default we only compile shaders and build resources for runtime generation when runtime generation is enabled.  Set this to 1 to always build them for all platforms, allowing you to toggle runtime generation in a cooked build."));
+
+static int32 GGrassMapMaxComponentsStreaming = 1;
+static FAutoConsoleVariableRef CVarGrassMapMaxComponentsStreaming(
+	TEXT("grass.GrassMap.MaxComponentsStreaming"),
+	GGrassMapMaxComponentsStreaming,
+	TEXT("How many landscape components can be streaming their textures at once for grass map renders, when using amortized runtime generation."));
+
+static int32 GGrassMapMaxComponentsRendering = 1;
+static FAutoConsoleVariableRef CVarGrassMapMaxComponentsRendering(
+	TEXT("grass.GrassMap.MaxComponentsRendering"),
+	GGrassMapMaxComponentsRendering,
+	TEXT("How many landscape components can be rendering grass maps at once, when using amortized runtime generation."));
+
+static int32 GGrassMapMaxComponentsForBlockingUpdate = 6;
+static FAutoConsoleVariableRef CVarGrassMapMaxComponentsForBlockingUpdate(
+	TEXT("grass.GrassMap.MaxComponentsForBlockingUpdate"),
+	GGrassMapMaxComponentsForBlockingUpdate,
+	TEXT("How many landscape components can update simultaneously when running a blocking grass map update (i.e. on editor save)."));
+
+static int32 GGrassMapMaxDiscardChecksPerFrame = 25;
+static FAutoConsoleVariableRef CVarGrassMapMaxDiscardChecksPerFrame(
+	TEXT("grass.GrassMap.MaxDiscardChecksPerFrame"),
+	GGrassMapMaxDiscardChecksPerFrame,
+	TEXT("How many landscape components are checked if they should discard their grass maps each frame."));
+
+static int32 GGrassMapPrioritizedMultiplier = 4;
+static FAutoConsoleVariableRef CVarGrassMapCreationPrioritizedMultiplier(
+	TEXT("grass.GrassMap.PrioritizedMultiplier"),
+	GGrassMapPrioritizedMultiplier,
+	TEXT("Multiplier applied to MaxComponentsStreaming and MaxComponentsRendering when grass creation is prioritized."));
 
-FLandscapeGrassMapsBuilder::FLandscapeGrassMapsBuilder(UWorld* InOwner)
+static float GGrassMapGuardBandMultiplier = 1.5f;
+static FAutoConsoleVariableRef CVarGrassMapGuardBandMultiplier(
+	TEXT("grass.GrassMap.GuardBandMultiplier"),
+	GGrassMapGuardBandMultiplier,
+	TEXT("Used to control discarding in the grass map runtime generation system. Approximate range, 1-4. Multiplied by the cull distance to control when we add grass maps."));
+
+static float GGrassMapGuardBandDiscardMultiplier = 1.6f;
+static FAutoConsoleVariableRef CVarGrassMapGuardBandDiscardMultiplier(
+	TEXT("grass.GrassMap.GuardBandDiscardMultiplier"),
+	GGrassMapGuardBandDiscardMultiplier,
+	TEXT("Used to control discarding in the grass map runtime generation system. Approximate range, 1-4. Multiplied by the cull distance to control when we discard grass maps."));
+
+DECLARE_CYCLE_STAT(TEXT("Update Component GrassMap "), STAT_UpdateComponentGrassMaps, STATGROUP_Foliage);
+DECLARE_CYCLE_STAT(TEXT("Prioritize Pending GrassMaps"), STAT_PrioritizePendingGrassMaps, STATGROUP_Foliage);
+DECLARE_CYCLE_STAT(TEXT("Render GrassMap"), STAT_RenderGrassMap, STATGROUP_Foliage);
+DECLARE_CYCLE_STAT(TEXT("Populate GrassMap"), STAT_PopulateGrassMap, STATGROUP_Foliage);
+DECLARE_CYCLE_STAT(TEXT("Remove Grass Instances"), STAT_RemoveGrassInstances, STATGROUP_Foliage);
+
+FLandscapeGrassMapsBuilder::FLandscapeGrassMapsBuilder(UWorld* InOwner, FLandscapeTextureStreamingManager& InTextureStreamingManager)
 	: World(InOwner)
+	, TextureStreamingManager(InTextureStreamingManager)
+#if WITH_EDITOR
 	, OutdatedGrassMapCount(0)
 	, GrassMapsLastCheckTime(0)
-{}
+#endif // WITH_EDITOR
+{
+}
+
+TAllocatorFixedSizeFreeList<sizeof(FLandscapeGrassMapsBuilder::FComponentState), 32> FLandscapeGrassMapsBuilder::StatePoolAllocator;
+
+namespace UE::Landscape
+{
+	uint32 ComputeGrassMapGenerationHash(const ULandscapeComponent* Component, UMaterialInterface* Material)
+	{
+		TRACE_CPUPROFILER_EVENT_SCOPE(ComputeGrassMapGenerationHash);
+		uint32 Hash = 0;
+
+		// Change this hash key to invalidate all cached grass density maps, if the generation functions change.
+		static uint32 InitialHash = GetTypeHash(FGuid("216D95C7651D4095ADC6A8459B4F181D"));
+
+		Hash = InitialHash;
+
+		// we only include material and texture hashes in editor (there is no automatic detection of changes in non-editor builds)
+	#if WITH_EDITOR
+		// Take into account any material state change : (excluding texture state)
+		Hash = FCrc::TypeCrc32(Material->ComputeAllStateCRC(), Hash);
+
+		// hash the heightmap texture (we use Source Id as it is a content hash, lighting guid is random)
+		UTexture2D* Heightmap = Component->GetHeightmap();
+		check(Heightmap->Source.IsValid());
+		Hash = FCrc::TypeCrc32(Heightmap->Source.GetId(), Hash);
+
+		// hash the weightmap textures
+		for (UTexture2D* Weightmap : Component->GetWeightmapTextures())
+		{
+			check(Weightmap->Source.IsValid());
+			Hash = FCrc::TypeCrc32(Weightmap->Source.GetId(), Hash);
+		}
+	#endif // WITH_EDITOR
+
+		return Hash;
+	}
+
+	uint32 ComputeGrassInstanceGenerationHash(uint32 GrassMapGenerationHash, const TArray<TObjectPtr<ULandscapeGrassType>>& GrassTypes)
+	{
+		TRACE_CPUPROFILER_EVENT_SCOPE(ComputeGrassInstanceGenerationHash);
+		// grass maps are input to grass instances
+		uint32 Hash = GrassMapGenerationHash;
+
+		// If anything changes in the grass types, that affects the grass instances
+		for (ULandscapeGrassType* GrassType : GrassTypes)
+		{
+			Hash = FCrc::TypeCrc32(GrassType ? GrassType->StateHash : 0, Hash);
+		}
+
+		return Hash;
+	}
+
+#if WITH_EDITOR
+	void CompileGrassMapShader(ULandscapeComponent* Component)
+	{
+		if (Component->GetMaterialInstanceCount(false) > 0)
+		{
+			if (UMaterialInstance* MaterialInstance = Component->GetMaterialInstance(0))
+			{
+				if (FMaterialResource* MaterialResource = MaterialInstance->GetMaterialResource(Component->GetWorld()->GetFeatureLevel()))
+				{
+					MaterialResource->FinishCompilation();
+				}
+			}
+		}
+	}
+#endif // WITH_EDITOR
+
+	bool CanRenderGrassMap(ULandscapeComponent *Component)
+	{
+		// Check we can render
+		UWorld* ComponentWorld = Component->GetWorld();
+		if (GUsingNullRHI || !ComponentWorld || !Component->SceneProxy)
+		{
+			return false;
+		}
+
+		UMaterialInstance* MaterialInstance = Component->GetMaterialInstanceCount(false) > 0 ? Component->GetMaterialInstance(0) : nullptr;
+		FMaterialResource* MaterialResource = MaterialInstance != nullptr ? MaterialInstance->GetMaterialResource(ComponentWorld->GetFeatureLevel()) : nullptr;
+
+		// Check we can render the material
+		if (MaterialResource == nullptr)
+		{
+			return false;
+		}
+
+		// We only need the GrassWeight shaders on the fixed grid vertex factory to render grass maps : 
+		FMaterialShaderTypes ShaderTypes;
+		UE::Landscape::Grass::AddGrassWeightShaderTypes(ShaderTypes);
+
+		FVertexFactoryType* LandscapeGrassVF = FindVertexFactoryType(FName(TEXT("FLandscapeFixedGridVertexFactory"), FNAME_Find));
+		if (!MaterialResource->HasShaders(ShaderTypes, LandscapeGrassVF))
+		{
+			return false;
+		}
+		return true;
+	}
+
+	// calculates the minimum distance between any cameras and the specified Worldbounds.
+	static inline double CalculateMinDistanceToCameras(const TArray<FVector>& Cameras, const FBoxSphereBounds& WorldBounds)
+	{
+		if (!Cameras.Num())
+		{
+			return 0.0f;
+		}
+		double MinSqrDistance = MAX_dbl;
+		for (const FVector& CameraPos : Cameras)
+		{
+			MinSqrDistance = FMath::Min<double>(MinSqrDistance, WorldBounds.ComputeSquaredDistanceFromBoxToPoint(CameraPos));
+		}
+		return MinSqrDistance;
+	}
+
+	static void SubmitGPUCommands(bool bBlockUntilComplete)
+	{
+		FEvent* ResultsReadyEvent = nullptr;
+		if (bBlockUntilComplete)
+		{
+			ResultsReadyEvent = FPlatformProcess::GetSynchEventFromPool(true);
+		}
+
+		ENQUEUE_RENDER_COMMAND(FFlushResourcesCommand)(
+			[ResultsReadyEvent](FRHICommandList& RHICmdList)
+			{
+				FRHICommandListExecutor::GetImmediateCommandList().ImmediateFlush(EImmediateFlushType::FlushRHIThreadFlushResources);
+				RHIFlushResources();
+				FRHICommandListExecutor::GetImmediateCommandList().SubmitCommandsAndFlushGPU();
+				if (ResultsReadyEvent)
+				{
+					// block render thread waiting for GPU to complete
+					FRHICommandListExecutor::GetImmediateCommandList().BlockUntilGPUIdle();
+					ResultsReadyEvent->Trigger();
+				}
+			});
+
+		if (ResultsReadyEvent)
+		{
+			// block game thread waiting for render thread to tell us the GPU is complete
+			ResultsReadyEvent->Wait();
+			ResultsReadyEvent->Reset();
+			FPlatformProcess::ReturnSynchEventToPool(ResultsReadyEvent);
+		}
+
+		return;
+	}
+}
+
+FLandscapeGrassMapsBuilder::~FLandscapeGrassMapsBuilder()
+{
+	// make sure all components were unregistered, so that state cleanup and deletion is triggered
+	for (auto It = ComponentStates.CreateIterator(); It; ++It)
+	{
+		FComponentState* State = It.Value();
+		ULandscapeComponent* Component = State->Component;
+		if (!ensure(Component == nullptr))
+		{
+			UnregisterComponent(Component);
+		}
+		State->TickCount = 7777;	// set to a large number so we don't wait to delete entries
+	}
+
+	// update component state until they all delete themselves (should happen on the first update unless a readback is active, and shouldn't take more than 3)
+	int32 Iterations = 0;
+	while (ensure(Iterations < 6) && ComponentStates.Num() > 0)
+	{
+		TArray<FVector> EmptyCamerasArray;
+		UpdateTrackedComponents(EmptyCamerasArray, 0);
+
+		// if we happen to get caught with a GPU readback in flight, submit GPU commands to make sure it moves forward
+		if (RenderingCount > 0)
+		{
+			UE::Landscape::SubmitGPUCommands(/* bBlockUntilComplete =  */ true);
+		}
+		Iterations++;
+	}
+}
+
+FLandscapeGrassMapsBuilder::FComponentState::FComponentState(ULandscapeComponent* Component)
+	: Component(Component)
+{
+#if WITH_EDITOR
+	UMaterialInterface* Material = Component->GetLandscapeMaterial();
+	if (Material)
+	{
+		GrassMapGenerationHash = UE::Landscape::ComputeGrassMapGenerationHash(Component, Material);
+		const TArray<TObjectPtr<ULandscapeGrassType>>& GrassTypes = Material->GetCachedExpressionData().GrassTypes;
+		GrassInstanceGenerationHash = UE::Landscape::ComputeGrassInstanceGenerationHash(GrassMapGenerationHash, GrassTypes);
+	}
+#endif // WITH_EDITOR
+}
+
+void FLandscapeGrassMapsBuilder::FPendingComponent::UpdatePriorityDistance(const TArray<FVector>&Cameras)
+{
+	TRACE_CPUPROFILER_EVENT_SCOPE(UpdatePriorityDistance);
+	ULandscapeComponent* Component = State->Component;
+	FBoxSphereBounds WorldBounds = Component->CalcBounds(Component->GetComponentTransform());
+	double MinSqrDistanceToComponent = UE::Landscape::CalculateMinDistanceToCameras(Cameras, WorldBounds);
+	PriorityKey = MinSqrDistanceToComponent;
+}
+
+bool FLandscapeGrassMapsBuilder::UpdateTrackedComponents(const TArray<FVector>& Cameras, int32 LocalMaxRendering)
+{
+	SCOPE_CYCLE_COUNTER(STAT_UpdateComponentGrassMaps);
+
+	bool bChanged = false;
+
+	// Array to store components that are updated after the initial update pass (should never be more than max streaming components)
+	static TArray<FComponentState*> StreamingStatesToProcess;
+	StreamingStatesToProcess.Reset();
+
+	// Set to store and components that need foliage instances removed
+	static TSet<ULandscapeComponent*> ComponentsToRemoveFoliageInstances;
+	ComponentsToRemoveFoliageInstances.Reset();
+
+	AmortizedUpdate.StartUpdateTick(ComponentStates.Num(), GGrassMapMaxDiscardChecksPerFrame);
+
+	// Iterate our components, removing invalid ones and counting how many are in each state.
+	// We can also immediately process any components in the populated or rendering states
+	// (those states don't need to consider throttling when moving to the next state, and processing them early frees up slots)
+	int32 ComponentStateIndex = 0;
+	for (auto It = ComponentStates.CreateIterator(); It; ++It, ++ComponentStateIndex)
+	{
+		TRACE_CPUPROFILER_EVENT_SCOPE(UpdateComponentState);
+
+		FComponentState* State = It.Value();
+		ULandscapeComponent* Component = State->Component;
+		State->TickCount++;
+
+		if (Component != nullptr)
+		{
+			if (GGrassEnable)
+			{
+				switch (State->Stage)
+				{
+				case EComponentStage::Pending:
+					if (World->IsGameWorld())
+					{
+						// if runtime grass generation is disabled, go straight to not ready.
+						if (!Component->GetLandscapeProxy()->bUseRuntimeGrassMapGeneration)
+						{
+							PendingToNotReady(*State);
+							continue; // next!
+						}
+					}
+					continue; // next!
+				case EComponentStage::NotReady:
+					// in game, any not ready component will never become ready.
+					// in editor, check to see if the conditions changed.
+					if (World->IsGameWorld() || !UE::Landscape::CanRenderGrassMap(Component))
+					{
+						continue; // still not ready - next!
+					}
+					break; // cancel and evict to restart build process
+				case EComponentStage::TextureStreaming:
+					// don't process streaming states yet -- first process rendering states to free up slots
+					StreamingStatesToProcess.Add(State);
+					continue; // next!
+
+				case EComponentStage::Rendering:
+					check(State->ActiveRender != nullptr);
+					if (State->ActiveRender->CheckAndUpdateAsyncReadback())
+					{
+						// TODO [chris.tchou] We could move this to an async task, it does have a decent cost to it
+						PopulateGrassDataFromReadback(*State);
+						bChanged = true;
+					}
+					continue; // next!
+
+				case EComponentStage::GrassMapsPopulated:
+					// only check for invalidation of populated grass maps once in a while
+					if (AmortizedUpdate.ShouldUpdate(ComponentStateIndex))
+					{
+						// detect if grass data has been cleared by someone manually calling Flush (i.e. when landscape edits are made)
+						if (!Component->GrassData->HasValidData())
+						{
+							break; // cancel and evict to restart process
+						}
+
+						// check if the component is too far from the camera and we can reclaim the grass data
+						if (GGrassMapUseRuntimeGeneration &&
+							State->IsBeyondEvictionRange(Cameras))
+						{
+							GRASS_DEBUG_LOG(TEXT("Evicting for being beyond eviction range"));
+							break; // cancel and evict
+						}
+
+#if WITH_EDITOR
+						UMaterialInterface* Material = Component->GetLandscapeMaterial();
+						if (Material == nullptr)
+						{
+							break; // cancel and evict
+						}
+
+						// check if any dependencies changed
+						uint32 CurGrassMapGenerationHash = UE::Landscape::ComputeGrassMapGenerationHash(Component, Material);
+						if (State->GrassMapGenerationHash != CurGrassMapGenerationHash)
+						{
+							break; // cancel and evict to restart process
+						}
+
+						// check if any grass types have changed -- this invalidates foliage instances but not the grass maps
+						const TArray<TObjectPtr<ULandscapeGrassType>>& GrassTypes = Component->GetGrassTypes();
+						uint32 CurGrassInstanceGenerationHash = UE::Landscape::ComputeGrassInstanceGenerationHash(CurGrassMapGenerationHash, GrassTypes);
+						if (State->GrassInstanceGenerationHash != CurGrassInstanceGenerationHash)
+						{
+							ComponentsToRemoveFoliageInstances.Add(Component);
+							Component->InvalidateGrassTypeSummary();
+							State->GrassInstanceGenerationHash = CurGrassInstanceGenerationHash;
+						}
+#endif // WITH_EDITOR
+					}
+					continue;	// next!
+
+				default:
+					check(false);	// unreachable
+					break;			// cancel and evict
+				}
+			}
+		}
+
+		// we only fall through to this statement if the code above didn't invoke `continue` (or if the component was unregistered)
+		if (CancelAndEvict(*State))
+		{
+			if (State->Component == nullptr)
+			{
+				RemoveFromPendingComponentHeap(State);
+
+				// destruct and free the state (return to our pool)
+				State->~FComponentState();
+				StatePoolAllocator.Free(State);
+				State = nullptr;
+
+				It.RemoveCurrent();
+				AmortizedUpdate.HandleDeletion(ComponentStateIndex);
+				check(PendingCount > 0);
+				PendingCount--;
+			}
+			else
+			{
+				// component is still registered, but has been invalidated.  Remove foliage instances.
+				ComponentsToRemoveFoliageInstances.Add(Component);
+			}
+			bChanged = true;
+		}
+	}
+
+	// kick off any deferred rendering
+	if (StreamingStatesToProcess.Num() > 0)
+	{
+		TRACE_CPUPROFILER_EVENT_SCOPE(StreamingStatesToProcess);
+		for (int32 i = 0; (RenderingCount < LocalMaxRendering) && (i < StreamingStatesToProcess.Num()); i++)
+		{
+			FComponentState* State = StreamingStatesToProcess[i];
+			check(State->Stage == EComponentStage::TextureStreaming);
+
+			if (State->AreTexturesStreamedIn())
+			{
+				KickOffRenderAndReadback(*State);
+				bChanged = true;
+			}
+		}
+		StreamingStatesToProcess.Reset();
+	}
+
+	if (ComponentsToRemoveFoliageInstances.Num() > 0)
+	{
+		SCOPE_CYCLE_COUNTER(STAT_RemoveGrassInstances);
+		World->GetSubsystem<ULandscapeSubsystem>()->RemoveGrassInstances(&ComponentsToRemoveFoliageInstances);
+	}
+
+	return bChanged;
+}
+
+void FLandscapeGrassMapsBuilder::StartPrioritizedGrassMapGeneration(const TArray<FVector>& Cameras, int32 MaxComponentsToStart)
+{
+	SCOPE_CYCLE_COUNTER(STAT_PrioritizePendingGrassMaps);
+
+	// no point in calling this if there are no cameras -- we can't calculate priority
+	check(Cameras.Num() > 0);
+
+	float DiscardDistanceScale = GGrassMapGuardBandDiscardMultiplier * GGrassCullDistanceScale;
+	float MustHaveDistanceScale = GGrassMapGuardBandMultiplier * GGrassCullDistanceScale;
+
+	// update pending component priorities (distances)
+	check(PendingComponentsHeap.Num() == PendingCount);
+	if (PendingCount)
+	{
+		// We update the priority of one element in each heap level.
+		// Because the heap is ordered by distance, this approximately updates
+		// closer elements more often than distant elements.
+		// The closest element is updated every frame, the second and third every other frame, 4-7 every fourth frame, etc.
+		// This way we update at most Log2(N) elements each frame.
+		for (int32 LevelSize = 1; LevelSize <= 65536; LevelSize += LevelSize)
+		{
+			// this should select successive elements on the given Level as the counter is incremented
+			int32 UpdateIndex = (PendingUpdateAmortizationCounter & (LevelSize-1)) + (LevelSize-1);
+			if (UpdateIndex >= PendingCount)
+			{
+				break;
+			}
+			PendingComponentsHeap[UpdateIndex].UpdatePriorityDistance(Cameras);
+		}
+		PendingUpdateAmortizationCounter++;
+
+		// also check if any have a negative priority, which indicates newly pending components, and immediately calculate priority for those
+		// TODO [chris.tchou] : if the camera jumps for some reason, we should recalculate all priorities immediately.
+		for (FPendingComponent& Pending : PendingComponentsHeap)
+		{
+			if (Pending.PriorityKey < 0.0)
+			{
+				Pending.UpdatePriorityDistance(Cameras);
+			}
+		}
+	}
+
+	// re-heapify to reflect updated priorities
+	{
+		TRACE_CPUPROFILER_EVENT_SCOPE(Heapify)
+		PendingComponentsHeap.Heapify();
+	}
+
+	// now pull as many elements off of the heap as we need
+	while ((MaxComponentsToStart > 0) && PendingComponentsHeap.Num() > 0)
+	{
+		FPendingComponent Pending = PendingComponentsHeap[0];
+		FComponentState* State = Pending.State;
+		ULandscapeComponent* Component = State->Component;
+		check(State->Stage == EComponentStage::Pending);
+
+		// runtime generation doesn't generate past this distance
+		if (GGrassMapUseRuntimeGeneration)
+		{
+			double MaxSpawnDistance = Component->GrassTypeSummary.MaxInstanceDiscardDistance * MustHaveDistanceScale;
+			if (Pending.PriorityKey > MaxSpawnDistance * MaxSpawnDistance)
+			{
+				// the closest component is not close enough to generate yet
+				break;
+			}
+		}
+
+		PendingComponentsHeap.HeapPop(Pending, false);
+		if (StartGrassMapGeneration(*State, false))
+		{
+			MaxComponentsToStart--;
+		}
+	}
+
+	TotalComponentsWaitingCount = PendingComponentsHeap.Num();
+}
+
+// called when components are registered to the world	
+void FLandscapeGrassMapsBuilder::RegisterComponent(ULandscapeComponent* Component)
+{
+	check(Component);
+	if (FComponentState* State = ComponentStates.FindRef(Component))
+	{
+		GRASS_DEBUG_LOG(TEXT("Re-Register %s (%d total)"), *Component->GetName(), ComponentStates.Num());
+		State->Component = Component;
+		State->TickCount = 0;
+	}
+	else
+	{
+		FComponentState* NewState = new(StatePoolAllocator.Allocate()) FComponentState(Component);
+
+		ComponentStates.Add(Component, NewState);
+		PendingComponentsHeap.Add(FPendingComponent(NewState));
+
+		PendingCount++;
+
+		GRASS_DEBUG_LOG(TEXT("Register %s (%d total)"), *Component->GetName(), ComponentStates.Num());
+	}
+}
+
+void FLandscapeGrassMapsBuilder::UnregisterComponent(ULandscapeComponent* Component)
+{
+	check(Component);
+	if (FComponentState* State = ComponentStates.FindRef(Component))
+	{
+		GRASS_DEBUG_LOG(TEXT("Unregister %s"), *Component->GetName());
+		State->Component = nullptr;				// we should no longer access the component, it may disappear
+		State->TickCount = 0;					// track how long since unregistered
+		// After ~2 ticks, Update will CancelAndEvict to clean up the remaining state
+	}
+	else
+	{
+		GRASS_DEBUG_LOG(TEXT("Unregister %s - NOT REGISTERED"), *Component->GetName());
+	}
+}
+
+// false if this program instance will never be able to render grass
+bool FLandscapeGrassMapsBuilder::CanEverRender() const
+{
+	return FApp::CanEverRender() && !GUsingNullRHI;
+}
+
+// false if the world can not currently render the grass (but this may change later, for example if preview modes are modified)
+bool FLandscapeGrassMapsBuilder::CanCurrentlyRender() const
+{
+	if (CanEverRender())
+	{
+		// GPU scene is required by landscape fixed grid vertex factory
+		EShaderPlatform ShaderPlatform = World->Scene->GetShaderPlatform();
+		ERHIFeatureLevel::Type FeatureLevel = World->GetFeatureLevel();
+		return UseGPUScene(ShaderPlatform, FeatureLevel);
+	}
+	return false;
+}
+
+
+void FLandscapeGrassMapsBuilder::AmortizedUpdateGrassMaps(
+	const TArray<FVector>& Cameras,
+	bool bPrioritizeCreation,
+	bool bAllowStartGrassMapGeneration)
+{
+#if !WITH_EDITOR
+	if (!GGrassMapUseRuntimeGeneration) // in cooked builds, we don't run any updates at all unless runtime generation is enabled
+	{
+		return;
+	}
+#endif // WITH_EDITOR
+
+	if (!CanEverRender())
+	{
+		return; // if we can never ever render, don't bother to do anything here
+	}
+
+	int32 AmortizedMaxStreaming = GGrassMapMaxComponentsStreaming;
+	int32 AmortizedMaxRendering = GGrassMapMaxComponentsRendering;
+
+	if (bPrioritizeCreation && (GGrassMapPrioritizedMultiplier > 1))
+	{
+		AmortizedMaxStreaming *= GGrassMapPrioritizedMultiplier;
+		AmortizedMaxRendering *= GGrassMapPrioritizedMultiplier;
+	}
+
+	UpdateTrackedComponents(Cameras, AmortizedMaxRendering);
+
+	// no point in looking to start new grass map generation if nothing is pending, if grass is disabled or there are no cameras
+	if (bAllowStartGrassMapGeneration && PendingCount > 0 && GGrassEnable && Cameras.Num() > 0)
+	{
+		// check our pipeline limits to make sure we have room to start components
+		int32 AvailableStreamingSlots = AmortizedMaxStreaming - StreamingCount;
+		StartPrioritizedGrassMapGeneration(Cameras, AvailableStreamingSlots);
+	}
+}
+
+
+bool FLandscapeGrassMapsBuilder::BuildGrassMapsNowForComponents(
+	TArrayView<TObjectPtr<ULandscapeComponent>> LandscapeComponents, FScopedSlowTask* SlowTask, bool bMarkDirty)
+{
+	TRACE_CPUPROFILER_EVENT_SCOPE(BuildGrassMapsNowForComponents);
+	int32 MaxStreamingRendering = FMath::Max(GGrassMapMaxComponentsForBlockingUpdate, 1);
+
+	if (LandscapeComponents.IsEmpty())
+	{
+		return true;
+	}
+
+	if (!CanCurrentlyRender())
+	{
+		return false; // can't build grass maps without rendering, unfortunately
+	}
+
+	auto UpdateProgress = [SlowTask](int Increment)
+	{
+		if (SlowTask && Increment && ((SlowTask->CompletedWork + Increment) <= SlowTask->TotalAmountOfWork))
+		{
+			SlowTask->EnterProgressFrame(static_cast<float>(Increment), FText::Format(LOCTEXT("GrassMaps_BuildGrassMapsProgress", "Building Grass Map {0} of {1})"), FText::AsNumber(SlowTask->CompletedWork), FText::AsNumber(SlowTask->TotalAmountOfWork)));
+		}
+	};
+
+	// using an empty cameras array causes all distance checks to return 0, so grass maps won't be evicted
+	TArray<FVector> EmptyCamerasArray;
+	int32 LastUpToDateCount = 0;
+	int32 UpToDateCount = 0;
+	
+	double StartTime = FPlatformTime::Seconds();
+	double LastChangeTime = StartTime;
+	while (UpToDateCount != LandscapeComponents.Num())
+	{
+		// ensure we are making progress within a reasonable amount of time TODO [chris.tchou] there should be a better way to detect non-progress here
+		double CurTime = FPlatformTime::Seconds();
+		if (CurTime > LastChangeTime + 15.0)
+		{
+			UE_LOG(LogGrass, Error, TEXT("ERROR: BuildGrassMapsNowForComponents() took too long, grass maps are not up to date"));
+			break;
+		}
+
+		// update all components that are tracked (without evicting)
+		bool bChanged = UpdateTrackedComponents(EmptyCamerasArray, MaxStreamingRendering);
+
+		UpToDateCount = 0;
+		int32 AvailableStreamingSlots = MaxStreamingRendering - StreamingCount; // here we don't limit by overall population count
+		for (ULandscapeComponent* Component : LandscapeComponents)
+		{
+			FComponentState* State = ComponentStates.FindRef(Component);
+			check(State); // components must be registered
+
+			if (State->Stage == EComponentStage::Pending)
+			{
+				// Start tracking to kick off the build process
+				const bool bForceCompileShaders = true;
+				if ((AvailableStreamingSlots > 0) && StartGrassMapGeneration(*State, bForceCompileShaders))
+				{
+					// modification isn't complete yet, but convenient to dirty the package when starting the process here
+					if (bMarkDirty)
+					{
+						Component->MarkPackageDirty();
+					}
+					AvailableStreamingSlots--;
+					bChanged = true;
+				}
+			}
+			if (State->Stage == EComponentStage::GrassMapsPopulated)
+			{
+				check(Component->GrassData->HasValidData()); // guaranteed by UpdateTrackedComponents()
+#if WITH_EDITOR
+				// guaranteed by UpdateTrackedComponents() (unless GrassData GenerationHash gets out of sync with the tracked state somehow)
+				check(Component->ComputeGrassMapGenerationHash() == Component->GrassData->GenerationHash);
+#endif // WITH_EDITOR
+				UpToDateCount++;
+			}
+		}
+
+		if (LastUpToDateCount != UpToDateCount)
+		{
+			UpdateProgress(UpToDateCount - LastUpToDateCount);
+			LastUpToDateCount = UpToDateCount;
+		}
+
+		if (bChanged)
+		{
+			LastChangeTime = FPlatformTime::Seconds();
+		}
+
+		// Queue up the gpu commands on the render thread, so the GPU can start working on them.
+		UE::Landscape::SubmitGPUCommands(/* bBlockUntilComplete =  */ false);
+
+		// Blocking texture streaming update.  At least the GPU should be working while we wait on streaming here.
+		// TODO [chris.tchou] : ideally this would be a non-blocking streaming update tick, so we can react to other updates finishing
+		TextureStreamingManager.WaitForTextureStreaming();
+	}
+
+	return (UpToDateCount == LandscapeComponents.Num());
+}
+
+bool FLandscapeGrassMapsBuilder::CancelAndEvict(FComponentState& State)
+{
+	TRACE_CPUPROFILER_EVENT_SCOPE(CancelAndEvict);
+
+	// handle transitioning from any stage to EComponentState::Pending
+	switch (State.Stage)
+	{
+	case EComponentStage::Pending:
+		check(PendingCount > 0);
+		break;
+	case EComponentStage::NotReady:
+		{
+			check(NotReadyCount > 0);
+			NotReadyCount--;
+			PendingCount++;
+		}
+		break;
+	case EComponentStage::TextureStreaming:
+		{
+			RemoveTextureStreamingRequests(State);
+			check(StreamingCount > 0);
+			StreamingCount--;
+			PendingCount++;
+		}
+		break;
+	case EComponentStage::Rendering:
+		{
+			RemoveTextureStreamingRequests(State);
+			if (State.ActiveRender != nullptr)
+			{
+				// calling update ensures it is pushed forward if there is still a readback in progress
+				if (!State.ActiveRender->CheckAndUpdateAsyncReadback())
+				{
+					// we can't cancel yet.. must wait for the readback to complete
+					return false;
+				}
+				delete State.ActiveRender;
+				State.ActiveRender = nullptr;
+			}
+			check(RenderingCount > 0);
+			RenderingCount--;
+			PendingCount++;
+		}
+		break;
+	case EComponentStage::GrassMapsPopulated:
+		{
+			ULandscapeComponent* Component = State.Component;
+		
+			if (Component == nullptr)
+			{
+				// component was unregistered. Wait a few ticks to see if it comes back before fully evicting.
+				if (State.TickCount < 2)
+				{
+					return false;
+				}
+			}
+
+			// if the component is still around, clear any existing grass data from it
+			if (Component && Component->GrassData->HasValidData())
+			{
+				Component->RemoveGrassMap();
+			}
+			check(PopulatedCount > 0);
+			PopulatedCount--;
+			PendingCount++;
+		}
+		break;
+	default:
+		check(false);	// unreachable
+		break;
+	}
+
+	check(State.ActiveRender == nullptr);
+	check(State.TexturesToStream.IsEmpty());
+
+	if (State.Stage != EComponentStage::Pending)
+	{
+		GRASS_DEBUG_LOG(TEXT("CancelAndEvict %s (from %d after %d ticks) p:%d S:%d R:%d P:%d NR:%d T:%d"),
+			State.Component ? *State.Component->GetName() : TEXT("<REMOVED>"),
+			State.Stage,
+			State.TickCount,
+			PendingCount,
+			StreamingCount,
+			RenderingCount,
+			PopulatedCount,
+			NotReadyCount,
+			ComponentStates.Num());
+
+		// back to pending state with you!
+		State.Stage = EComponentStage::Pending;
+		State.TickCount = 0;
+		if (State.Component != nullptr)
+		{
+			// don't bother to add if component is null as we will just have to remove it immediately in the deallocate
+			PendingComponentsHeap.Add(FPendingComponent(&State));
+		}
+	}
+
+	return true;
+}
+
+bool FLandscapeGrassMapsBuilder::StartGrassMapGeneration(FComponentState& State, bool bForceCompileShaders)
+{
+	TRACE_CPUPROFILER_EVENT_SCOPE(StartGrassMapGeneration);
+	check(State.Stage == EComponentStage::Pending);
+
+	ULandscapeComponent* Component = State.Component;
+
+	if (World->IsGameWorld()) // including PIE
+	{
+		// if runtime grass generation is disabled, go straight to not ready
+		if (!Component->GetLandscapeProxy()->bUseRuntimeGrassMapGeneration)
+		{
+			PendingToNotReady(State);
+			return false;
+		}
+	}
+	else if (World->IsEditorWorld())
+	{
+#if WITH_EDITOR
+		// recalculate hashes
+		UMaterialInterface* Material = Component->GetLandscapeMaterial();
+		if (Material)
+		{
+			State.GrassMapGenerationHash = UE::Landscape::ComputeGrassMapGenerationHash(Component, Material);
+			const TArray<TObjectPtr<ULandscapeGrassType>>& GrassTypes = Material->GetCachedExpressionData().GrassTypes;
+			State.GrassInstanceGenerationHash = UE::Landscape::ComputeGrassInstanceGenerationHash(State.GrassMapGenerationHash, GrassTypes);
+		}
+
+		// in editor, if the existing grass data is valid and has a matching hash, then we can skip straight to Populated
+		if (Component->GrassData->HasValidData() &&
+			Component->GrassData->GenerationHash == State.GrassMapGenerationHash)
+		{
+			PendingToPopulatedFastPathAlreadyHasData(State);
+			return true;
+		}
+#endif // WITH_EDITOR
+	}
+
+	// handle the easy case of empty grass types; skip directly to the populated state without running the pipeline
+	if (Component->GetGrassTypes().IsEmpty())
+	{
+		PendingToPopulatedFastPathNoGrass(State);
+		return true;
+	}
+	
+	// if we can't currently render, it's not ready
+	if (!UE::Landscape::CanRenderGrassMap(Component))
+	{
+#if WITH_EDITOR
+		if (bForceCompileShaders)
+		{
+			UE::Landscape::CompileGrassMapShader(Component);
+		}
+
+		if (!UE::Landscape::CanRenderGrassMap(Component))
+#endif // WITH_EDITOR
+		{
+			PendingToNotReady(State);
+			return false;
+		}
+	}
+
+	PendingToStreaming(State);
+	return true;
+}
+
+void FLandscapeGrassMapsBuilder::PendingToNotReady(FComponentState& State)
+{
+	check(PendingCount > 0);
+	PendingCount--;
+	State.Stage = EComponentStage::NotReady;
+	NotReadyCount++;
+
+	RemoveFromPendingComponentHeap(&State);
+
+	GRASS_DEBUG_LOG(TEXT("%s Pending -> NotReady (after %d ticks) p:%d S:%d R:%d P:%d NR:%d T:%d"),
+		*State.Component->GetName(),
+		State.TickCount,
+		PendingCount,
+		StreamingCount,
+		RenderingCount,
+		PopulatedCount,
+		NotReadyCount,
+		ComponentStates.Num());
+
+	State.TickCount = 0;
+}
+
+void FLandscapeGrassMapsBuilder::PendingToPopulatedFastPathAlreadyHasData(FComponentState& State)
+{
+	check(PendingCount > 0);
+	PendingCount--;
+	State.Stage = EComponentStage::GrassMapsPopulated;
+	PopulatedCount++;
+
+	RemoveFromPendingComponentHeap(&State);
+
+	GRASS_DEBUG_LOG(TEXT("%s Pending -> Populated(Existing) (after %d ticks) p:%d S:%d R:%d P:%d NR:%d T:%d"),
+		*State.Component->GetName(),
+		State.TickCount,
+		PendingCount,
+		StreamingCount,
+		RenderingCount,
+		PopulatedCount,
+		NotReadyCount,
+		ComponentStates.Num());
+
+	State.TickCount = 0;
+}
+
+void FLandscapeGrassMapsBuilder::PendingToPopulatedFastPathNoGrass(FComponentState& State)
+{
+	ULandscapeComponent* Component = State.Component;
+
+	TUniquePtr<FLandscapeComponentGrassData> NewGrassData = MakeUnique<FLandscapeComponentGrassData>();
+	Component->GrassData = MakeShareable(NewGrassData.Release());
+#if WITH_EDITOR
+	Component->GrassData->GenerationHash = State.GrassMapGenerationHash;
+#endif // WITH_EDITOR
+	Component->GrassData->NumElements = 0;
+
+	check(PendingCount > 0);
+	PendingCount--;
+	State.Stage = EComponentStage::GrassMapsPopulated;
+	PopulatedCount++;
+
+	RemoveFromPendingComponentHeap(&State);
+
+	GRASS_DEBUG_LOG(TEXT("%s Pending -> Populated(Empty) (after %d ticks) p:%d S:%d R:%d P:%d NR:%d T:%d"),
+		*Component->GetName(),
+		State.TickCount,
+		PendingCount,
+		StreamingCount,
+		RenderingCount,
+		PopulatedCount,
+		NotReadyCount,
+		ComponentStates.Num());
+
+	State.TickCount = 0;
+}
+
+void FLandscapeGrassMapsBuilder::PendingToStreaming(FComponentState& State)
+{
+	TRACE_CPUPROFILER_EVENT_SCOPE(PendingToStreaming);
+	ULandscapeComponent* Component = State.Component;
+
+	// determine which textures we need to stream by inspecting the material, kick off streaming requests for them
+	// TODO [chris.tchou] : also grab asset textures from the material.
+	State.TexturesToStream.Add(Component->GetHeightmap());
+
+	ERHIFeatureLevel::Type FeatureLevel = World->GetFeatureLevel();
+	for (UTexture2D* WeightmapTexture : Component->GetRenderedWeightmapTexturesForFeatureLevel(FeatureLevel))
+	{
+		State.TexturesToStream.Add(WeightmapTexture);
+	}
+
+	// Request streaming for our textures
+	for (UTexture* Texture : State.TexturesToStream)
+	{
+		TextureStreamingManager.RequestTextureFullyStreamedIn(Texture, /* bWaitForStreaming= */ false);
+	}
+
+	check(PendingCount > 0);
+	PendingCount--;
+	State.Stage = EComponentStage::TextureStreaming;
+	StreamingCount++;
+
+	RemoveFromPendingComponentHeap(&State);
+
+	GRASS_DEBUG_LOG(TEXT("%s Pending -> Streaming (after %d ticks) p:%d S:%d R:%d P:%d NR:%d U:%d"),
+		*Component->GetName(),
+		State.TickCount,
+		PendingCount,
+		StreamingCount,
+		RenderingCount,
+		PopulatedCount,
+		NotReadyCount,
+		ComponentStates.Num());
+
+	State.TickCount = 0;
+}
+
+void FLandscapeGrassMapsBuilder::RemoveFromPendingComponentHeap(FComponentState* State)
+{
+	for (int32 Index = 0; Index < PendingComponentsHeap.Num(); Index++)
+	{
+		if (PendingComponentsHeap[Index].State == State)
+		{
+			PendingComponentsHeap.RemoveAtSwap(Index);
+			break;
+		}
+	}
+}
+
+void FLandscapeGrassMapsBuilder::KickOffRenderAndReadback(FComponentState& State)
+{
+	SCOPE_CYCLE_COUNTER(STAT_RenderGrassMap);
+
+	check(State.Stage == EComponentStage::TextureStreaming);
+	check(StreamingCount > 0);
+	StreamingCount--;
+
+	TArray<int32> HeightMips;
+
+	ULandscapeComponent* Component = State.Component;
+
+	const bool bBakeMaterialPositionOffsetIntoCollision = (Component->GetLandscapeProxy() && Component->GetLandscapeProxy()->bBakeMaterialPositionOffsetIntoCollision);
+	if (bBakeMaterialPositionOffsetIntoCollision)
+	{
+		if (Component->CollisionMipLevel > 0)
+		{
+			HeightMips.Add(Component->CollisionMipLevel);
+		}
+		if (Component->SimpleCollisionMipLevel > Component->CollisionMipLevel)
+		{
+			HeightMips.Add(Component->SimpleCollisionMipLevel);
+		}
+	}
+
+	check(State.ActiveRender == nullptr);
+	constexpr bool bInNeedsGrassmap = true;
+	constexpr bool bInNeedsHeightmap = true;
+	constexpr bool bUseAsyncReadback = true;
+	State.ActiveRender = new FLandscapeGrassWeightExporter(Component->GetLandscapeProxy(), { Component }, bInNeedsGrassmap, bInNeedsHeightmap, MoveTemp(HeightMips), bUseAsyncReadback);
+	check(State.ActiveRender != nullptr);
+
+	State.Stage = EComponentStage::Rendering;
+	RenderingCount++;
+
+	GRASS_DEBUG_LOG(TEXT("%s Streaming -> Rendering (after %d ticks) p:%d S:%d R:%d P:%d NR:%d T:%d"),
+		*Component->GetName(),
+		State.TickCount,
+		PendingCount,
+		StreamingCount,
+		RenderingCount,
+		PopulatedCount,
+		NotReadyCount,
+		ComponentStates.Num());
+
+	State.TickCount = 0;
+}
+
+void FLandscapeGrassMapsBuilder::PopulateGrassDataFromReadback(FComponentState& State)
+{
+	SCOPE_CYCLE_COUNTER(STAT_PopulateGrassMap);
+
+	check(State.Stage == EComponentStage::Rendering);
+	check(RenderingCount > 0);
+	RenderingCount--;
+
+	// now that render is complete, we can drop the texture streaming requests and allow textures to stream out
+	RemoveTextureStreamingRequests(State);
+
+	check(State.ActiveRender != nullptr);
+	State.ActiveRender->ApplyResults();
+	delete State.ActiveRender;
+	State.ActiveRender = nullptr;
+
+	State.Stage = EComponentStage::GrassMapsPopulated;
+	PopulatedCount++;
+
+	GRASS_DEBUG_LOG(TEXT("%s Rendering -> Populated (after %d ticks) p:%d S:%d R:%d P:%d NR:%d T:%d"),
+		*State.Component->GetName(),
+		State.TickCount,
+		PendingCount,
+		StreamingCount,
+		RenderingCount,
+		PopulatedCount,
+		NotReadyCount,
+		ComponentStates.Num());
+
+	State.TickCount = 0;
+}
+
+void FLandscapeGrassMapsBuilder::RemoveTextureStreamingRequests(FComponentState& State)
+{
+	for (UTexture* Texture : State.TexturesToStream)
+	{
+		TextureStreamingManager.UnrequestTextureFullyStreamedIn(Texture);
+	}
+	State.TexturesToStream.Empty();
+}
+
+bool FLandscapeGrassMapsBuilder::FComponentState::AreTexturesStreamedIn() const
+{
+	for (UTexture* Texture : TexturesToStream)
+	{
+		if (!FLandscapeTextureStreamingManager::IsTextureFullyStreamedIn(Texture))
+			return false;
+	}
+	return true;
+}
+
+bool FLandscapeGrassMapsBuilder::FComponentState::IsRenderReadbackComplete() const
+{
+	check(ActiveRender != nullptr);
+	return ActiveRender->CheckAndUpdateAsyncReadback();
+}
+
+bool FLandscapeGrassMapsBuilder::FComponentState::IsBeyondEvictionRange(const TArray<FVector>& Cameras) const
+{
+	check(Stage == EComponentStage::GrassMapsPopulated);
+	FBoxSphereBounds WorldBounds = Component->CalcBounds(Component->GetComponentTransform());
+	float MinSqrDistanceToComponent = UE::Landscape::CalculateMinDistanceToCameras(Cameras, WorldBounds);
+	float DiscardDistanceScale = GGrassMapGuardBandDiscardMultiplier * GGrassCullDistanceScale;
+	float MinEvictDistance = Component->GrassTypeSummary.MaxInstanceDiscardDistance * DiscardDistanceScale;
+	return (MinSqrDistanceToComponent > MinEvictDistance * MinEvictDistance);
+}
+
+
+#if WITH_EDITOR
 
 void FLandscapeGrassMapsBuilder::Build()
 {
 	if (World)
 	{
-		int32 Count = GetOutdatedGrassMapCount();
-		FScopedSlowTask SlowTask(static_cast<float>(Count), (LOCTEXT("GrassMaps_BuildGrassMaps", "Building Grass maps")));
-		SlowTask.MakeDialog();
-
+		// iterate proxies, update grass types and create the list of components to build
+		TArray<TObjectPtr<ULandscapeComponent>> LandscapeComponentsToBuild;
 		for (TActorIterator<ALandscapeProxy> ProxyIt(World); ProxyIt; ++ProxyIt)
 		{
-			ProxyIt->BuildGrassMaps(&SlowTask);
+			check(!ProxyIt->HasAnyFlags(RF_ClassDefaultObject));
+
+			for (ULandscapeComponent* Component : ProxyIt->LandscapeComponents)
+			{
+				FComponentState* State = ComponentStates.FindRef(Component);
+
+				if (Component->UpdateGrassTypes() ||
+					State->Stage != EComponentStage::GrassMapsPopulated)	// TODO check up to date slow?
+				{
+					
+					Component->MarkPackageDirty();
+					LandscapeComponentsToBuild.Add(Component);
+				}
+			}
+
+			ProxyIt->UpdateGrassTypeSummary();
+		}
+
+		// build the grass maps
+		if (LandscapeComponentsToBuild.Num() > 0)
+		{
+			FScopedSlowTask SlowTask(static_cast<float>(LandscapeComponentsToBuild.Num()), (LOCTEXT("GrassMaps_BuildGrassMaps", "Building Grass maps")));
+			SlowTask.MakeDialog();
+
+			bool bMarkDirty = true;
+			BuildGrassMapsNowForComponents(LandscapeComponentsToBuild, &SlowTask, bMarkDirty);
+		}
+	}
+}
+
+int32 FLandscapeGrassMapsBuilder::CountOutdatedGrassMaps(const TArray<TObjectPtr<ULandscapeComponent>>& LandscapeComponents) const
+{
+	int32 ProxyOutdatedGrassMapCount = 0;
+	for (ULandscapeComponent* Component : LandscapeComponents)
+	{
+	#if WITH_EDITOR
+		if (Component->IsGrassMapOutdated())
+	#else
+		if (!Component->GrassData->HasValidData())
+	#endif // WITH_EDITOR
+		{
+			ProxyOutdatedGrassMapCount++;
 		}
 	}
+	return ProxyOutdatedGrassMapCount;
 }
 
 int32 FLandscapeGrassMapsBuilder::GetOutdatedGrassMapCount(bool bInForceUpdate) const
@@ -56,7 +1225,7 @@ int32 FLandscapeGrassMapsBuilder::GetOutdatedGrassMapCount(bool bInForceUpdate)
 			OutdatedGrassMapCount = 0;
 			for (TActorIterator<ALandscapeProxy> ProxyIt(World); ProxyIt; ++ProxyIt)
 			{
-				OutdatedGrassMapCount += ProxyIt->GetOutdatedGrassMapCount();
+				OutdatedGrassMapCount += CountOutdatedGrassMaps(ProxyIt->LandscapeComponents);
 			}
 		}
 	}
@@ -65,3 +1234,4 @@ int32 FLandscapeGrassMapsBuilder::GetOutdatedGrassMapCount(bool bInForceUpdate)
 
 #endif // WITH_EDITOR
 #undef LOCTEXT_NAMESPACE
+#undef GRASS_DEBUG_LOG
\ No newline at end of file
diff --git a/Engine/Source/Runtime/Landscape/Private/LandscapeGrassMapsBuilder.h b/Engine/Source/Runtime/Landscape/Private/LandscapeGrassMapsBuilder.h
--- a/Engine/Source/Runtime/Landscape/Private/LandscapeGrassMapsBuilder.h
+++ b/Engine/Source/Runtime/Landscape/Private/LandscapeGrassMapsBuilder.h
@@ -2,9 +2,13 @@
 #pragma once
 
 #include "CoreMinimal.h"
+#include "UObject\WeakObjectPtrTemplates.h"
 #include "LandscapeComponent.h"
+#include "LandscapeTextureStreamingManager.h"
+#include "Containers\AllocatorFixedSizeFreeList.h"
 
-#if WITH_EDITOR
+class FLandscapeGrassWeightExporter;
+struct FScopedSlowTask;
 
 /**
  * Helper class used to Build or monitor outdated Grass maps of a world
@@ -12,13 +16,216 @@
 class FLandscapeGrassMapsBuilder
 {
 public:
-	LANDSCAPE_API FLandscapeGrassMapsBuilder(UWorld* InWorld);
+	LANDSCAPE_API FLandscapeGrassMapsBuilder(UWorld* InWorld, FLandscapeTextureStreamingManager& InTextureStreamingManager);
+	~FLandscapeGrassMapsBuilder();
+
+#if WITH_EDITOR
+	/** Synchronously build all grass maps on all proxies in the world */
 	LANDSCAPE_API void Build();
+	
+	/** Return the number of landscape components in the world that have grass map data that is not up to date (NOT including components missing grass map data altogether) */
 	LANDSCAPE_API int32 GetOutdatedGrassMapCount(bool bInForceUpdate = true) const;
+#endif // WITH_EDITOR
+
+	// count components that have valid data but whose generation hash does not match
+	int32 CountOutdatedGrassMaps(const TArray<TObjectPtr<ULandscapeComponent>>& LandscapeComponents) const;
+
+	// called when components are registered to the world
+	void RegisterComponent(ULandscapeComponent* Component);
+	// called when components are unregistered from the world.
+	void UnregisterComponent(ULandscapeComponent* Component);
+
+	// get the number of grass maps that are still waiting to render, as of the last AmortizedUpdateGrassMaps()
+	int32 GetTotalGrassMapsWaitingToRender() const { return TotalComponentsWaitingCount; }
+
+	// Amortized Update of Grass Maps near the specified cameras.
+	// if Cameras is empty, it considers all distances to be zero for update purposes:
+	//   - it won't evict anything for distance
+	//   - it will not start tracking (building grass maps) for any components
+	void AmortizedUpdateGrassMaps(const TArray<FVector>& Cameras, bool bPrioritizeCreation, bool bAllowStartGrassMapGeneration);
+
+	// synchronous build of grassmaps for a specific set of components
+	// returns true if it successfully makes all LandscapeComponents have up to date Grass Maps.
+	bool BuildGrassMapsNowForComponents(TArrayView<TObjectPtr<ULandscapeComponent>> LandscapeComponents, FScopedSlowTask* SlowTask, bool bMarkDirty);
+
 private:
-	UWorld* World;
-	mutable int32 OutdatedGrassMapCount;
-	mutable double GrassMapsLastCheckTime;
-};
 
+	// false if this program instance will never be able to render grass
+	bool CanEverRender() const;
+
+	// false if the world can not currently render the grass (but this may change later, for example if preview modes are modified)
+	bool CanCurrentlyRender() const;
+
+	// Update all of the non-pending components.
+	// If passed an empty Cameras array, distances are calculated as zero (i.e. it won't evict for distance)
+	// returns true if any components changed states
+	bool UpdateTrackedComponents(const TArray<FVector>& Cameras, int32 LocalMaxRendering);
+
+	// Start the grass map generation process on pending components in priority order
+	// (based on distance from the given Camera set) -- Cameras must not be empty.
+	void StartPrioritizedGrassMapGeneration(const TArray<FVector>& Cameras, int32 MaxComponentsToStart);
+	
+	enum class EComponentStage
+	{
+		Pending,							// initial state; the component is waiting to start the generation process.
+		NotReady,							// tried to start generation process, but either no grass types exist or the material is not ready. wait for that to change.
+		TextureStreaming,					// texture streaming was requested.  wait for the mips to be available
+		Rendering,							// GPU render commands were sent -- waiting for async readback to complete
+		GrassMapsPopulated,					// grass maps are built and are ready to create instances
+	};
+
+	struct FComponentState
+	{
+	public:
+		// this pointer is valid as long as the Component is registered
+		// (the FComponentState itself may continue to exist after unregistration, until all resources are cleaned up)
+		ULandscapeComponent* Component = nullptr;
+
+		EComponentStage Stage = EComponentStage::Pending;
+
+		// counts the number of ticks this component has remained in the current stage
+		int32 TickCount = 0;
+
+		#if WITH_EDITOR
+		// the hashes of dependencies when this component's grass map was built, for tracking automatic invalidation
+		uint32 GrassMapGenerationHash = 0;
+		uint32 GrassInstanceGenerationHash = 0;
+		#endif // WITH_EDITOR
+
+		// list of textures to stream prior to rendering this component (valid only in TextureSreaming and Rendering stages)
+		TArray<UTexture*> TexturesToStream;
+		
+		// the active render (valid only in Rendering stage)
+		FLandscapeGrassWeightExporter* ActiveRender = nullptr;
+
+		FComponentState(ULandscapeComponent* Component);
+
+		bool AreTexturesStreamedIn() const;
+		bool IsRenderReadbackComplete() const;
+		bool IsBeyondEvictionRange(const TArray<FVector>& Cameras) const;
+	};
+
+	// components in the Pending state are also tracked in a priority heap, using this structure
+	struct FPendingComponent
+	{
+		FPendingComponent(FComponentState* InComponentState)
+			: State(InComponentState)
+		{}
+
+		// comparison operator used for Heap functions -- smaller keys have priority over larger keys
+		bool operator<(const FPendingComponent& Other) const
+		{
+			return PriorityKey < Other.PriorityKey;
+		}
+
+		void UpdatePriorityDistance(const TArray<FVector>& Cameras);
+
+		FComponentState* State = nullptr;
+		double PriorityKey = -1.0;
+	};
+
+	// structure to do an amortized update of a set of update elements in a (conceptual) array
+	struct FAmortizedUpdate
+	{
+		// the first index to update this tick (inclusive)
+		int32 FirstIndex = 0;
+
+		// the last index to update this tick (exclusive)
+		int32 LastIndex = 0;
+
+		// Call at the beginning of the update tick to set up how many updates you want to run this tick
+		// UpdateElementCount is the total number of update elements (to wrap when we hit the end of the possibility space)
+		void StartUpdateTick(int32 UpdateElementCount, int32 MaxUpdatesThisFrame)
+		{
+			check(MaxUpdatesThisFrame >= 0);
+			FirstIndex = LastIndex;
+			if (FirstIndex >= UpdateElementCount)
+			{
+				FirstIndex = 0;
+			}
+			check(FirstIndex >= 0);
+			LastIndex = FirstIndex + MaxUpdatesThisFrame;
+		}
+
+		// return true if the given element should update
+		bool ShouldUpdate(int32 UpdateElementIndex)
+		{
+			if ((UpdateElementIndex >= FirstIndex) && (UpdateElementIndex < LastIndex))
+			{
+				return true;
+			}
+			return false;
+		}
+
+		// call if an update element is deleted from the array, to update the valid ranges
+		void HandleDeletion(int32 DeletedUpdateElementIndex)
+		{
+			if (DeletedUpdateElementIndex < FirstIndex)
+			{
+				FirstIndex--;
+			}
+			if (DeletedUpdateElementIndex < LastIndex)
+			{
+				LastIndex--;
+			}
+
+		}
+	};
+
+	UWorld* World = nullptr;
+
+	// counts of how many components are currently in each stage
+	int32 PendingCount = 0;
+	int32 NotReadyCount = 0;
+	int32 StreamingCount = 0;
+	int32 RenderingCount = 0;
+	int32 PopulatedCount = 0;
+
+	// number of components that need to render but are waiting (as of the last call to StartTrackingComponents())
+	int32 TotalComponentsWaitingCount = 0;
+
+	// store the grass map state of each registered (or recently unregistered) component
+	TMap<ULandscapeComponent*, FComponentState*> ComponentStates;
+
+	static TAllocatorFixedSizeFreeList<sizeof(FComponentState), 32> StatePoolAllocator;
+
+	TArray<FPendingComponent> PendingComponentsHeap;
+	int32 PendingUpdateAmortizationCounter = 0;
+
+	// state to amortize the update of components
+	FAmortizedUpdate AmortizedUpdate;
+
+	// system to manage texture streaming requests
+	FLandscapeTextureStreamingManager& TextureStreamingManager;
+
+	// tries to cancel any in flight operations and transition back to the Pending state
+	bool CancelAndEvict(FComponentState& State);
+
+	// try to kick off the grass map generation pipeline -- returns true if texture streaming was started, false if the shaders are not ready, or if it shortcut to populated because data exists already
+	bool StartGrassMapGeneration(FComponentState& State, bool bForceCompileShaders);
+
+	// once textures are streamed, this kicks off the grass data render, and async GPU readback
+	void KickOffRenderAndReadback(FComponentState& State);
+
+	// once the GPU readback is complete, this starts processing the data to generate a GrassData structure
+	void ProcessGrassDataAsync(FComponentState& State);
+
+	// once the GPU readback is complete, this populates the grass data on the component (and cancels texture stream requests)
+	void PopulateGrassDataFromReadback(FComponentState& State);
+		void RemoveTextureStreamingRequests(FComponentState& State);
+
+	// state transition helpers
+	void PendingToNotReady(FComponentState& State);
+	void PendingToPopulatedFastPathAlreadyHasData(FComponentState& State);
+	void PendingToPopulatedFastPathNoGrass(FComponentState& State);
+	void PendingToStreaming(FComponentState& State);
+	void RemoveFromPendingComponentHeap(FComponentState* State);
+
+#if WITH_EDITOR
+	// cached count of how many grass maps are outdated, and the last time we calculated that value
+	mutable int32 OutdatedGrassMapCount = 0;
+	mutable double GrassMapsLastCheckTime = 0.0;
 #endif // WITH_EDITOR
+};
+
+
diff --git a/Engine/Source/Runtime/Landscape/Private/LandscapeGrassWeightExporter.cpp b/Engine/Source/Runtime/Landscape/Private/LandscapeGrassWeightExporter.cpp
--- a/Engine/Source/Runtime/Landscape/Private/LandscapeGrassWeightExporter.cpp
+++ b/Engine/Source/Runtime/Landscape/Private/LandscapeGrassWeightExporter.cpp
@@ -13,17 +13,19 @@
 #include "SimpleMeshDrawCommandPass.h"
 #include "TextureResource.h"
 #include "RenderCaptureInterface.h"
+#include "ShaderPlatformCachedIniValue.h"
 
-#if WITH_EDITOR
 class FLandscapeGrassWeightVS;
 class FLandscapeGrassWeightPS;
 
-int32 RenderCaptureNextGrassmapDraws = 0;
+int32 GRenderCaptureNextGrassmapDraws = 0;
 static FAutoConsoleVariableRef CVarRenderCaptureNextGrassmapDraws(
-	TEXT("landscape.RenderCaptureNextGrassmapDraws"),
-	RenderCaptureNextGrassmapDraws,
+	TEXT("grass.GrassMap.RenderCaptureNextDraws"),
+	GRenderCaptureNextGrassmapDraws,
 	TEXT("Trigger render captures during the next N grassmap draw calls."));
 
+extern int32 GGrassMapAlwaysBuildRuntimeGenerationResources;
+
 BEGIN_SHADER_PARAMETER_STRUCT(FLandscapeGrassPassParameters, )
 	SHADER_PARAMETER_STRUCT_REF(FViewUniformShaderParameters, View)
 	SHADER_PARAMETER_RDG_UNIFORM_BUFFER(FSceneUniformParameters, Scene)
@@ -38,6 +40,117 @@ public:
 	FVector2f RenderOffset;
 };
 
+static bool ShouldCacheLandscapeGrassShaders(const FMeshMaterialShaderPermutationParameters& Parameters)
+{
+	static FShaderPlatformCachedIniValue<int32> GrassMapsUseRuntimeGenerationPerPlatform(TEXT("grass.GrassMap.UseRuntimeGeneration"));
+
+	bool bIsEditorPlatform = 
+		IsFeatureLevelSupported(Parameters.Platform, ERHIFeatureLevel::SM5) &&
+		EnumHasAllFlags(Parameters.Flags, EShaderPermutationFlags::HasEditorOnlyData);
+
+	bool bShouldBuildForPlatform = 
+		bIsEditorPlatform ||
+		GGrassMapAlwaysBuildRuntimeGenerationResources ||
+		(GrassMapsUseRuntimeGenerationPerPlatform.Get(Parameters.Platform) != 0);
+
+	bool bIsFixedGridVertexFactory =
+		Parameters.VertexFactoryType == FindVertexFactoryType(FName(TEXT("FLandscapeFixedGridVertexFactory"), FNAME_Find));
+
+	// We only need grass weight shaders for Landscape fixed grid vertex factories
+	// And only for platforms that have runtime generation enabled or are editor platforms (or if we are always building resources)
+	return  (Parameters.MaterialParameters.bIsUsedWithLandscape || Parameters.MaterialParameters.bIsSpecialEngineMaterial) &&
+			bIsFixedGridVertexFactory &&
+			bShouldBuildForPlatform;
+}
+
+class FLandscapeGrassWeightVS : public FMeshMaterialShader
+{
+	DECLARE_SHADER_TYPE(FLandscapeGrassWeightVS, MeshMaterial);
+
+	LAYOUT_FIELD(FShaderParameter, RenderOffsetParameter);
+
+protected:
+
+	FLandscapeGrassWeightVS()
+	{}
+
+	FLandscapeGrassWeightVS(const FMeshMaterialShaderType::CompiledShaderInitializerType& Initializer)
+		: FMeshMaterialShader(Initializer)
+	{
+		RenderOffsetParameter.Bind(Initializer.ParameterMap, TEXT("RenderOffset"));
+		PassUniformBuffer.Bind(Initializer.ParameterMap, FSceneTextureUniformParameters::FTypeInfo::GetStructMetadata()->GetShaderVariableName());
+	}
+
+public:
+
+	static bool ShouldCompilePermutation(const FMeshMaterialShaderPermutationParameters& Parameters)
+	{
+		return ShouldCacheLandscapeGrassShaders(Parameters);
+	}
+
+	void GetShaderBindings(
+		const FScene* Scene,
+		ERHIFeatureLevel::Type FeatureLevel,
+		const FPrimitiveSceneProxy* PrimitiveSceneProxy,
+		const FMaterialRenderProxy& MaterialRenderProxy,
+		const FMaterial& Material,
+		const FMeshPassProcessorRenderState& DrawRenderState,
+		const FLandscapeGrassWeightShaderElementData& ShaderElementData,
+		FMeshDrawSingleShaderBindings& ShaderBindings) const
+	{
+		FMeshMaterialShader::GetShaderBindings(Scene, FeatureLevel, PrimitiveSceneProxy, MaterialRenderProxy, Material, DrawRenderState, ShaderElementData, ShaderBindings);
+
+		ShaderBindings.Add(RenderOffsetParameter, ShaderElementData.RenderOffset);
+	}
+};
+
+IMPLEMENT_MATERIAL_SHADER_TYPE(, FLandscapeGrassWeightVS, TEXT("/Engine/Private/LandscapeGrassWeight.usf"), TEXT("VSMain"), SF_Vertex);
+
+class FLandscapeGrassWeightPS : public FMeshMaterialShader
+{
+	DECLARE_SHADER_TYPE(FLandscapeGrassWeightPS, MeshMaterial);
+	LAYOUT_FIELD(FShaderParameter, OutputPassParameter);
+public:
+
+	static bool ShouldCompilePermutation(const FMeshMaterialShaderPermutationParameters& Parameters)
+	{
+		return ShouldCacheLandscapeGrassShaders(Parameters);
+	}
+
+	FLandscapeGrassWeightPS(const ShaderMetaType::CompiledShaderInitializerType& Initializer)
+		: FMeshMaterialShader(Initializer)
+	{
+		OutputPassParameter.Bind(Initializer.ParameterMap, TEXT("OutputPass"));
+		PassUniformBuffer.Bind(Initializer.ParameterMap, FSceneTextureUniformParameters::FTypeInfo::GetStructMetadata()->GetShaderVariableName());
+	}
+
+	FLandscapeGrassWeightPS()
+	{}
+
+	void GetShaderBindings(
+		const FScene* Scene,
+		ERHIFeatureLevel::Type FeatureLevel,
+		const FPrimitiveSceneProxy* PrimitiveSceneProxy,
+		const FMaterialRenderProxy& MaterialRenderProxy,
+		const FMaterial& Material,
+		const FMeshPassProcessorRenderState& DrawRenderState,
+		const FLandscapeGrassWeightShaderElementData& ShaderElementData,
+		FMeshDrawSingleShaderBindings& ShaderBindings) const
+	{
+		FMeshMaterialShader::GetShaderBindings(Scene, FeatureLevel, PrimitiveSceneProxy, MaterialRenderProxy, Material, DrawRenderState, ShaderElementData, ShaderBindings);
+
+		ShaderBindings.Add(OutputPassParameter, ShaderElementData.OutputPass);
+	}
+};
+
+IMPLEMENT_MATERIAL_SHADER_TYPE(, FLandscapeGrassWeightPS, TEXT("/Engine/Private/LandscapeGrassWeight.usf"), TEXT("PSMain"), SF_Pixel);
+
+void UE::Landscape::Grass::AddGrassWeightShaderTypes(FMaterialShaderTypes& InOutShaderTypes)
+{
+	InOutShaderTypes.AddShaderType<FLandscapeGrassWeightVS>();
+	InOutShaderTypes.AddShaderType<FLandscapeGrassWeightPS>();
+}
+
 class FLandscapeGrassWeightMeshProcessor : public FMeshPassProcessor
 {
 public:
@@ -149,10 +262,6 @@ bool FLandscapeGrassWeightMeshProcessor::Process(
 {
 	const FVertexFactory* VertexFactory = MeshBatch.VertexFactory;
 
-	TMeshProcessorShaders<
-		FLandscapeGrassWeightVS,
-		FLandscapeGrassWeightPS> PassShaders;
-
 	FMaterialShaderTypes ShaderTypes;
 	ShaderTypes.AddShaderType<FLandscapeGrassWeightVS>();
 	ShaderTypes.AddShaderType<FLandscapeGrassWeightPS>();
@@ -163,6 +272,9 @@ bool FLandscapeGrassWeightMeshProcessor::Process(
 		return false;
 	}
 
+	TMeshProcessorShaders<
+		FLandscapeGrassWeightVS,
+		FLandscapeGrassWeightPS> PassShaders;
 	Shaders.TryGetVertexShader(PassShaders.VertexShader);
 	Shaders.TryGetPixelShader(PassShaders.PixelShader);
 
@@ -335,8 +447,8 @@ FLandscapeGrassWeightExporter::FLandscapeGrassWeightExporter(ALandscapeProxy* In
 
 	UE::RenderCommandPipe::FSyncScope SyncScope;
 
-	RenderCaptureInterface::FScopedCapture RenderCapture((RenderCaptureNextGrassmapDraws != 0), TEXT("LandscapeGrassmapCapture"));
-	RenderCaptureNextGrassmapDraws = FMath::Max(0, RenderCaptureNextGrassmapDraws - 1);
+	RenderCaptureInterface::FScopedCapture RenderCapture((GRenderCaptureNextGrassmapDraws != 0), TEXT("LandscapeGrassmapCapture"));
+	GRenderCaptureNextGrassmapDraws = FMath::Max(0, GRenderCaptureNextGrassmapDraws - 1);
 
 	// render
 	FLandscapeGrassWeightExporter_RenderThread* Exporter = this;
@@ -348,8 +460,98 @@ FLandscapeGrassWeightExporter::FLandscapeGrassWeightExporter(ALandscapeProxy* In
 		});
 }
 
+
+struct FByteBuffer2DView : public IBuffer2DView<uint8>
+{
+	uint8* BufferStart = nullptr;
+	int32 ByteStrideY = 0;
+	int32 ByteStrideX = 0;
+	int32 NumX = 0;
+	int32 NumY = 0;
+
+	// copy elements from buffer to Dest, in X then Y order
+	virtual void CopyTo(uint8* Dest, int32 SizeInBytes) const override
+	{
+		for (int Y = 0; SizeInBytes > 0 && Y < NumY; Y++)
+		{
+			uint8* Src = BufferStart + Y * ByteStrideY;
+			int32 CopyCountX = FMath::Min(SizeInBytes, NumX);
+			// we can't use memcpy because of the ByteStride
+			while (CopyCountX--)
+			{
+				*Dest = *Src;
+				Dest++;
+				Src += ByteStrideX;
+			}
+			SizeInBytes -= NumX;
+		}
+	}
+
+	bool IsAllZero() const
+	{
+		for (int Y = 0; Y < NumY; Y++)
+		{
+			uint8 MaxBits = 0;
+			uint8* Src = BufferStart + Y * ByteStrideY;
+			for (int X = 0; X < NumX; X++)
+			{
+				MaxBits = MaxBits | *Src;
+				Src += ByteStrideX;
+			}
+			if (MaxBits != 0)
+			{
+				return false;
+			}
+		}
+		return true;
+	}
+
+	virtual int32 Num() const override { return NumX * NumY; }
+};
+
+// FColor memory layout matches our BGRA GPU layout only on little endian CPUs!
+#if PLATFORM_LITTLE_ENDIAN
+	#define BGRA_AS_FCOLOR_BLUE B
+	#define BGRA_AS_FCOLOR_GREEN G
+	#define BGRA_AS_FCOLOR_RED R
+	#define BGRA_AS_FCOLOR_ALPHA A
+#else
+	#define BGRA_AS_FCOLOR_BLUE A
+	#define BGRA_AS_FCOLOR_GREEN R
+	#define BGRA_AS_FCOLOR_RED G
+	#define BGRA_AS_FCOLOR_ALPHA B
+#endif // PLATFORM_LITTLE_ENDIAN
+
+struct FHeightBuffer2DView : IBuffer2DView<uint16>
+{
+	FColor* BufferStart = nullptr;
+	int32 StrideY = 0;
+	int32 NumX = 0;
+	int32 NumY = 0;
+
+	// copy elements from buffer to Dest, in X then Y order
+	virtual void CopyTo(uint16* Dest, int32 Count) const override
+	{
+		for (int y = 0; Count > 0 && y < NumY; y++)
+		{
+			FColor* Src = BufferStart + y * StrideY;
+			int32 CopyCountX = FMath::Min(Count, NumX);
+			while (CopyCountX--)
+			{
+				*Dest = (((uint16) Src->BGRA_AS_FCOLOR_RED) << 8) + (uint16)(Src->BGRA_AS_FCOLOR_GREEN);
+				Dest++;
+				Src++;
+			}
+			Count -= NumX;
+		}
+	}
+
+	virtual int32 Num() const override { return NumX * NumY; }
+};
+
 TMap<ULandscapeComponent*, TUniquePtr<FLandscapeComponentGrassData>, TInlineSetAllocator<1>> FLandscapeGrassWeightExporter::FetchResults()
 {
+	TRACE_CPUPROFILER_EVENT_SCOPE(FetchResults);
 	TMap<ULandscapeComponent*, TUniquePtr<FLandscapeComponentGrassData>, TInlineSetAllocator<1>> Results;
 	TArray<FColor> Samples;
 
@@ -371,10 +573,8 @@ TMap<ULandscapeComponent*, TUniquePtr<FLandscapeComponentGrassData>, TInlineSetA
 	}
 
 	Results.Reserve(ComponentInfos.Num());
-
-	// Local data will be moved in contiguous array at the end of export (to minimize slack waste)
-	TArray<uint16> HeightData;
-	TMap<ULandscapeGrassType*, TArray<uint8>> WeightData;
+	FHeightBuffer2DView HeightData;
+	TMap<ULandscapeGrassType*, IBuffer2DView<uint8>*> WeightData;
 
 	for (auto& ComponentInfo : ComponentInfos)
 	{
@@ -382,33 +582,35 @@ TMap<ULandscapeComponent*, TUniquePtr<FLandscapeComponentGrassData>, TInlineSetA
 
 		TUniquePtr<FLandscapeComponentGrassData> NewGrassData = MakeUnique<FLandscapeComponentGrassData>(Component);
 
-		if (ComponentInfo.FirstHeightMipsPassIndex != MAX_int32)
-		{
-			HeightData.Empty(FMath::Square(ComponentSizeVerts));
-		}
-		else
-		{
-			HeightData.Empty(0);
-		}
+		int32 ComponentSizeVerts2 = FMath::Square(ComponentSizeVerts);
+
+		HeightData.NumX = ComponentSizeVerts;
+		HeightData.NumY = ComponentSizeVerts;
+		HeightData.StrideY = TargetSize.X;
+
 #if WITH_EDITORONLY_DATA
 		NewGrassData->HeightMipData.Empty(HeightMips.Num());
 #endif // WITH_EDITORONLY_DATA
 
 		WeightData.Empty();
-		TArray<TArray<uint8>*> GrassWeightArrays;
-		GrassWeightArrays.Empty(ComponentInfo.RequestedGrassTypes.Num());
-		for (auto GrassType : ComponentInfo.RequestedGrassTypes)
-		{
-			WeightData.Add(GrassType);
-		}
 
-		// need a second loop because the WeightData map will reallocate its arrays as grass types are added
-		for (auto GrassType : ComponentInfo.RequestedGrassTypes)
+		// this array is in 1:1 correspondence with ComponentInfo.RequestedGrassTypes
+		TArray<FByteBuffer2DView> GrassWeightArrays;
+		GrassWeightArrays.SetNum(ComponentInfo.RequestedGrassTypes.Num());
+
+		for (int32 Index = 0; Index < GrassWeightArrays.Num(); Index++)
 		{
-			TArray<uint8>* DataArray = WeightData.Find(GrassType);
-			check(DataArray);
-			DataArray->Empty(FMath::Square(ComponentSizeVerts));
-			GrassWeightArrays.Add(DataArray);
+			FByteBuffer2DView* WeightView = &GrassWeightArrays[Index];
+			ULandscapeGrassType* GrassType = ComponentInfo.RequestedGrassTypes[Index];
+
+			WeightView->NumX = ComponentSizeVerts;
+			WeightView->NumY = ComponentSizeVerts;
+			WeightView->ByteStrideX = 4;
+			WeightView->ByteStrideY = TargetSize.X * 4;
+			WeightView->BufferStart = nullptr;
+
+			// Note: WeightData points directly at the elements of GrassWeightArrays (DO NOT REALLOCATE GRASSWEIGHTARRAYS)
+			WeightData.Add(GrassType, WeightView);
 		}
 
 		// output debug bitmap
@@ -423,65 +625,39 @@ TMap<ULandscapeComponent*, TUniquePtr<FLandscapeComponentGrassData>, TInlineSetA
 		}
 #endif
 
-		// FColor memory layout matches our BGRA GPU layout only on little endian CPUs!
-		#if PLATFORM_LITTLE_ENDIAN
-			#define BGRA_AS_FCOLOR_BLUE B
-			#define BGRA_AS_FCOLOR_GREEN G
-			#define BGRA_AS_FCOLOR_RED R
-			#define BGRA_AS_FCOLOR_ALPHA A
-		#else
-			#define BGRA_AS_FCOLOR_BLUE A
-			#define BGRA_AS_FCOLOR_GREEN R
-			#define BGRA_AS_FCOLOR_RED G
-			#define BGRA_AS_FCOLOR_ALPHA B
-		#endif // PLATFORM_LITTLE_ENDIAN
-
+		int32 GrassTypeCount = ComponentInfo.RequestedGrassTypes.Num();
 		for (int32 PassIdx = 0; PassIdx < ComponentInfo.NumPasses; PassIdx++)
 		{
 			FColor* SampleData = &Samples[ComponentInfo.PixelOffsetX + PassIdx * ComponentSizeVerts];
+
 			if (PassIdx < ComponentInfo.FirstHeightMipsPassIndex)
 			{
-				if (PassIdx == 0)
+				if (PassIdx == 0)	// height in RG, grass weights in BA
 				{
-					for (int32 y = 0; y < ComponentSizeVerts; y++)
+					HeightData.BufferStart = SampleData;
+					if (GrassTypeCount > 0)
 					{
-						for (int32 x = 0; x < ComponentSizeVerts; x++)
+						GrassWeightArrays[0].BufferStart = &SampleData->BGRA_AS_FCOLOR_BLUE;
+						if (GrassTypeCount > 1)
 						{
-							FColor& Sample = SampleData[x + y * TargetSize.X];
-							uint16 Height = (((uint16)Sample.BGRA_AS_FCOLOR_RED) << 8) + (uint16)(Sample.BGRA_AS_FCOLOR_GREEN);
-							HeightData.Add(Height);
-							if (ComponentInfo.RequestedGrassTypes.Num() > 0)
-							{
-								GrassWeightArrays[0]->Add(Sample.BGRA_AS_FCOLOR_BLUE);
-								if (ComponentInfo.RequestedGrassTypes.Num() > 1)
-								{
-									GrassWeightArrays[1]->Add(Sample.BGRA_AS_FCOLOR_ALPHA);
-								}
-							}
+							GrassWeightArrays[1].BufferStart = &SampleData->BGRA_AS_FCOLOR_ALPHA;
 						}
 					}
 				}
 				else
 				{
-					for (int32 y = 0; y < ComponentSizeVerts; y++)
+					int32 TypeIdx = PassIdx * 4 - 2;
+
+					GrassWeightArrays[TypeIdx+0].BufferStart = &SampleData->BGRA_AS_FCOLOR_RED;
+					if (GrassTypeCount > TypeIdx+1)
 					{
-						for (int32 x = 0; x < ComponentSizeVerts; x++)
+						GrassWeightArrays[TypeIdx+1].BufferStart = &SampleData->BGRA_AS_FCOLOR_GREEN;
+						if (GrassTypeCount > TypeIdx + 2)
 						{
-							FColor& Sample = SampleData[x + y * TargetSize.X];
-
-							int32 TypeIdx = PassIdx * 4 - 2;
-							GrassWeightArrays[TypeIdx++]->Add(Sample.BGRA_AS_FCOLOR_RED);
-							if (TypeIdx < ComponentInfo.RequestedGrassTypes.Num())
+							GrassWeightArrays[TypeIdx + 2].BufferStart = &SampleData->BGRA_AS_FCOLOR_BLUE;
+							if (GrassTypeCount > TypeIdx + 3)
 							{
-								GrassWeightArrays[TypeIdx++]->Add(Sample.BGRA_AS_FCOLOR_GREEN);
-								if (TypeIdx < ComponentInfo.RequestedGrassTypes.Num())
-								{
-									GrassWeightArrays[TypeIdx++]->Add(Sample.BGRA_AS_FCOLOR_BLUE);
-									if (TypeIdx < ComponentInfo.RequestedGrassTypes.Num())
-									{
-										GrassWeightArrays[TypeIdx++]->Add(Sample.BGRA_AS_FCOLOR_ALPHA);
-									}
-								}
+								GrassWeightArrays[TypeIdx + 3].BufferStart = &SampleData->BGRA_AS_FCOLOR_ALPHA;
 							}
 						}
 					}
@@ -493,13 +669,15 @@ TMap<ULandscapeComponent*, TUniquePtr<FLandscapeComponentGrassData>, TInlineSetA
 				const int32 Mip = HeightMips[PassIdx - ComponentInfo.FirstHeightMipsPassIndex];
 				int32 MipSizeVerts = NumSubsections * (SubsectionSizeQuads >> Mip);
 				TArray<uint16>& MipHeightData = NewGrassData->HeightMipData.Add(Mip);
+				MipHeightData.SetNumUninitialized(MipSizeVerts* MipSizeVerts);
+				uint16* DstMipHeight = MipHeightData.GetData();
 				for (int32 y = 0; y < MipSizeVerts; y++)
 				{
+					FColor* SrcSample = &SampleData[y * TargetSize.X];
 					for (int32 x = 0; x < MipSizeVerts; x++)
 					{
-						FColor& Sample = SampleData[x + y * TargetSize.X];
-						uint16 Height = (((uint16)Sample.BGRA_AS_FCOLOR_RED) << 8) + (uint16)(Sample.BGRA_AS_FCOLOR_GREEN);
-						MipHeightData.Add(Height);
+						*DstMipHeight++ = (((uint16)SrcSample->BGRA_AS_FCOLOR_RED) << 8) + (uint16)(SrcSample->BGRA_AS_FCOLOR_GREEN);
+						SrcSample++;
 					}
 				}
 #endif // WITH_EDITORONLY_DATA
@@ -510,20 +688,22 @@ TMap<ULandscapeComponent*, TUniquePtr<FLandscapeComponentGrassData>, TInlineSetA
 		#undef BGRA_AS_FCOLOR_GREEN
 		#undef BGRA_AS_FCOLOR_RED
 		#undef BGRA_AS_FCOLOR_ALPHA
-
-		// remove null grass type if we had one (can occur if the node has null entries)
-		WeightData.Remove(nullptr);
-
-		// Remove any grass data that is entirely weight 0
-		for (auto Iter(WeightData.CreateIterator()); Iter; ++Iter)
+		
 		{
-			if (Iter->Value.IndexOfByPredicate([&](const int8& Weight) { return Weight != 0; }) == INDEX_NONE)
+			TRACE_CPUPROFILER_EVENT_SCOPE(RemoveZeroWeight);
+			for (auto Iter(WeightData.CreateIterator()); Iter; ++Iter)
 			{
-				Iter.RemoveCurrent();
+				// Remove null grass type if we had one (can occur if the node has null entries)
+				// Remove any grass data that is entirely weight 0
+				if (Iter->Key == nullptr ||
+					((FByteBuffer2DView*)Iter->Value)->IsAllZero())
+				{
+					Iter.RemoveCurrent();
+				}
 			}
 		}
 
-		NewGrassData->InitializeFrom(HeightData, WeightData);
+		NewGrassData->InitializeFrom(&HeightData, WeightData);
 		Results.Add(Component, MoveTemp(NewGrassData));
 	}
 
@@ -532,6 +712,8 @@ TMap<ULandscapeComponent*, TUniquePtr<FLandscapeComponentGrassData>, TInlineSetA
 
 void FLandscapeGrassWeightExporter::ApplyResults()
 {
+	TRACE_CPUPROFILER_EVENT_SCOPE(ApplyResults);
+
 	TMap<ULandscapeComponent*, TUniquePtr<FLandscapeComponentGrassData>, TInlineSetAllocator<1>> NewGrassData = FetchResults();
 
 	for (auto&& GrassDataPair : NewGrassData)
@@ -579,103 +761,3 @@ void FLandscapeGrassWeightExporter::AddReferencedObjects(UObject* InThis, FRefer
 		Collector.AddReferencedObjects(Info.RequestedGrassTypes);
 	}
 }
-
-static bool ShouldCacheLandscapeGrassShaders(const FMeshMaterialShaderPermutationParameters& Parameters)
-{
-	// We only need grass weight shaders for Landscape vertex factories on desktop platforms
-	return (Parameters.MaterialParameters.bIsUsedWithLandscape || Parameters.MaterialParameters.bIsSpecialEngineMaterial) &&
-		IsFeatureLevelSupported(Parameters.Platform, ERHIFeatureLevel::SM5) &&
-		Parameters.VertexFactoryType == FindVertexFactoryType(FName(TEXT("FLandscapeFixedGridVertexFactory"), FNAME_Find)) &&
-		EnumHasAllFlags(Parameters.Flags, EShaderPermutationFlags::HasEditorOnlyData);
-}
-
-class FLandscapeGrassWeightVS : public FMeshMaterialShader
-{
-	DECLARE_SHADER_TYPE(FLandscapeGrassWeightVS, MeshMaterial);
-
-	LAYOUT_FIELD(FShaderParameter, RenderOffsetParameter);
-
-protected:
-
-	FLandscapeGrassWeightVS()
-	{}
-
-	FLandscapeGrassWeightVS(const FMeshMaterialShaderType::CompiledShaderInitializerType& Initializer)
-		: FMeshMaterialShader(Initializer)
-	{
-		RenderOffsetParameter.Bind(Initializer.ParameterMap, TEXT("RenderOffset"));
-		PassUniformBuffer.Bind(Initializer.ParameterMap, FSceneTextureUniformParameters::FTypeInfo::GetStructMetadata()->GetShaderVariableName());
-	}
-
-public:
-
-	static bool ShouldCompilePermutation(const FMeshMaterialShaderPermutationParameters& Parameters)
-	{
-		return ShouldCacheLandscapeGrassShaders(Parameters);
-	}
-
-	void GetShaderBindings(
-		const FScene* Scene,
-		ERHIFeatureLevel::Type FeatureLevel,
-		const FPrimitiveSceneProxy* PrimitiveSceneProxy,
-		const FMaterialRenderProxy& MaterialRenderProxy,
-		const FMaterial& Material,
-		const FMeshPassProcessorRenderState& DrawRenderState,
-		const FLandscapeGrassWeightShaderElementData& ShaderElementData,
-		FMeshDrawSingleShaderBindings& ShaderBindings) const
-	{
-		FMeshMaterialShader::GetShaderBindings(Scene, FeatureLevel, PrimitiveSceneProxy, MaterialRenderProxy, Material, DrawRenderState, ShaderElementData, ShaderBindings);
-
-		ShaderBindings.Add(RenderOffsetParameter, ShaderElementData.RenderOffset);
-	}
-};
-
-IMPLEMENT_MATERIAL_SHADER_TYPE(, FLandscapeGrassWeightVS, TEXT("/Engine/Private/LandscapeGrassWeight.usf"), TEXT("VSMain"), SF_Vertex);
-
-class FLandscapeGrassWeightPS : public FMeshMaterialShader
-{
-	DECLARE_SHADER_TYPE(FLandscapeGrassWeightPS, MeshMaterial);
-	LAYOUT_FIELD(FShaderParameter, OutputPassParameter);
-public:
-
-	static bool ShouldCompilePermutation(const FMeshMaterialShaderPermutationParameters& Parameters)
-	{
-		return ShouldCacheLandscapeGrassShaders(Parameters);
-	}
-
-	FLandscapeGrassWeightPS(const ShaderMetaType::CompiledShaderInitializerType& Initializer)
-		: FMeshMaterialShader(Initializer)
-	{
-		OutputPassParameter.Bind(Initializer.ParameterMap, TEXT("OutputPass"));
-		PassUniformBuffer.Bind(Initializer.ParameterMap, FSceneTextureUniformParameters::FTypeInfo::GetStructMetadata()->GetShaderVariableName());
-	}
-
-	FLandscapeGrassWeightPS()
-	{}
-
-	void GetShaderBindings(
-		const FScene* Scene,
-		ERHIFeatureLevel::Type FeatureLevel,
-		const FPrimitiveSceneProxy* PrimitiveSceneProxy,
-		const FMaterialRenderProxy& MaterialRenderProxy,
-		const FMaterial& Material,
-		const FMeshPassProcessorRenderState& DrawRenderState,
-		const FLandscapeGrassWeightShaderElementData& ShaderElementData,
-		FMeshDrawSingleShaderBindings& ShaderBindings) const
-	{
-		FMeshMaterialShader::GetShaderBindings(Scene, FeatureLevel, PrimitiveSceneProxy, MaterialRenderProxy, Material, DrawRenderState, ShaderElementData, ShaderBindings);
-
-		ShaderBindings.Add(OutputPassParameter, ShaderElementData.OutputPass);
-	}
-};
-
-IMPLEMENT_MATERIAL_SHADER_TYPE(, FLandscapeGrassWeightPS, TEXT("/Engine/Private/LandscapeGrassWeight.usf"), TEXT("PSMain"), SF_Pixel);
-
-void UE::Landscape::Grass::AddGrassWeightShaderTypes(FMaterialShaderTypes& InOutShaderTypes)
-{
-	InOutShaderTypes.AddShaderType<FLandscapeGrassWeightVS>();
-	InOutShaderTypes.AddShaderType<FLandscapeGrassWeightPS>();
-}
-
-#endif // WITH_EDITOR
-
diff --git a/Engine/Source/Runtime/Landscape/Private/LandscapePrivate.h b/Engine/Source/Runtime/Landscape/Private/LandscapePrivate.h
--- a/Engine/Source/Runtime/Landscape/Private/LandscapePrivate.h
+++ b/Engine/Source/Runtime/Landscape/Private/LandscapePrivate.h
@@ -8,6 +8,7 @@
 
 LANDSCAPE_API DECLARE_LOG_CATEGORY_EXTERN(LogLandscape, Log, All);
 LANDSCAPE_API DECLARE_LOG_CATEGORY_EXTERN(LogLandscapeBP, Display, All);
+LANDSCAPE_API DECLARE_LOG_CATEGORY_EXTERN(LogGrass, Log, All);
 
 /**
  * Landscape stats
diff --git a/Engine/Source/Runtime/Landscape/Private/LandscapeRender.cpp b/Engine/Source/Runtime/Landscape/Private/LandscapeRender.cpp
--- a/Engine/Source/Runtime/Landscape/Private/LandscapeRender.cpp
+++ b/Engine/Source/Runtime/Landscape/Private/LandscapeRender.cpp
@@ -161,6 +161,9 @@ extern TAutoConsoleVariable<int32> CVarLandscapeShowDirty;
 
 extern RENDERER_API TAutoConsoleVariable<float> CVarStaticMeshLODDistanceScale;
 
+extern int32 GGrassMapUseRuntimeGeneration;
+extern int32 GGrassMapAlwaysBuildRuntimeGenerationResources;
+
 #if !UE_BUILD_SHIPPING
 int32 GVarDumpLandscapeLODsCurrentFrame = 0;
 
@@ -226,22 +229,31 @@ static FAutoConsoleVariableRef CVarLandscapeRayTracingGeometryFractionalLODUpdat
 
 namespace UE::Landscape
 {
-bool NeedsFixedGridVertexFactory(EShaderPlatform InShaderPlatform)
-{
-	bool bNeedsFixedGridVertexFactory = false;
-	// We need the fixed grid vertex factory for virtual texturing : 
-	bNeedsFixedGridVertexFactory |= UseVirtualTexturing(InShaderPlatform);
+	bool NeedsFixedGridVertexFactory(EShaderPlatform InShaderPlatform)
+	{
+		bool bNeedsFixedGridVertexFactory = false;
+		// We need the fixed grid vertex factory for virtual texturing : 
+		bNeedsFixedGridVertexFactory |= UseVirtualTexturing(InShaderPlatform);
 
-	// We need the fixed grid vertex factory for rendering Landscape into Lumen Surface Cache: 
-	bNeedsFixedGridVertexFactory |= DoesPlatformSupportLumenGI(InShaderPlatform);
+		// We need the fixed grid vertex factory for rendering Landscape into Lumen Surface Cache: 
+		bNeedsFixedGridVertexFactory |= DoesPlatformSupportLumenGI(InShaderPlatform);
 
-	// We need the fixed grid vertex factory for rendering the water info texture : 
-	// This cvar is defined in the water plugin and searching for it should return nullptr if the plugin is not loaded
-	const bool bWaterPluginLoaded = IConsoleManager::Get().FindConsoleVariable(TEXT("r.Water.WaterInfo.RenderMethod")) != nullptr;
-	bNeedsFixedGridVertexFactory |= bWaterPluginLoaded;
+		// We need the fixed grid vertex factory for rendering the water info texture : 
+		// This cvar is defined in the water plugin and searching for it should return nullptr if the plugin is not loaded
+		const bool bWaterPluginLoaded = IConsoleManager::Get().FindConsoleVariable(TEXT("r.Water.WaterInfo.RenderMethod")) != nullptr;
+		bNeedsFixedGridVertexFactory |= bWaterPluginLoaded;
 
-	return bNeedsFixedGridVertexFactory;
-}
+		return bNeedsFixedGridVertexFactory;
+	}
+
+	bool ShouldBuildGrassMapRenderingResources()
+	{
+		#if WITH_EDITOR
+			return true;
+		#else
+			return (GGrassMapUseRuntimeGeneration || GGrassMapAlwaysBuildRuntimeGenerationResources);
+		#endif // WITH_EDITOR
+	}
 } // namespace UE::Landscape
 
 //
@@ -1487,10 +1499,8 @@ void FLandscapeComponentSceneProxy::CreateRenderThreadResources(FRHICommandListB
 
 		bool bNeedsFixedGridVertexFactory = NeedsFixedGridVertexFactory(GetScene().GetShaderPlatform());
 
-#if WITH_EDITOR
 		// We also need the fixed vertex factor for grass/physical materials : 
 		bNeedsFixedGridVertexFactory |= (SharedBuffers->GrassIndexBuffer != nullptr);
-#endif // WITH_EDITOR
 
 		if (bNeedsFixedGridVertexFactory)
 		{
@@ -1532,7 +1542,6 @@ void FLandscapeComponentSceneProxy::CreateRenderThreadResources(FRHICommandListB
 		LandscapeFixedGridUniformShaderParameters[LodIndex].InitResource(RHICmdList);
 	}
 
-#if WITH_EDITOR
 	// Create MeshBatch for grass rendering
 	if (SharedBuffers->GrassIndexBuffer)
 	{
@@ -1585,7 +1594,6 @@ void FLandscapeComponentSceneProxy::CreateRenderThreadResources(FRHICommandListB
 			CollisionBatchElement->FirstIndex = SharedBuffers->GrassIndexMipOffsets[Mip];
 		}
 	}
-#endif
 }
 
 #if RHI_RAYTRACING
@@ -3097,7 +3105,6 @@ void FLandscapeSharedBuffers::CreateIndexBuffers(FRHICommandListBase& RHICmdList
 	}
 }
 
-#if WITH_EDITOR
 template <typename INDEX_TYPE>
 void FLandscapeSharedBuffers::CreateGrassIndexBuffer(FRHICommandListBase& RHICmdList, const FName& InOwnerName)
 {
@@ -3143,7 +3150,6 @@ void FLandscapeSharedBuffers::CreateGrassIndexBuffer(FRHICommandListBase& RHICmd
 	IndexBuffer->InitResource(RHICmdList);
 	GrassIndexBuffer = IndexBuffer;
 }
-#endif
 
 FLandscapeSharedBuffers::FLandscapeSharedBuffers(FRHICommandListBase& RHICmdList, const int32 InSharedBuffersKey, const int32 InSubsectionSizeQuads, const int32 InNumSubsections, const ERHIFeatureLevel::Type InFeatureLevel, const FName& InOwnerName)
 	: SharedBuffersKey(InSharedBuffersKey)
@@ -3157,10 +3163,7 @@ FLandscapeSharedBuffers::FLandscapeSharedBuffers(FRHICommandListBase& RHICmdList
 	, TileVertexFactory(nullptr)
 	, TileDataBuffer(nullptr)
 	, bUse32BitIndices(false)
-
-#if WITH_EDITOR
 	, GrassIndexBuffer(nullptr)
-#endif
 {
 	NumVertices = FMath::Square(SubsectionSizeVerts) * FMath::Square(NumSubsections);
 	
@@ -3182,22 +3185,18 @@ FLandscapeSharedBuffers::FLandscapeSharedBuffers(FRHICommandListBase& RHICmdList
 	{
 		bUse32BitIndices = true;
 		CreateIndexBuffers<uint32>(RHICmdList, InOwnerName);
-#if WITH_EDITOR
-		if (InFeatureLevel > ERHIFeatureLevel::ES3_1)
+		if (UE::Landscape::ShouldBuildGrassMapRenderingResources())
 		{
 			CreateGrassIndexBuffer<uint32>(RHICmdList, InOwnerName);
 		}
-#endif
 	}
 	else
 	{
 		CreateIndexBuffers<uint16>(RHICmdList, InOwnerName);
-#if WITH_EDITOR
-		if (InFeatureLevel > ERHIFeatureLevel::ES3_1)
+		if (UE::Landscape::ShouldBuildGrassMapRenderingResources())
 		{
 			CreateGrassIndexBuffer<uint16>(RHICmdList, InOwnerName);
 		}
-#endif
 	}
 }
 
@@ -3225,13 +3224,11 @@ FLandscapeSharedBuffers::~FLandscapeSharedBuffers()
 	}
 #endif
 
-#if WITH_EDITOR
 	if (GrassIndexBuffer)
 	{
 		GrassIndexBuffer->ReleaseResource();
 		delete GrassIndexBuffer;
 	}
-#endif
 
 	delete VertexFactory;
 
@@ -3363,7 +3360,8 @@ void FLandscapeVertexFactory::InitRHI(FRHICommandListBase& RHICmdList)
 	// position decls
 	Elements.Add(AccessStreamComponent(Data.PositionComponent, 0));
 
-	AddPrimitiveIdStreamElement(EVertexInputStreamType::Default, Elements, 1, 0xff);
+	// Use the same attribute on mobile and non-mobile, to enable the GPUScene path on both.
+	AddPrimitiveIdStreamElement(EVertexInputStreamType::Default, Elements, /* AttributeIndex = */ 1, /* AttributeIndex_Mobile = */ 1);
 	// create the actual device decls
 	InitDeclaration(Elements);
 }
@@ -3384,8 +3382,7 @@ void FLandscapeVertexFactory::ModifyCompilationEnvironment(const FVertexFactoryS
 {
 	FVertexFactory::ModifyCompilationEnvironment(Parameters, OutEnvironment);
 
-	// TODO: support GPUScene on mobile
-	OutEnvironment.SetDefine(TEXT("VF_SUPPORTS_PRIMITIVE_SCENE_DATA"), Parameters.VertexFactoryType->SupportsPrimitiveIdStream() && UseGPUScene(Parameters.Platform, GetMaxSupportedFeatureLevel(Parameters.Platform)) && !IsMobilePlatform(Parameters.Platform));
+	OutEnvironment.SetDefine(TEXT("VF_SUPPORTS_PRIMITIVE_SCENE_DATA"), Parameters.VertexFactoryType->SupportsPrimitiveIdStream() && UseGPUScene(Parameters.Platform, GetMaxSupportedFeatureLevel(Parameters.Platform)));
 
 	// Make sure landscape vertices go back to local space so that we have consistency between the transform on normals and geometry
 	OutEnvironment.SetDefine(TEXT("RAY_TRACING_DYNAMIC_MESH_IN_LOCAL_SPACE"), TEXT("1"));
diff --git a/Engine/Source/Runtime/Landscape/Private/LandscapeSubsystem.cpp b/Engine/Source/Runtime/Landscape/Private/LandscapeSubsystem.cpp
--- a/Engine/Source/Runtime/Landscape/Private/LandscapeSubsystem.cpp
+++ b/Engine/Source/Runtime/Landscape/Private/LandscapeSubsystem.cpp
@@ -61,6 +61,8 @@ static FAutoConsoleVariableRef CVarUpdateProxyActorRenderMethodUpdateOnTickAtRun
 	GUpdateProxyActorRenderMethodOnTickAtRuntime,
 	TEXT("Update landscape proxy's rendering method (nanite enabled) when ticked. Always enabled in editor."));
 
+extern int32 GGrassMapUseRuntimeGeneration;
+
 DECLARE_CYCLE_STAT(TEXT("LandscapeSubsystem Tick"), STAT_LandscapeSubsystemTick, STATGROUP_Landscape);
 
 #define LOCTEXT_NAMESPACE "LandscapeSubsystem"
@@ -107,8 +109,12 @@ void ULandscapeSubsystem::Initialize(FSubsystemCollectionBase& Collection)
 		LandscapeNaniteEnabledCVar->OnChangedDelegate().AddUObject(this, &ULandscapeSubsystem::OnNaniteEnabledChanged);
 	}
 
+	TextureStreamingManager = new FLandscapeTextureStreamingManager();
+	check(TextureStreamingManager);
+
+	GrassMapsBuilder = new FLandscapeGrassMapsBuilder(GetWorld(), *TextureStreamingManager);
+
 #if WITH_EDITOR
-	GrassMapsBuilder = new FLandscapeGrassMapsBuilder(GetWorld());
 	PhysicalMaterialBuilder = new FLandscapePhysicalMaterialBuilder(GetWorld());
 
 	if (!IsRunningCommandlet())
@@ -180,6 +186,7 @@ void ULandscapeSubsystem::Deinitialize()
 void ULandscapeSubsystem::HandlePostGarbageCollect()
 {
 	ALandscapeProxy::RemoveInvalidExclusionBoxes();
+	GetTextureStreamingManager()->CleanupInvalidEntries();
 }
 
 TStatId ULandscapeSubsystem::GetStatId() const
@@ -187,6 +194,29 @@ TStatId ULandscapeSubsystem::GetStatId() const
 	RETURN_QUICK_DECLARE_CYCLE_STAT(ULandscapeSubsystem, STATGROUP_Tickables);
 }
 
+void ULandscapeSubsystem::RegisterComponent(ULandscapeComponent* Component)
+{
+	GetGrassMapBuilder()->RegisterComponent(Component);
+}
+
+void ULandscapeSubsystem::UnregisterComponent(ULandscapeComponent* Component)
+{
+	GetGrassMapBuilder()->UnregisterComponent(Component);
+}
+
+
+void ULandscapeSubsystem::RemoveGrassInstances(const TSet<ULandscapeComponent*>* ComponentsToRemoveGrassInstances)
+{
+	TRACE_CPUPROFILER_EVENT_SCOPE(RemoveGrassInstances);
+	for (TWeakObjectPtr<ALandscapeProxy> ProxyPtr : Proxies)
+	{
+		if (ALandscapeProxy* Proxy = ProxyPtr.Get())
+		{
+			Proxy->FlushGrassComponents(ComponentsToRemoveGrassInstances, /*bFlushGrassMaps = */false);
+		}
+	}
+}
+
 void ULandscapeSubsystem::RegenerateGrass(bool bInFlushGrass, bool bInForceSync, TOptional<TArrayView<FVector>> InOptionalCameraLocations)
 {
 	TRACE_CPUPROFILER_EVENT_SCOPE(ULandscapeSubsystem::RegenerateGrass);
@@ -200,14 +230,7 @@ void ULandscapeSubsystem::RegenerateGrass(bool bInFlushGrass, bool bInForceSync,
 
 	if (bInFlushGrass)
 	{
-		TRACE_CPUPROFILER_EVENT_SCOPE(FlushGrass);
-		for (TWeakObjectPtr<ALandscapeProxy> ProxyPtr : Proxies)
-		{
-			if (ALandscapeProxy* Proxy = ProxyPtr.Get())
-			{
-				Proxy->FlushGrassComponents(/*OnlyForComponents = */nullptr, /*bFlushGrassMaps = */false);
-			}
-		}
+		RemoveGrassInstances();
 	}
 
 	{
@@ -268,6 +291,7 @@ void ULandscapeSubsystem::Tick(float DeltaTime)
 	Super::Tick(DeltaTime);
 
 	UWorld* World = GetWorld();
+	bool bIsGameWorld = World->IsGameWorld();
 
 #if WITH_EDITOR
 	AppCurrentDateTime = FDateTime::Now();
@@ -300,14 +324,13 @@ void ULandscapeSubsystem::Tick(float DeltaTime)
 	{
 		if (OldCameras.Num() || World->ViewLocationsRenderedLastFrame.Num())
 		{
-			Cameras = &OldCameras;
-			// there is a bug here, which often leaves us with no cameras in the editor
+			// there is a bug here, which often leaves us with no cameras in the editor -- try to fall back to previous camera position(s)
 			if (World->ViewLocationsRenderedLastFrame.Num())
 			{
 				check(IsInGameThread());
-				Cameras = &World->ViewLocationsRenderedLastFrame;
-				OldCameras = *Cameras;
+				OldCameras = World->ViewLocationsRenderedLastFrame;
 			}
+			Cameras = &OldCameras;
 		}
 	}
 	else
@@ -325,6 +348,66 @@ void ULandscapeSubsystem::Tick(float DeltaTime)
 		}
 	}
 
+	// dereference and filter TWeakObjectPtr proxies once, up front, to get the list of active proxies
+	// also run early update on Proxies, and determine if all of the proxies are ready for grass generation to start
+	// TODO [chris.tchou] We should drop the usage of TWeakObjectPtr to reduce the dereference cost, can rely on register/unregister instead
+	bool bAllProxiesReadyForGrassMapGeneration = true;
+	bool bAllProxiesRuntimeGrassMapsDisabled = true;
+	static TArray<ALandscapeProxy*> ActiveProxies;
+	ActiveProxies.Reset(Proxies.Num());
+	for (TWeakObjectPtr<ALandscapeProxy> ProxyPtr : Proxies)
+	{
+		if (ALandscapeProxy* Proxy = ProxyPtr.Get())
+		{
+			ActiveProxies.Add(Proxy);
+			
+			// Update the proxies proxy
+			{
+				if (Proxy->bUseRuntimeGrassMapGeneration)
+				{
+					bAllProxiesRuntimeGrassMapsDisabled = false;
+				}
+
+#if WITH_EDITOR
+				if (!bIsGameWorld)
+				{
+					// in editor, automatically update component grass types if the material changes
+					for (ULandscapeComponent* Component : Proxy->LandscapeComponents)
+					{
+						Component->UpdateGrassTypes();
+					}
+
+					// before starting grass map generation in editor, ensure layers are up to date
+					// and grass update is enabled (if disabled then a tool is currently operating)
+					ALandscape* Landscape = Proxy->GetLandscapeActor();
+					if ((Landscape == nullptr) || !Landscape->bGrassUpdateEnabled || !Landscape->IsUpToDate())
+					{
+						bAllProxiesReadyForGrassMapGeneration = false;
+					}
+				}
+#endif // WITH_EDITOR
+
+				// Update the grass type summary if necessary
+				if (!Proxy->IsGrassTypeSummaryValid())
+				{
+					Proxy->UpdateGrassTypeSummary();
+				}
+			}
+		}
+	}
+
+	bool bGrassMapGenerationDisabled = bAllProxiesRuntimeGrassMapsDisabled;
+#if WITH_EDITOR
+	if (GIsEditor && !bIsGameWorld)
+	{
+		bGrassMapGenerationDisabled = false;
+	}
+#endif // WITH_EDITOR
+
+	bool bAllowStartGrassMapGeneration = bAllProxiesReadyForGrassMapGeneration && !bGrassMapGenerationDisabled && (!bIsGameWorld || GGrassMapUseRuntimeGeneration);
+
+	GrassMapsBuilder->AmortizedUpdateGrassMaps(Cameras ? *Cameras : TArray<FVector>(), bIsGrassCreationPrioritized, bAllowStartGrassMapGeneration);
+
 	int32 InOutNumComponentsCreated = 0;
 #if WITH_EDITOR
 	int32 NumProxiesUpdated = 0;
@@ -336,50 +419,50 @@ void ULandscapeSubsystem::Tick(float DeltaTime)
 		NumNaniteMeshUpdatesAvailable -= NumMeshesToUpdate;
 	}
 #endif // WITH_EDITOR
-	for (TWeakObjectPtr<ALandscapeProxy> ProxyPtr : Proxies)
+	for (ALandscapeProxy* Proxy : ActiveProxies)
 	{
-		if (ALandscapeProxy* Proxy = ProxyPtr.Get())
-		{
 #if WITH_EDITOR
-			if (GIsEditor)
+		if (GIsEditor)
+		{
+			if (ALandscape* Landscape = Cast<ALandscape>(Proxy))
 			{
-				if (ALandscape* Landscape = Cast<ALandscape>(Proxy))
-				{
-					Landscape->TickLayers(DeltaTime);
-				}
+				Landscape->TickLayers(DeltaTime);
+			}
 
-				// editor-only
-				if (!World->IsPlayInEditor())
-				{
-					Proxy->UpdatePhysicalMaterialTasks();
-				}
+			// editor-only
+			if (!World->IsPlayInEditor())
+			{
+				Proxy->UpdatePhysicalMaterialTasks();
 			}
+		}
 
-			Proxy->GetAsyncWorkMonitor().Tick(DeltaTime);
+		Proxy->GetAsyncWorkMonitor().Tick(DeltaTime);
 
-			if (IsLiveNaniteRebuildEnabled())
+		if (IsLiveNaniteRebuildEnabled())
+		{
+			if (NumProxiesUpdated < NumMeshesToUpdate && Proxy->GetAsyncWorkMonitor().CheckIfUpdateTriggeredAndClear(FAsyncWorkMonitor::EAsyncWorkType::BuildNaniteMeshes))
 			{
-				if (NumProxiesUpdated < NumMeshesToUpdate && Proxy->GetAsyncWorkMonitor().CheckIfUpdateTriggeredAndClear(FAsyncWorkMonitor::EAsyncWorkType::BuildNaniteMeshes))
-				{
-					NumProxiesUpdated++;
-					Proxy->UpdateNaniteRepresentation(/* const ITargetPlatform* = */nullptr);
-				}
+				NumProxiesUpdated++;
+				Proxy->UpdateNaniteRepresentation(/* const ITargetPlatform* = */nullptr);
 			}
+		}
 #endif //WITH_EDITOR
-			if (Cameras && Proxy->ShouldTickGrass())
-			{
-				Proxy->TickGrass(*Cameras, InOutNumComponentsCreated);
-			}
+		// TODO [chris.tchou] : this stops all async task processing if cameras go away, which might leave tasks dangling
+		if (Cameras && Proxy->ShouldTickGrass())
+		{
+			Proxy->TickGrass(*Cameras, InOutNumComponentsCreated);
+		}
 
 #if !WITH_EDITOR
-			if (GUpdateProxyActorRenderMethodOnTickAtRuntime)
+		if (GUpdateProxyActorRenderMethodOnTickAtRuntime)
 #endif // WITH_EDITOR
-			{
-				Proxy->UpdateRenderingMethod();
-			}
+		{
+			Proxy->UpdateRenderingMethod();
 		}
 	}
 
+	ActiveProxies.Reset();
+
 #if WITH_EDITOR
 	if (GIsEditor && !World->IsPlayInEditor())
 	{
diff --git a/Engine/Source/Runtime/Landscape/Private/LandscapeTextureStreamingManager.cpp b/Engine/Source/Runtime/Landscape/Private/LandscapeTextureStreamingManager.cpp
new file mode 100644
--- /dev/null
+++ b/Engine/Source/Runtime/Landscape/Private/LandscapeTextureStreamingManager.cpp
@@ -0,0 +1,164 @@
+// Copyright Epic Games, Inc. All Rights Reserved.
+
+#include "LandscapeTextureStreamingManager.h"
+#include "Engine/Texture.h"
+
+namespace UE::Landscape
+{
+	// double check that a texture is forced resident
+	static inline void EnsureTextureForcedResident(UTexture *Texture)
+	{
+		// if other systems mess with this flag, then restore it to what it should be
+		// Any code that is directly messing with the flag on one of our
+		// landscape related textures should go through this streaming system instead
+		if (!ensure(Texture->bForceMiplevelsToBeResident))
+		{
+			Texture->bForceMiplevelsToBeResident = true;
+		}
+	}
+}
+
+bool FLandscapeTextureStreamingManager::RequestTextureFullyStreamedIn(UTexture* Texture, bool bWaitForStreaming)
+{
+	TWeakObjectPtr<UTexture> TexturePtr = Texture;
+	FTextureState& State = TextureStates.FindOrAdd(TexturePtr);
+
+	if (State.RequestCount == 0)
+	{
+		Texture->bForceMiplevelsToBeResident = true;
+	}
+	else
+	{
+		UE::Landscape::EnsureTextureForcedResident(Texture);
+	}
+	State.RequestCount++;
+
+	if (IsTextureFullyStreamedIn(Texture))
+	{
+		return true;
+	}
+	else if (bWaitForStreaming)
+	{
+		Texture->WaitForStreaming();
+		return IsTextureFullyStreamedIn(Texture);
+	}
+	return false;
+}
+
+bool FLandscapeTextureStreamingManager::RequestTextureFullyStreamedInForever(UTexture* Texture, bool bWaitForStreaming)
+{
+	TWeakObjectPtr<UTexture> TexturePtr = Texture;
+	FTextureState& State = TextureStates.FindOrAdd(TexturePtr);
+	State.bForever = true;
+	Texture->bForceMiplevelsToBeResident = true;
+
+	if (IsTextureFullyStreamedIn(Texture))
+	{
+		return true;
+	}
+	else if (bWaitForStreaming)
+	{
+		Texture->WaitForStreaming();
+		return IsTextureFullyStreamedIn(Texture);
+	}
+	return false;
+}
+
+void FLandscapeTextureStreamingManager::UnrequestTextureFullyStreamedIn(UTexture* Texture)
+{
+	if (Texture == nullptr)
+	{
+		return;
+	}
+
+	TWeakObjectPtr<UTexture> TexturePtr = Texture;
+	FTextureState* State = TextureStates.Find(TexturePtr);
+	if (State)
+	{
+		if (State->RequestCount > 0)
+		{
+			State->RequestCount--;
+			if (!State->bForever && State->RequestCount <= 0)
+			{
+				// allow stream out, remove tracking
+				Texture->bForceMiplevelsToBeResident = false;
+				TextureStates.Remove(TexturePtr);
+			}
+			else
+			{
+				UE::Landscape::EnsureTextureForcedResident(Texture);
+			}
+		}
+		else
+		{
+			// only way the request count should get to zero is if the texture is flagged as forever streamed.
+			ensure(State->bForever);
+			UE::Landscape::EnsureTextureForcedResident(Texture);
+		}
+	}
+}
+
+bool FLandscapeTextureStreamingManager::WaitForTextureStreaming()
+{
+	TRACE_CPUPROFILER_EVENT_SCOPE(LandscapeTextureStreamingaAnager_WaitForTextureStreaming);
+	bool bFullyStreamed = true;
+	for (auto It = TextureStates.CreateIterator(); It; ++It)
+	{
+		UTexture* Texture = It.Key().Get();
+		if (Texture)
+		{
+			UE::Landscape::EnsureTextureForcedResident(Texture);
+			if (!Texture->IsFullyStreamedIn())
+			{
+				Texture->WaitForStreaming();
+			}
+#if WITH_EDITOR
+			bFullyStreamed = bFullyStreamed && !Texture->IsDefaultTexture();
+#endif // WITH_EDITOR		
+			bFullyStreamed = bFullyStreamed && Texture->IsFullyStreamedIn();
+		}
+		else
+		{
+			// the texture was unloaded... we can remove this entry
+			It.RemoveCurrent();
+		}
+	}
+	return bFullyStreamed;
+}
+
+void FLandscapeTextureStreamingManager::CleanupInvalidEntries()
+{
+	for (auto It = TextureStates.CreateIterator(); It; ++It)
+	{
+		TWeakObjectPtr<UTexture>& TexPtr = It.Key();
+		if (!TexPtr.IsValid())
+		{
+			It.RemoveCurrent();
+		}
+	}
+}
+
+bool FLandscapeTextureStreamingManager::IsTextureFullyStreamedIn(UTexture* InTexture)
+{
+	return InTexture &&
+#if WITH_EDITOR
+		!InTexture->IsDefaultTexture() &&
+#endif // WITH_EDITOR
+		!InTexture->HasPendingInitOrStreaming() && InTexture->IsFullyStreamedIn();
+}
+
+FLandscapeTextureStreamingManager::~FLandscapeTextureStreamingManager()
+{
+	if (!ensure(TextureStates.IsEmpty()))
+	{
+		// clear force stream flag on all textures, just in case
+		for (auto It = TextureStates.CreateIterator(); It; ++It)
+		{
+			UTexture* Texture = It.Key().Get();
+			if (Texture)
+			{
+				Texture->bForceMiplevelsToBeResident = false;
+			}
+		}
+	}
+}
\ No newline at end of file
diff --git a/Engine/Source/Runtime/Landscape/Private/LandscapeTextureStreamingManager.h b/Engine/Source/Runtime/Landscape/Private/LandscapeTextureStreamingManager.h
new file mode 100644
--- /dev/null
+++ b/Engine/Source/Runtime/Landscape/Private/LandscapeTextureStreamingManager.h
@@ -0,0 +1,52 @@
+// Copyright Epic Games, Inc. All Rights Reserved.
+#pragma once
+
+#include "CoreMinimal.h"
+#include "UObject/WeakObjectPtr.h"
+
+class UTexture;
+
+class FLandscapeTextureStreamingManager
+{
+public:
+	// Put in a request to have all streamable mips of a texture streamed in.
+	// bWaitForStreaming true will cause it to block until the texture streaming operation completes.
+	// bWaitForStreaming false will asynchronously stream in the textures, and you can check for completion with IsTextureFullyStreamedIn().
+	// The texture will* remain streamed in until you call UnrequestTextureFullyStreamedIn()
+	// Returns true if the texture is currently fully streamed in.
+	bool RequestTextureFullyStreamedIn(UTexture* Texture, bool bWaitForStreaming);
+
+	// Remove a request to have all streamable mips of a texture streamed in.
+	// The texture may still remain streamed in if someone else has a request,
+	// or if the texture manager doesn't see any need to evict it to reclaim texture memory.
+	void UnrequestTextureFullyStreamedIn(UTexture* Texture);
+
+	// Request that this texture be fully streamed in for as long as the texture is loaded. No backsies.
+	// bWaitForStreaming will also block until the texture is fully streamed in.
+	// Returns true if the texture is fully streamed.
+	bool RequestTextureFullyStreamedInForever(UTexture* Texture, bool bWaitForStreaming);
+
+	// TODO [chris.tchou] : a non-blocking tick to update streaming textures would be more efficient in some looping cases.
+	// void TickStreaming();
+
+	// Synchronously waits for all textures to complete their requested streaming. Returns true if successful.
+	// This necessary if you are requesting streaming and waiting for the textures to be streamed within a local loop.
+	// (normally they don't complete streaming until a certain point in the update tick)
+	bool WaitForTextureStreaming();
+
+	// Call this clean up any old entries in tracked TextureStates for textures that have been unloaded without first being Unrequested.
+	void CleanupInvalidEntries();
+
+	static bool IsTextureFullyStreamedIn(UTexture* Texture);
+
+	~FLandscapeTextureStreamingManager();
+
+private:
+	struct FTextureState
+	{
+		int32 RequestCount = 0;
+		bool bForever = false;
+	};
+
+	TMap<TWeakObjectPtr<UTexture>, FTextureState, FDefaultSetAllocator, TWeakObjectPtrMapKeyFuncs<TWeakObjectPtr<UTexture>, FTextureState>> TextureStates;
+};
diff --git a/Engine/Source/Runtime/Landscape/Public/LandscapeRender.h b/Engine/Source/Runtime/Landscape/Public/LandscapeRender.h
--- a/Engine/Source/Runtime/Landscape/Public/LandscapeRender.h
+++ b/Engine/Source/Runtime/Landscape/Public/LandscapeRender.h
@@ -96,7 +96,8 @@ LANDSCAPE_API extern TObjectPtr<UMaterialInterface> GLandscapeDirtyMaterial;
 
 namespace UE::Landscape
 {
-bool NeedsFixedGridVertexFactory(EShaderPlatform InShaderPlatform);
+	bool NeedsFixedGridVertexFactory(EShaderPlatform InShaderPlatform);
+	bool ShouldBuildGrassMapRenderingResources();
 } // namespace UE::Landscape
 
 /** The uniform shader parameters for a landscape draw call. */
@@ -368,13 +369,15 @@ public:
 	FLandscapeVertexFactory* TileVertexFactory;
 	FVertexBuffer* TileDataBuffer;
 	
+	// array per mip level, storing FIndexBuffer pointers
 	FIndexBuffer** IndexBuffers;
+
+	// array per mip level, storing index buffer ranges per subssection, and for the entire component
 	FLandscapeIndexRanges* IndexRanges;
+
 	bool bUse32BitIndices;
-#if WITH_EDITOR
 	FIndexBuffer* GrassIndexBuffer;
 	TArray<int32, TInlineAllocator<8>> GrassIndexMipOffsets;
-#endif
 
 #if RHI_RAYTRACING
 	TArray<FIndexBuffer*> ZeroOffsetIndexBuffers;
@@ -385,10 +388,8 @@ public:
 	template <typename INDEX_TYPE>
 	void CreateIndexBuffers(FRHICommandListBase& RHICmdList, const FName& OwnerName);
 	
-#if WITH_EDITOR
 	template <typename INDEX_TYPE>
 	void CreateGrassIndexBuffer(FRHICommandListBase& RHICmdList, const FName& InOwnerName);
-#endif
 
 	LANDSCAPE_API virtual ~FLandscapeSharedBuffers();
 };
@@ -782,11 +783,9 @@ protected:
 	bool bUsesLandscapeCulling;
 
 
-#if WITH_EDITOR
 	// Precomputed grass rendering MeshBatch and per-LOD params
 	FMeshBatch                           GrassMeshBatch;
 	TArray<FLandscapeBatchElementParams> GrassBatchParams;
-#endif
 
 	FVector4f WeightmapScaleBias;
 	TArray<UTexture2D*> WeightmapTextures;
@@ -888,9 +887,7 @@ public:
 	friend class FLandscapeVertexFactoryPixelShaderParameters;
 	friend struct FLandscapeBatchElementParams;
 
-#if WITH_EDITOR
 	const FMeshBatch& GetGrassMeshBatch() const { return GrassMeshBatch; }
-#endif
 
 	// FLandcapeSceneProxy
 	LANDSCAPE_API void ChangeComponentScreenSizeToUseSubSections_RenderThread(float InComponentScreenSizeToUseSubSections);
diff --git a/Engine/Source/Runtime/Landscape/Public/LandscapeSubsystem.h b/Engine/Source/Runtime/Landscape/Public/LandscapeSubsystem.h
--- a/Engine/Source/Runtime/Landscape/Public/LandscapeSubsystem.h
+++ b/Engine/Source/Runtime/Landscape/Public/LandscapeSubsystem.h
@@ -12,7 +12,12 @@ class AWorldSettings;
 class IConsoleVariable;
 class ULandscapeInfo;
 class FLandscapeNotificationManager;
+class ULandscapeComponent;
+class FLandscapeGrassMapsBuilder;
+class FLandscapeTextureStreamingManager;
 struct FDateTime;
+struct FScopedSlowTask;
+
 namespace UE::Landscape
 {
 	enum class EOutdatedDataFlags : uint8;
@@ -57,8 +62,11 @@ public:
 	virtual TStatId GetStatId() const override;
 	// End FTickableGameObject overrides
 
+	// setting this to true causes grass instance generation to go wider (multiplies the limits by GGrassCreationPrioritizedMultipler)
 	void PrioritizeGrassCreation(bool bPrioritizeGrassCreation) { bIsGrassCreationPrioritized = bPrioritizeGrassCreation; }
 	bool IsGrassCreationPrioritized() const { return bIsGrassCreationPrioritized; }
+	FLandscapeGrassMapsBuilder* GetGrassMapBuilder() { return GrassMapsBuilder; }
+	FLandscapeTextureStreamingManager* GetTextureStreamingManager() { return TextureStreamingManager; }
 
 	/**
 	 * Can be called at runtime : (optionally) flushes grass on all landscape components and updates them
@@ -67,9 +75,18 @@ public:
 	 * @param InOptionalCameraLocations : (optional) camera locations that should be used when updating the grass. If not specified, the usual (streaming manager-based) view locations will be used
 	 */
 	LANDSCAPE_API void RegenerateGrass(bool bInFlushGrass, bool bInForceSync, TOptional<TArrayView<FVector>> InOptionalCameraLocations = TOptional<TArrayView<FVector>>());
+	
+	// Remove all grass instances from the specified components.  If passed null, removes all grass instances from all proxies.
+	void RemoveGrassInstances(const TSet<ULandscapeComponent*>* ComponentsToRemoveGrassInstances = nullptr);
+
+	// called when components are registered to the world	
+	void RegisterComponent(ULandscapeComponent* Component);
+	void UnregisterComponent(ULandscapeComponent* Component);
 
 #if WITH_EDITOR
 	LANDSCAPE_API void BuildAll();
+
+	// Synchronously build grass maps for all components
 	LANDSCAPE_API void BuildGrassMaps();
 	LANDSCAPE_API void BuildPhysicalMaterial();
 
@@ -108,7 +125,6 @@ public:
 	FDateTime GetAppCurrentDateTime();
 	LANDSCAPE_API void AddAsyncEvent(FGraphEventRef GraphEventRef);
 
-
 	// Returns true if we should build nanite meshes in parallel asynchronously. 
 	bool IsMultithreadedNaniteBuildEnabled();
 
@@ -137,8 +153,10 @@ private:
 	TArray<TWeakObjectPtr<ALandscapeProxy>> Proxies;
 	FDelegateHandle OnNaniteWorldSettingsChangedHandle;
 
+	FLandscapeTextureStreamingManager* TextureStreamingManager = nullptr;
+	FLandscapeGrassMapsBuilder* GrassMapsBuilder = nullptr;
+
 #if WITH_EDITOR
-	class FLandscapeGrassMapsBuilder* GrassMapsBuilder = nullptr;
 	class FLandscapePhysicalMaterialBuilder* PhysicalMaterialBuilder = nullptr;
 	
 	FLandscapeNotificationManager* NotificationManager = nullptr;
