Adds MIC recompilation to DeleteUnusedLayers. Allows the specialized landscape MIC to be updated without the unused layers.

#localization none
#tests Delete Unused Layer on baked Landscapes

diff --git a/Engine/Source/Runtime/Landscape/Private/LandscapeEditLayers.cpp b/Engine/Source/Runtime/Landscape/Private/LandscapeEditLayers.cpp
--- a/Engine/Source/Runtime/Landscape/Private/LandscapeEditLayers.cpp
+++ b/Engine/Source/Runtime/Landscape/Private/LandscapeEditLayers.cpp
@@ -8788,9 +8788,10 @@ void ALandscapeProxy::UpdateCachedHasLayersContent(bool InCheckComponentDataInte
 
 namespace
 {
-	void DeleteUnusedLayersImpl(ULandscapeComponent* InComponent, const FGuid& InLayerGuid)
+	bool DeleteUnusedLayersImpl(ULandscapeComponent* InComponent, const FGuid& InLayerGuid)
 	{
 		TArray<FWeightmapLayerAllocationInfo>& ComponentWeightmapLayerAllocations = InComponent->GetWeightmapLayerAllocations(InLayerGuid);
+		bool bWasModified = false;
 
 		for (int32 LayerIdx = 0; LayerIdx < ComponentWeightmapLayerAllocations.Num();)
 		{
@@ -8817,18 +8818,32 @@ namespace
 			constexpr bool bShouldDirtyPackage = true;
 
 			// If DeleteLayerIfAllZero returns true, We just removed the current layer allocation, so we need to iterate on the new current index.
-			if (!InComponent->DeleteLayerIfAllZero(InLayerGuid, TextDataPtr, Texture->GetSizeX(), LayerIdx, bShouldDirtyPackage))
+			if (InComponent->DeleteLayerIfAllZero(InLayerGuid, TextDataPtr, Texture->GetSizeX(), LayerIdx, bShouldDirtyPackage))
+			{
+				bWasModified = true;
+			}
+			else
 			{
 				++LayerIdx;
 			}
 
 			Texture->Source.UnlockMip(0);
 		}
+
+		if (bWasModified)
+		{
+			InComponent->UpdateMaterialInstances();
+			InComponent->MarkRenderStateDirty();
+		}
+
+		return bWasModified;
 	}
 }
 
 void ALandscapeProxy::DeleteUnusedLayers()
 {
+	bool bWasModified = false;
+	
 	for (ULandscapeComponent* Component : LandscapeComponents)
 	{
 		if (Component == nullptr)
@@ -8836,13 +8851,18 @@ void ALandscapeProxy::DeleteUnusedLayers()
 			continue;
 		}
 
-		Component->ForEachLayer([Component](const FGuid& LayerGuid, FLandscapeLayerComponentData& LayerData)
+		Component->ForEachLayer([Component, &bWasModified](const FGuid& LayerGuid, FLandscapeLayerComponentData& LayerData)
 		{
-			DeleteUnusedLayersImpl(Component, LayerGuid);
+			bWasModified = DeleteUnusedLayersImpl(Component, LayerGuid);
 		});
 
 		// Execute ClearUnusedLayersImpl on the final Layer.
-		DeleteUnusedLayersImpl(Component, FGuid());
+		bWasModified = DeleteUnusedLayersImpl(Component, FGuid());
+
+		if (bWasModified)
+		{
+			InvalidateNaniteRepresentation(false);
+		}
 	}
 }
 
