Bugfixes for ULandscapeSplinesComponent::CopyToSplineComponent
- Account for reversed spline meshes within non-reversed landscape spline segments
- Added some QoL for auto detecting looping landscape splines

diff --git a/Engine/Source/Runtime/Landscape/Private/LandscapeSplines.cpp b/Engine/Source/Runtime/Landscape/Private/LandscapeSplines.cpp
--- a/Engine/Source/Runtime/Landscape/Private/LandscapeSplines.cpp
+++ b/Engine/Source/Runtime/Landscape/Private/LandscapeSplines.cpp
@@ -789,9 +789,10 @@ void ULandscapeSplinesComponent::CopyToSplineComponent(USplineComponent* SplineC
 	ULandscapeSplineControlPoint* LastControlPoint = nullptr;
 
 	// Arrive Tangent to apply to the next spline point.  Populated as the previous spline mesh's End Tangent
-	FVector ArriveTangent;
-	
-	ULandscapeSplineSegment* CurrentSegment = Segments[0];
+	FVector ArriveTangent(0,0,0);
+
+	ULandscapeSplineSegment* FirstSegment = Segments[0];
+	ULandscapeSplineSegment* CurrentSegment = FirstSegment;
 	do
 	{
 		if (!ensure(NumIterations < MAX_ITERATIONS))
@@ -814,65 +815,71 @@ void ULandscapeSplinesComponent::CopyToSplineComponent(USplineComponent* SplineC
 
 		VisitedSegments.Add(CurrentSegment);
 
-		// No guarantee if the control point we're looking at is the segments start or end control point
+		// No guarantee if the control point we're looking at is the segment's start or end control point
 		// If it's the end control point, we must reverse the direction we traverse over this segment
 		// Connections[0] == Start; Connections[1] == End;
-		bool bReverseTraversal = CurrentSegment->Connections[1].ControlPoint == LastControlPoint;
+		const bool bReverseTraversal = CurrentSegment->Connections[1].ControlPoint == LastControlPoint;
+		const FLandscapeSplineSegmentConnection& CurrentConnection = bReverseTraversal ? CurrentSegment->Connections[1] : CurrentSegment->Connections[0];
+		const FLandscapeSplineSegmentConnection& NextConnection = bReverseTraversal ? CurrentSegment->Connections[0] : CurrentSegment->Connections[1];
 
-		// Walk down the spline meshes
 		TArray<USplineMeshComponent*> SegmentSplineMeshComponents = CurrentSegment->GetLocalMeshComponents();
-		if (!bReverseTraversal)
-		{
-			for (int Index = 0; Index < SegmentSplineMeshComponents.Num(); ++Index)
-			{
-				USplineMeshComponent* SplineMesh = SegmentSplineMeshComponents[Index];
-				if (SplineMesh == nullptr)
-				{
-					continue;
-				}
 
-				// Create a spline component point from each spline mesh in the segment
-				SplineComponent->AddSplinePoint(SplineMesh->GetStartPosition() + SplineMesh->GetComponentLocation(), ESplineCoordinateSpace::World, false);
-				SplineComponent->SetTangentsAtSplinePoint(CurrentPointIndex, ArriveTangent, SplineMesh->GetStartTangent(), ESplineCoordinateSpace::World, false);
-
-				++CurrentPointIndex;
-				ArriveTangent = SplineMesh->GetEndTangent();
-			}
+		// No guarantee on the direction of individual spline mesh component directions
+		// Test the first one and determine if the start or end location is closer
+		bool bReverseSplineMeshes = false;
+		const FVector ControlPointWorldLocation = GetComponentTransform().TransformPosition(CurrentConnection.ControlPoint->Location);
+		if (const USplineMeshComponent* FirstSplineMesh = SegmentSplineMeshComponents[bReverseTraversal ? SegmentSplineMeshComponents.Num() - 1 : 0])
+		{
+			const float StartDistanceSqr = FVector::DistSquared(FirstSplineMesh->GetStartPosition() + FirstSplineMesh->GetComponentLocation(), ControlPointWorldLocation);
+			const float EndDistanceSqr = FVector::DistSquared(FirstSplineMesh->GetEndPosition() + FirstSplineMesh->GetComponentLocation(), ControlPointWorldLocation);
+			bReverseSplineMeshes = EndDistanceSqr < StartDistanceSqr;
 		}
-		else
+		
+		// Walk down the spline meshes
+		for (int32 Index = 0; Index < SegmentSplineMeshComponents.Num(); ++Index)
 		{
-			// Same as the true block, but reversed direction
-			for (int Index = SegmentSplineMeshComponents.Num() - 1; Index >= 0; --Index)
+			const int32 ActualIndex = bReverseTraversal ? SegmentSplineMeshComponents.Num() - (1 + Index) : Index;
+			
+			USplineMeshComponent* SplineMesh = SegmentSplineMeshComponents[ActualIndex];
+			if (SplineMesh == nullptr)
 			{
-				USplineMeshComponent* SplineMesh = SegmentSplineMeshComponents[Index];
-				if (SplineMesh == nullptr)
-				{
-					continue;
-				}
+				continue;
+			}
 
-				// Swap start with end since we're going in reverse
-				SplineComponent->AddSplinePoint(SplineMesh->GetEndPosition() + SplineMesh->GetComponentLocation(), ESplineCoordinateSpace::World, false);
-				SplineComponent->SetTangentsAtSplinePoint(CurrentPointIndex, ArriveTangent, SplineMesh->GetEndTangent(), ESplineCoordinateSpace::World, false);
+			FVector Position = bReverseSplineMeshes ? SplineMesh->GetEndPosition() : SplineMesh->GetStartPosition();
+			Position += SplineMesh->GetComponentLocation();
+			
+			FVector LeaveTangent = bReverseSplineMeshes ? -SplineMesh->GetEndTangent() : SplineMesh->GetStartTangent();
 
-				++CurrentPointIndex;
-				ArriveTangent = SplineMesh->GetStartTangent();
-			}
+			// Create a spline component point from each spline mesh in the segment
+			SplineComponent->AddSplinePoint(Position, ESplineCoordinateSpace::World, false);
+			SplineComponent->SetTangentsAtSplinePoint(CurrentPointIndex, ArriveTangent, LeaveTangent, ESplineCoordinateSpace::World, false);
+			
+			ArriveTangent = bReverseSplineMeshes ? SplineMesh->GetStartPosition() : SplineMesh->GetEndTangent();
+				
+			++CurrentPointIndex;
 		}
 		
 		// Search for the next segment
-		FLandscapeSplineSegmentConnection& Connection = bReverseTraversal ? CurrentSegment->Connections[0] : CurrentSegment->Connections[1];
 		bool bFoundNextSegment = false;
-		if (Connection.ControlPoint)
+		bool bLoopingSpline = false;
+		if (NextConnection.ControlPoint)
 		{
-			for (FLandscapeSplineConnection& ControlPointConnection : Connection.ControlPoint->ConnectedSegments)
+			for (FLandscapeSplineConnection& ControlPointConnection : NextConnection.ControlPoint->ConnectedSegments)
 			{
 				if (ControlPointConnection.Segment == nullptr || VisitedSegments.Contains(ControlPointConnection.Segment))
 				{
+					// Check for loop
+					if (ControlPointConnection.Segment == FirstSegment)
+					{
+						bLoopingSpline = true;
+					}
+					
 					continue;
 				}
 				
 				CurrentSegment = ControlPointConnection.Segment;
-				LastControlPoint = Connection.ControlPoint;
+				LastControlPoint = NextConnection.ControlPoint;
 				bFoundNextSegment = true;
 				break;
 			}
@@ -880,6 +887,18 @@ void ULandscapeSplinesComponent::CopyToSplineComponent(USplineComponent* SplineC
 
 		if (!bFoundNextSegment)
 		{
+			SplineComponent->SetClosedLoop(bLoopingSpline);
+			
+			if (bLoopingSpline)
+			{
+				// Populate first point's arrive tangent
+				if (SplineComponent->GetNumberOfSplinePoints() > 0)
+				{
+					const FVector LeaveTangent = SplineComponent->SplineCurves.Position.Points[0].LeaveTangent;
+					SplineComponent->SetTangentsAtSplinePoint(0, ArriveTangent, LeaveTangent, ESplineCoordinateSpace::World, false);
+				}
+			}
+			
 			UE_LOG(LogLandscape, Verbose, TEXT("%s Could not find next segment or all segments have been traversed."), ANSI_TO_TCHAR(__FUNCTION__));
 			break;
 		}
