Shrink exclude boxes arrays in Landscape system

#localization none
#tests local replay runs

diff --git a/Engine/Source/Runtime/Landscape/Private/LandscapeGrass.cpp b/Engine/Source/Runtime/Landscape/Private/LandscapeGrass.cpp
--- a/Engine/Source/Runtime/Landscape/Private/LandscapeGrass.cpp
+++ b/Engine/Source/Runtime/Landscape/Private/LandscapeGrass.cpp
@@ -2418,6 +2418,8 @@ void ALandscapeProxy::RemoveInvalidExclusionBoxes()
 			GGrassExclusionChangeTag++;
 		}
 	}
+
+	GGrassExclusionBoxes.Compact();
 }
 
 #if WITH_EDITOR
@@ -2588,6 +2590,8 @@ void ALandscapeProxy::UpdateGrass(const TArray<FVector>& Cameras, int32& InOutNu
 							}
 						}
 
+						Component->ActiveExcludedBoxes.Shrink();
+
 						//Sort exclude boxes by their volume size, as the biggest boxes are more likely to exclude points and early out when building grass
 						Component->ActiveExcludedBoxes.Sort([](const ULandscapeComponent::FExcludeBox& A, const ULandscapeComponent::FExcludeBox& B) {
 								return A.Box.GetVolume() > B.Box.GetVolume();
