Landscape: check against negative cell sizes for TPointHashGrid3

Use TPointHashGrid3 instead of TPointHashGrid2 as the export function can be used to export landscapes to worldspace.

diff --git a/Engine/Source/Runtime/Landscape/Private/LandscapeEdit.cpp b/Engine/Source/Runtime/Landscape/Private/LandscapeEdit.cpp
--- a/Engine/Source/Runtime/Landscape/Private/LandscapeEdit.cpp
+++ b/Engine/Source/Runtime/Landscape/Private/LandscapeEdit.cpp
@@ -92,7 +92,7 @@ LandscapeEdit.cpp: Landscape editing
 #include "Algo/ForEach.h"
 #include "Serialization/MemoryWriter.h"
 #include "Engine/Canvas.h"
-#include "Spatial/PointHashGrid2.h"
+#include "Spatial/PointHashGrid3.h"
 #include "Engine/Texture2DArray.h"
 
 DEFINE_LOG_CATEGORY(LogLandscape);
@@ -4074,9 +4074,11 @@ bool ALandscapeProxy::ExportToRawMeshDataCopy(const FRawMeshExportParams& InExpo
 		};
 
 		// size the TPointHashGrid2 cells to be the size to something smaller than a landscape quad
-		const FVector TransformedCellSize = ComponentToExportCoordinatesTransform.TransformVector(FVector( 0.1f, 0.1f, 0.1f ));
-		const float CellSize = FMath::Min(TransformedCellSize.X , TransformedCellSize.Y);
-		UE::Geometry::TPointHashGrid2<FVertexID, float> DeduplicatedVertexIDs(CellSize, INDEX_NONE);
+		FBox Box(FVector(0, 0, 0), FVector(0.1f, 0.1f, 0.1f));
+		FBox TransformedBox = Box.TransformBy(ComponentToExportCoordinatesTransform.ToMatrixWithScale());
+		const float CellSize = TransformedBox.GetSize().GetAbsMax();
+		check(CellSize > 0.0f);
+		UE::Geometry::TPointHashGrid3<FVertexID, float> DeduplicatedVertexIDs(CellSize, INDEX_NONE);
 		DeduplicatedVertexIDs.Reserve(( PaddedComponentSizeXQuadsLOD + 1) * (PaddedComponentSizeYQuadsLOD + 1));
 		
 		// Export to MeshDescription
@@ -4130,18 +4132,17 @@ bool ALandscapeProxy::ExportToRawMeshDataCopy(const FRawMeshExportParams& InExpo
 					for (int32 i = 0; i < NewPositions.Num(); i++)
 					{
 						const FVector3f NewPos (NewPositions[i]);
-						const FVector2f NewPos2D (NewPos.X, NewPos.Y);
 						TPair<FVertexID, float> ExistingVertexID = DeduplicatedVertexIDs.FindNearestInRadius(
-							NewPos2D, TMathUtilConstants<float>::ZeroTolerance, 
-							[&VertexPositions, NewPos2D](const FVertexID& VertexID)
+							NewPos, TMathUtilConstants<float>::ZeroTolerance, 
+							[&VertexPositions, NewPos](const FVertexID& VertexID)
 							{
-								return FVector2f::DistSquared(FVector2f(VertexPositions[VertexID].X, VertexPositions[VertexID].Y), NewPos2D);
+								return FVector3f::DistSquared(VertexPositions[VertexID], NewPos);
 							});
 						FVertexID VertexID;
 						if (ExistingVertexID.Key == INDEX_NONE)
 						{
 							VertexID = OutRawMesh.CreateVertex();
-							DeduplicatedVertexIDs.InsertPointUnsafe(VertexID, NewPos2D);
+							DeduplicatedVertexIDs.InsertPointUnsafe(VertexID, NewPos);
 							VertexPositions[VertexID] = NewPos;
 						}
 						else
@@ -5668,8 +5669,7 @@ bool ALandscapeProxy::CanEditChange(const FProperty* InProperty) const
 			|| (PropertyName == GET_MEMBER_NAME_CHECKED(ALandscapeProxy, bEnableNanite))
 			|| (PropertyName == GET_MEMBER_NAME_CHECKED(ALandscapeProxy, bNaniteSkirtEnabled))
 			|| (PropertyName == GET_MEMBER_NAME_CHECKED(ALandscapeProxy, NaniteSkirtDepth))
-			|| (PropertyName == GET_MEMBER_NAME_CHECKED(ALandscapeProxy, NaniteLODIndex))
-			|| (PropertyName == GET_MEMBER_NAME_CHECKED(ALandscapeProxy, NanitePositionPrecision)))
+			|| (PropertyName == GET_MEMBER_NAME_CHECKED(ALandscapeProxy, NaniteLODIndex)))
 		{
 			return false;
 		}
@@ -5863,16 +5863,6 @@ void ALandscapeProxy::PostEditChangeProperty(FPropertyChangedEvent& PropertyChan
 		InvalidateGeneratedComponentData(/* bInvalidateLightingCache = */false);
 		MarkComponentsRenderStateDirty();
 	}
-	if (GIsEditor && PropertyName == GET_MEMBER_NAME_CHECKED(ALandscapeProxy, NanitePositionPrecision))
-	{
-		// This property is entirely shared with the parent material so don't let it be set to a different value than the parent : 
-		if (ALandscape* Parent = GetLandscapeActor(); (Parent != this) && (NanitePositionPrecision != Parent->GetNanitePositionPrecision()))
-		{
-			NanitePositionPrecision = Parent->GetNanitePositionPrecision();
-		}
-		InvalidateGeneratedComponentData(/* bInvalidateLightingCache = */false);
-		MarkComponentsRenderStateDirty();
-	}
 	else if (PropertyName == GET_MEMBER_NAME_CHECKED(ALandscapeProxy, bUseDynamicMaterialInstance))
 	{
 		MarkComponentsRenderStateDirty();
