Fix an infinite loop on save when Grass.Enabled false.

Also fixed an assert when doing a blocking grass update: if you don't run a full update every loop then sanity checks will fail.

Also move the allocator pool to be a member instead of global static, and speed up save by using blocking GPU submit calls.

diff --git a/Engine/Source/Runtime/Landscape/Private/LandscapeGrassMapsBuilder.cpp b/Engine/Source/Runtime/Landscape/Private/LandscapeGrassMapsBuilder.cpp
--- a/Engine/Source/Runtime/Landscape/Private/LandscapeGrassMapsBuilder.cpp
+++ b/Engine/Source/Runtime/Landscape/Private/LandscapeGrassMapsBuilder.cpp
@@ -92,8 +92,6 @@ FLandscapeGrassMapsBuilder::FLandscapeGrassMapsBuilder(UWorld* InOwner, FLandsca
 {
 }
 
-TAllocatorFixedSizeFreeList<sizeof(FLandscapeGrassMapsBuilder::FComponentState), 32> FLandscapeGrassMapsBuilder::StatePoolAllocator;
-
 namespace UE::Landscape
 {
 	uint32 ComputeGrassMapGenerationHash(const ULandscapeComponent* Component, UMaterialInterface* Material)
@@ -244,19 +242,23 @@ FLandscapeGrassMapsBuilder::~FLandscapeGrassMapsBuilder()
 	{
 		FComponentState* State = It.Value();
 		ULandscapeComponent* Component = State->Component;
-		if (!ensure(Component == nullptr))
+		if (Component != nullptr)
 		{
+			// This happens when deleting a level, the components are not unregistered before the world is destroyed.
 			UnregisterComponent(Component);
 		}
 		State->TickCount = 7777;	// set to a large number so we don't wait to delete entries
 	}
 
-	// update component state until they all delete themselves (should happen on the first update unless a readback is active, and shouldn't take more than 3)
+	// update component state until they all delete themselves
+	// (this should happen on the first update, unless a GPU readback is active.
+	// And it shouldn't take more than 3 update if there is a GPU readback.
 	int32 Iterations = 0;
-	while (ensure(Iterations < 6) && ComponentStates.Num() > 0)
+	while (Iterations < 6 && ComponentStates.Num() > 0)
 	{
 		TArray<FVector> EmptyCamerasArray;
-		UpdateTrackedComponents(EmptyCamerasArray, 0);
+		int32 UpdateAllComponentCount = ComponentStates.Num();
+		UpdateTrackedComponents(EmptyCamerasArray, 0, UpdateAllComponentCount, /* bCancelAndEvictAll = */ true);
 
 		// if we happen to get caught with a GPU readback in flight, submit GPU commands to make sure it moves forward
 		if (RenderingCount > 0)
@@ -265,6 +267,20 @@ FLandscapeGrassMapsBuilder::~FLandscapeGrassMapsBuilder()
 		}
 		Iterations++;
 	}
+
+	if (!ensure(ComponentStates.Num() == 0))
+	{
+		// somehow we failed to free the components the right way (either a GPU readback is stuck, or state logic is broken)
+		// force free the state anyways and hope for the best
+		for (auto It = ComponentStates.CreateIterator(); It; ++It)
+		{
+			FComponentState* State = It.Value();
+			UE_LOG(LogGrass, Warning, TEXT("Failed to clear grass data state after %d iterations (stage:%d ticks:%d), forcing deletion of the state."), Iterations, State->Stage, State->TickCount);
+			State->~FComponentState();
+			StatePoolAllocator.Free(State);
+			It.RemoveCurrent();
+		}
+	}
 }
 
 FLandscapeGrassMapsBuilder::FComponentState::FComponentState(ULandscapeComponent* Component)
@@ -290,7 +306,7 @@ void FLandscapeGrassMapsBuilder::FPendingComponent::UpdatePriorityDistance(const
 	PriorityKey = MinSqrDistanceToComponent;
 }
 
-bool FLandscapeGrassMapsBuilder::UpdateTrackedComponents(const TArray<FVector>& Cameras, int32 LocalMaxRendering)
+bool FLandscapeGrassMapsBuilder::UpdateTrackedComponents(const TArray<FVector>& Cameras, int32 LocalMaxRendering, int32 MaxExpensiveUpdateChecksToPerform, bool bCancelAndEvictAll)
 {
 	SCOPE_CYCLE_COUNTER(STAT_UpdateComponentGrassMaps);
 
@@ -304,7 +320,7 @@ bool FLandscapeGrassMapsBuilder::UpdateTrackedComponents(const TArray<FVector>&
 	static TSet<ULandscapeComponent*> ComponentsToRemoveFoliageInstances;
 	ComponentsToRemoveFoliageInstances.Reset();
 
-	AmortizedUpdate.StartUpdateTick(ComponentStates.Num(), GGrassMapMaxDiscardChecksPerFrame);
+	AmortizedUpdate.StartUpdateTick(ComponentStates.Num(), MaxExpensiveUpdateChecksToPerform);
 
 	// Iterate our components, removing invalid ones and counting how many are in each state.
 	// We can also immediately process any components in the populated or rendering states
@@ -318,95 +334,96 @@ bool FLandscapeGrassMapsBuilder::UpdateTrackedComponents(const TArray<FVector>&
 		ULandscapeComponent* Component = State->Component;
 		State->TickCount++;
 
-		if (Component != nullptr)
+		if (bCancelAndEvictAll || (Component == nullptr))
+		{
+			// fallthrough to CancelAndEvict below
+		}
+		else
 		{
-			if (GGrassEnable)
+			switch (State->Stage)
 			{
-				switch (State->Stage)
+			case EComponentStage::Pending:
+				if (World->IsGameWorld())
 				{
-				case EComponentStage::Pending:
-					if (World->IsGameWorld())
+					// if runtime grass generation is disabled, go straight to not ready.
+					if (!Component->GetLandscapeProxy()->bUseRuntimeGrassMapGeneration)
 					{
-						// if runtime grass generation is disabled, go straight to not ready.
-						if (!Component->GetLandscapeProxy()->bUseRuntimeGrassMapGeneration)
-						{
-							PendingToNotReady(*State);
-							continue; // next!
-						}
+						PendingToNotReady(*State);
+						continue; // next!
 					}
-					continue; // next!
-				case EComponentStage::NotReady:
-					// in game, any not ready component will never become ready.
-					// in editor, check to see if the conditions changed.
-					if (World->IsGameWorld() || !UE::Landscape::CanRenderGrassMap(Component))
+				}
+				continue; // next!
+			case EComponentStage::NotReady:
+				// in game, any not ready component will never become ready.
+				// in editor, check to see if the conditions changed.
+				if (World->IsGameWorld() || !UE::Landscape::CanRenderGrassMap(Component))
+				{
+					continue; // still not ready - next!
+				}
+				break; // cancel and evict to restart build process
+			case EComponentStage::TextureStreaming:
+				// don't process streaming states yet -- first process rendering states to free up slots
+				StreamingStatesToProcess.Add(State);
+				continue; // next!
+
+			case EComponentStage::Rendering:
+				check(State->ActiveRender != nullptr);
+				if (State->ActiveRender->CheckAndUpdateAsyncReadback())
+				{
+					// TODO [chris.tchou] We could move this to an async task, it does have a decent cost to it
+					PopulateGrassDataFromReadback(*State);
+					bChanged = true;
+				}
+				continue; // next!
+
+			case EComponentStage::GrassMapsPopulated:
+				// only check for invalidation of populated grass maps once in a while
+				if (AmortizedUpdate.ShouldUpdate(ComponentStateIndex))
+				{
+					// detect if grass data has been cleared by someone manually calling Flush (i.e. when landscape edits are made)
+					if (!Component->GrassData->HasValidData())
 					{
-						continue; // still not ready - next!
+						break; // cancel and evict to restart process
 					}
-					break; // cancel and evict to restart build process
-				case EComponentStage::TextureStreaming:
-					// don't process streaming states yet -- first process rendering states to free up slots
-					StreamingStatesToProcess.Add(State);
-					continue; // next!
-
-				case EComponentStage::Rendering:
-					check(State->ActiveRender != nullptr);
-					if (State->ActiveRender->CheckAndUpdateAsyncReadback())
+
+					// check if the component is too far from the camera and we can reclaim the grass data
+					if (GGrassMapUseRuntimeGeneration &&
+						State->IsBeyondEvictionRange(Cameras))
 					{
-						// TODO [chris.tchou] We could move this to an async task, it does have a decent cost to it
-						PopulateGrassDataFromReadback(*State);
-						bChanged = true;
+						GRASS_DEBUG_LOG(TEXT("Evicting for being beyond eviction range"));
+						break; // cancel and evict
 					}
-					continue; // next!
 
-				case EComponentStage::GrassMapsPopulated:
-					// only check for invalidation of populated grass maps once in a while
-					if (AmortizedUpdate.ShouldUpdate(ComponentStateIndex))
+#if WITH_EDITOR
+					UMaterialInterface* Material = Component->GetLandscapeMaterial();
+					if (Material == nullptr)
 					{
-						// detect if grass data has been cleared by someone manually calling Flush (i.e. when landscape edits are made)
-						if (!Component->GrassData->HasValidData())
-						{
-							break; // cancel and evict to restart process
-						}
-
-						// check if the component is too far from the camera and we can reclaim the grass data
-						if (GGrassMapUseRuntimeGeneration &&
-							State->IsBeyondEvictionRange(Cameras))
-						{
-							GRASS_DEBUG_LOG(TEXT("Evicting for being beyond eviction range"));
-							break; // cancel and evict
-						}
+						break; // cancel and evict
+					}
 
-#if WITH_EDITOR
-						UMaterialInterface* Material = Component->GetLandscapeMaterial();
-						if (Material == nullptr)
-						{
-							break; // cancel and evict
-						}
-
-						// check if any dependencies changed
-						uint32 CurGrassMapGenerationHash = UE::Landscape::ComputeGrassMapGenerationHash(Component, Material);
-						if (State->GrassMapGenerationHash != CurGrassMapGenerationHash)
-						{
-							break; // cancel and evict to restart process
-						}
-
-						// check if any grass types have changed -- this invalidates foliage instances but not the grass maps
-						const TArray<TObjectPtr<ULandscapeGrassType>>& GrassTypes = Component->GetGrassTypes();
-						uint32 CurGrassInstanceGenerationHash = UE::Landscape::ComputeGrassInstanceGenerationHash(CurGrassMapGenerationHash, GrassTypes);
-						if (State->GrassInstanceGenerationHash != CurGrassInstanceGenerationHash)
-						{
-							ComponentsToRemoveFoliageInstances.Add(Component);
-							Component->InvalidateGrassTypeSummary();
-							State->GrassInstanceGenerationHash = CurGrassInstanceGenerationHash;
-						}
-#endif // WITH_EDITOR
+					// check if any dependencies changed
+					uint32 CurGrassMapGenerationHash = UE::Landscape::ComputeGrassMapGenerationHash(Component, Material);
+					if (State->GrassMapGenerationHash != CurGrassMapGenerationHash)
+					{
+						break; // cancel and evict to restart process
 					}
-					continue;	// next!
 
-				default:
-					check(false);	// unreachable
-					break;			// cancel and evict
+					// check if any grass types have changed -- this invalidates foliage instances but not the grass maps
+					const TArray<TObjectPtr<ULandscapeGrassType>>& GrassTypes = Component->GetGrassTypes();
+					uint32 CurGrassInstanceGenerationHash = UE::Landscape::ComputeGrassInstanceGenerationHash(CurGrassMapGenerationHash, GrassTypes);
+					if (State->GrassInstanceGenerationHash != CurGrassInstanceGenerationHash)
+					{
+						ComponentsToRemoveFoliageInstances.Add(Component);
+						Component->InvalidateGrassTypeSummary();
+						State->GrassInstanceGenerationHash = CurGrassInstanceGenerationHash;
+					}
+#endif // WITH_EDITOR
 				}
+				continue;	// next!
+
+			default:
+				check(false);	// unreachable
+				break;			// cancel and evict
 			}
 		}
 
@@ -625,7 +642,8 @@ void FLandscapeGrassMapsBuilder::AmortizedUpdateGrassMaps(
 		AmortizedMaxRendering *= GGrassMapPrioritizedMultiplier;
 	}
 
-	UpdateTrackedComponents(Cameras, AmortizedMaxRendering);
+	bool bCancelAndEvictAll = !GGrassEnable;
+	UpdateTrackedComponents(Cameras, AmortizedMaxRendering, GGrassMapMaxDiscardChecksPerFrame, bCancelAndEvictAll);
 
 	// no point in looking to start new grass map generation if nothing is pending, if grass is disabled or there are no cameras
 	if (bAllowStartGrassMapGeneration && PendingCount > 0 && GGrassEnable && Cameras.Num() > 0)
@@ -679,7 +697,8 @@ bool FLandscapeGrassMapsBuilder::BuildGrassMapsNowForComponents(
 		}
 
 		// update all components that are tracked (without evicting)
-		bool bChanged = UpdateTrackedComponents(EmptyCamerasArray, MaxStreamingRendering);
+		int32 UpdateAllComponentCount = ComponentStates.Num();
+		bool bChanged = UpdateTrackedComponents(EmptyCamerasArray, MaxStreamingRendering, UpdateAllComponentCount, /* bCancelAndEvictAll= */ false);
 
 		UpToDateCount = 0;
 		int32 AvailableStreamingSlots = MaxStreamingRendering - StreamingCount; // here we don't limit by overall population count
@@ -705,9 +724,9 @@ bool FLandscapeGrassMapsBuilder::BuildGrassMapsNowForComponents(
 			}
 			if (State->Stage == EComponentStage::GrassMapsPopulated)
 			{
-				check(Component->GrassData->HasValidData()); // guaranteed by UpdateTrackedComponents()
+				check(Component->GrassData->HasValidData()); // guaranteed by UpdateTrackedComponents(), as long as we update all of the components
 #if WITH_EDITOR
-				// guaranteed by UpdateTrackedComponents() (unless GrassData GenerationHash gets out of sync with the tracked state somehow)
+				// guaranteed by UpdateTrackedComponents(), as long as we update all of the components
 				check(Component->ComputeGrassMapGenerationHash() == Component->GrassData->GenerationHash);
 #endif // WITH_EDITOR
 				UpToDateCount++;
@@ -725,14 +744,24 @@ bool FLandscapeGrassMapsBuilder::BuildGrassMapsNowForComponents(
 			LastChangeTime = FPlatformTime::Seconds();
 		}
 
-		// Queue up the gpu commands on the render thread, so the GPU can start working on them.
-		UE::Landscape::SubmitGPUCommands(/* bBlockUntilComplete =  */ false);
+		// If any rendering is in flight, queue up the gpu commands on the render thread, so the GPU can start working on them.
+		if (RenderingCount > 0)
+		{
+			// TODO [chris.tchou] it currently seems to be faster to block here; otherwise it takes a long time to complete the readback
+			// not sure why this is, something must be getting starved in the non-blocking path.
+			UE::Landscape::SubmitGPUCommands(/* bBlockUntilComplete =  */ true);
+		}
 
-		// Blocking texture streaming update.  At least the GPU should be working while we wait on streaming here.
+		// If any streaming is in flight, do a blocking texture streaming update.
 		// TODO [chris.tchou] : ideally this would be a non-blocking streaming update tick, so we can react to other updates finishing
-		TextureStreamingManager.WaitForTextureStreaming();
+		if (StreamingCount > 0)
+		{
+			TextureStreamingManager.WaitForTextureStreaming();
+		}
 	}
 
+	UE_LOG(LogGrass, Verbose, TEXT("BuildGrassMapsNowForComponents() updated %d/%d components in %f seconds"), UpToDateCount, LandscapeComponents.Num(), FPlatformTime::Seconds() - StartTime);
+
 	return (UpToDateCount == LandscapeComponents.Num());
 }
 
diff --git a/Engine/Source/Runtime/Landscape/Private/LandscapeGrassMapsBuilder.h b/Engine/Source/Runtime/Landscape/Private/LandscapeGrassMapsBuilder.h
--- a/Engine/Source/Runtime/Landscape/Private/LandscapeGrassMapsBuilder.h
+++ b/Engine/Source/Runtime/Landscape/Private/LandscapeGrassMapsBuilder.h
@@ -59,7 +59,7 @@ private:
 	// Update all of the non-pending components.
 	// If passed an empty Cameras array, distances are calculated as zero (i.e. it won't evict for distance)
 	// returns true if any components changed states
-	bool UpdateTrackedComponents(const TArray<FVector>& Cameras, int32 LocalMaxRendering);
+	bool UpdateTrackedComponents(const TArray<FVector>& Cameras, int32 LocalMaxRendering, int32 MaxExpensiveUpdateChecksToPerform, bool bCancelAndEvictAll);
 
 	// Start the grass map generation process on pending components in priority order
 	// (based on distance from the given Camera set) -- Cameras must not be empty.
@@ -184,11 +184,11 @@ private:
 	// number of components that need to render but are waiting (as of the last call to StartTrackingComponents())
 	int32 TotalComponentsWaitingCount = 0;
 
+	TAllocatorFixedSizeFreeList<sizeof(FComponentState), 32> StatePoolAllocator;
+
 	// store the grass map state of each registered (or recently unregistered) component
 	TMap<ULandscapeComponent*, FComponentState*> ComponentStates;
 
-	static TAllocatorFixedSizeFreeList<sizeof(FComponentState), 32> StatePoolAllocator;
-
 	TArray<FPendingComponent> PendingComponentsHeap;
 	int32 PendingUpdateAmortizationCounter = 0;
 
