Minor mem saving in grass component's exclude boxes

diff --git a/Engine/Source/Runtime/Landscape/Private/LandscapeGrass.cpp b/Engine/Source/Runtime/Landscape/Private/LandscapeGrass.cpp
--- a/Engine/Source/Runtime/Landscape/Private/LandscapeGrass.cpp
+++ b/Engine/Source/Runtime/Landscape/Private/LandscapeGrass.cpp
@@ -1889,6 +1889,7 @@ struct FAsyncGrassBuilder : public FGrassBuilderBase
 
 		if (InExcludedBoxes.Num())
 		{
+			ExcludedBoxes.Reserve(InExcludedBoxes.Num());
 			FMatrix BoxXForm = GrassInstancedStaticMeshComponent->GetComponentToWorld().ToMatrixWithScale().Inverse() * XForm.Inverse();
 			for (const FBox& Box : InExcludedBoxes)
 			{
@@ -2712,6 +2713,7 @@ void ALandscapeProxy::UpdateGrass(const TArray<FVector>& Cameras, int32& InOutNu
 														NewComp.ExcludedBoxes.Add(BoxWrapper.Box);
 													}
 												}
+												NewComp.ExcludedBoxes.Shrink();
 												if (NewComp.ExcludedBoxes != Existing->ExcludedBoxes)
 												{
 													bRebuildForBoxes = true;
@@ -2750,6 +2752,7 @@ void ALandscapeProxy::UpdateGrass(const TArray<FVector>& Cameras, int32& InOutNu
 													NewComp.ExcludedBoxes.Add(BoxWrapper.Box);
 												}
 											}
+											NewComp.ExcludedBoxes.Shrink();
 										}
 										NewComp.ExclusionChangeTag = GGrassExclusionChangeTag;
 
