Move cook time grass data stripping to serialization, and move cook time collision stripping to the NeedsLoadForClient/Server() callbacks.

This fixes deadlock issues when BeginCacheForCookedPlatform() is called multiple times to cache multiple platforms simultaneously, and also fixes an issue with collision stripping not being reversible.

diff --git a/Engine/Source/Runtime/Landscape/Classes/LandscapeComponent.h b/Engine/Source/Runtime/Landscape/Classes/LandscapeComponent.h
--- a/Engine/Source/Runtime/Landscape/Classes/LandscapeComponent.h
+++ b/Engine/Source/Runtime/Landscape/Classes/LandscapeComponent.h
@@ -694,20 +694,6 @@ public:
 	 * Returns false if there are no non-null physical materials. (We probably don't want to use if no physical material connections are bound.)
 	 */
 	bool GetRenderPhysicalMaterials(TArray<UPhysicalMaterial*>& OutPhysicalMaterials) const;
-
-	// Non-serialized data used to implement BeginCacheForCookedPlatformData() / ClearAllCachedCookedPlatformData()
-	struct FPlatformCook
-	{
-		// -1 indicates not currently cached for a cooked platform
-		int32 CurrentCookedPlatformOrdinal = -1;
-
-		// flags indicating if the current cooked platform stripped grass or collision
-		bool bStripGrassData = false;
-
-		// if Landscape data is stripped by a platform cook, it is stored here so it can be restored
-		TSharedPtr<FLandscapeComponentGrassData> StrippedGrassData;
-	};
-	FPlatformCook PlatformCook;
 #endif // WITH_EDITOR
 
 	//~ Begin UObject Interface.	
@@ -725,8 +711,6 @@ public:
 
 #if WITH_EDITOR
 	virtual void BeginCacheForCookedPlatformData(const ITargetPlatform* TargetPlatform) override;
-	virtual bool IsCachedCookedPlatformDataLoaded(const ITargetPlatform* TargetPlatform) override;
-	virtual void ClearAllCachedCookedPlatformData() override;
 	virtual void PreEditUndo() override;
 	virtual void PostEditUndo() override;
 	virtual void PreEditChange(FProperty* PropertyThatWillChange) override;
diff --git a/Engine/Source/Runtime/Landscape/Classes/LandscapeHeightfieldCollisionComponent.h b/Engine/Source/Runtime/Landscape/Classes/LandscapeHeightfieldCollisionComponent.h
--- a/Engine/Source/Runtime/Landscape/Classes/LandscapeHeightfieldCollisionComponent.h
+++ b/Engine/Source/Runtime/Landscape/Classes/LandscapeHeightfieldCollisionComponent.h
@@ -237,6 +237,8 @@ public:
 	virtual void PostLoad() override;
 	virtual void PreSave(FObjectPreSaveContext ObjectSaveContext) override;
 #if WITH_EDITOR
+	virtual bool NeedsLoadForClient() const override;
+	virtual bool NeedsLoadForServer() const override;
 	virtual void ExportCustomProperties(FOutputDevice& Out, uint32 Indent) override;
 	virtual void ImportCustomProperties(const TCHAR* SourceText, FFeedbackContext* Warn) override;
 	virtual void PostEditImport() override;
diff --git a/Engine/Source/Runtime/Landscape/Private/Landscape.cpp b/Engine/Source/Runtime/Landscape/Private/Landscape.cpp
--- a/Engine/Source/Runtime/Landscape/Private/Landscape.cpp
+++ b/Engine/Source/Runtime/Landscape/Private/Landscape.cpp
@@ -330,86 +330,15 @@ void ULandscapeComponent::BeginCacheForCookedPlatformData(const ITargetPlatform*
 {
 	Super::BeginCacheForCookedPlatformData(TargetPlatform);
 
-	// first check if there is a cooked state cached currently, and reset back to non cooked state if so
-	if (PlatformCook.CurrentCookedPlatformOrdinal != -1)
-	{
-		ClearAllCachedCookedPlatformData();
-	}
-
-	// record the platform we are cooking for, so we know what state we are in
-	PlatformCook.CurrentCookedPlatformOrdinal = TargetPlatform->GetPlatformOrdinal();
-	PlatformCook.bStripGrassData = false;
-	PlatformCook.StrippedGrassData = nullptr;
-
 	if (!HasAnyFlags(RF_ClassDefaultObject))
 	{
 		if (TargetPlatform->SupportsFeature(ETargetPlatformFeatures::MobileRendering))
 		{
 			CheckGenerateMobilePlatformData(/*bIsCooking = */ true, TargetPlatform);
 		}
-
-		// determine whether our target platform is going to need serialized grass data
-		bool bNeedSerializedGrassData = false;
-		{
-			TSharedPtr<IConsoleVariable> TargetPlatformUseRuntimeGeneration =
-				CVarGrassMapUseRuntimeGeneration->GetPlatformValueVariable(*TargetPlatform->IniPlatformName());
-			check(TargetPlatformUseRuntimeGeneration.IsValid());
-
-			if (!TargetPlatformUseRuntimeGeneration->GetBool())
-			{
-				bNeedSerializedGrassData = true;
-			}
-		}
-		PlatformCook.bStripGrassData = !bNeedSerializedGrassData;
-
-		if (ALandscapeProxy* Proxy = GetLandscapeProxy())
-		{
-			// Also strip grass data according to Proxy flags (when not cooking for editor)
-			if (!TargetPlatform->AllowsEditorObjects())
-			{
-				if (CVarAllowGrassStripping->GetBool() &&
-					((Proxy->bStripGrassWhenCookedClient && Proxy->bStripGrassWhenCookedServer) ||
-					(Proxy->bStripGrassWhenCookedClient && TargetPlatform->IsClientOnly()) ||
-					(Proxy->bStripGrassWhenCookedServer && TargetPlatform->IsServerOnly())))
-				{
-					PlatformCook.bStripGrassData = true;
-				}
-			}
-		}
-		
-		if (PlatformCook.bStripGrassData)
-		{
-			// save existing value before we strip it, so we can restore later
-			PlatformCook.StrippedGrassData = GrassData;
-			TUniquePtr<FLandscapeComponentGrassData> NewGrassData = MakeUnique<FLandscapeComponentGrassData>();
-			GrassData = MakeShareable(NewGrassData.Release());
-			GrassData->NumElements = 0;
-		}
-	}
-}
-
-bool ULandscapeComponent::IsCachedCookedPlatformDataLoaded(const ITargetPlatform* TargetPlatform)
-{
-	return (PlatformCook.CurrentCookedPlatformOrdinal == TargetPlatform->GetPlatformOrdinal());
-}
-
-void ULandscapeComponent::ClearAllCachedCookedPlatformData()
-{
-	// restore back to "non-cooked" state
-	if (PlatformCook.CurrentCookedPlatformOrdinal >= 0 && PlatformCook.bStripGrassData)
-	{
-		if (PlatformCook.StrippedGrassData.IsValid())
-		{
-			GrassData = PlatformCook.StrippedGrassData.ToSharedRef();
-		}
 	}
-
-	PlatformCook.CurrentCookedPlatformOrdinal = -1;
-	PlatformCook.bStripGrassData = false;
-	PlatformCook.StrippedGrassData = nullptr;
 }
 
-
 void ALandscapeProxy::CheckGenerateMobilePlatformData(bool bIsCooking, const ITargetPlatform* TargetPlatform)
 {
 	for (ULandscapeComponent* Component : LandscapeComponents)
@@ -752,15 +681,39 @@ void ULandscapeComponent::Serialize(FArchive& Ar)
 	Ar.UsingCustomVersion(FFortniteMainBranchObjectVersion::GUID);
 	Ar.UsingCustomVersion(FEditorObjectVersion::GUID);
 
+	bool bStripGrassData = false;
 #if WITH_EDITOR
 	if (Ar.IsCooking() && !HasAnyFlags(RF_ClassDefaultObject))
 	{
+		const ITargetPlatform* TargetPlatform = Ar.CookingTarget();
+
 		// for -oldcook:
 		// the old cooker calls BeginCacheForCookedPlatformData after the package export set is tagged, so the mobile material doesn't get saved, so we have to do CheckGenerateMobilePlatformData in serialize
 		// the new cooker clears the texture source data before calling serialize, causing GeneratePlatformVertexData to crash, so we have to do CheckGenerateMobilePlatformData in BeginCacheForCookedPlatformData
-		if (Ar.CookingTarget()->SupportsFeature(ETargetPlatformFeatures::MobileRendering))
+		if (TargetPlatform->SupportsFeature(ETargetPlatformFeatures::MobileRendering))
+		{
+			CheckGenerateMobilePlatformData(/*bIsCooking = */ true, TargetPlatform);
+		}
+
+		// determine whether our target platform is going to need serialized grass data
+		TSharedPtr<IConsoleVariable> TargetPlatformUseRuntimeGeneration =
+			CVarGrassMapUseRuntimeGeneration->GetPlatformValueVariable(*TargetPlatform->IniPlatformName());
+		check(TargetPlatformUseRuntimeGeneration.IsValid());
+		bStripGrassData = TargetPlatformUseRuntimeGeneration->GetBool();
+
+		if (ALandscapeProxy* Proxy = GetLandscapeProxy())
 		{
-			CheckGenerateMobilePlatformData(/*bIsCooking = */ true, Ar.CookingTarget());
+			// Also strip grass data according to Proxy flags (when not cooking for editor)
+			if (!TargetPlatform->AllowsEditorObjects())
+			{
+				if (CVarAllowGrassStripping->GetBool() &&
+					((Proxy->bStripGrassWhenCookedClient && Proxy->bStripGrassWhenCookedServer) ||
+					 (Proxy->bStripGrassWhenCookedClient && TargetPlatform->IsClientOnly()) ||
+					 (Proxy->bStripGrassWhenCookedServer && TargetPlatform->IsServerOnly())))
+				{
+					bStripGrassData = true;
+				}
+			}
 		}
 	}
 
@@ -928,7 +881,16 @@ void ULandscapeComponent::Serialize(FArchive& Ar)
 		}
 		else
 		{
-			Ar << GrassData.Get();
+			if (bStripGrassData)
+			{
+				FLandscapeComponentGrassData EmptyGrassData;
+				EmptyGrassData.NumElements = 0;
+				Ar << EmptyGrassData;
+			}
+			else
+			{
+				Ar << GrassData.Get();
+			}
 		}
 
 		// When loading or saving a component, validate that grass data is valid : 
diff --git a/Engine/Source/Runtime/Landscape/Private/LandscapeCollision.cpp b/Engine/Source/Runtime/Landscape/Private/LandscapeCollision.cpp
--- a/Engine/Source/Runtime/Landscape/Private/LandscapeCollision.cpp
+++ b/Engine/Source/Runtime/Landscape/Private/LandscapeCollision.cpp
@@ -2529,6 +2529,26 @@ void ULandscapeHeightfieldCollisionComponent::PreSave(FObjectPreSaveContext Obje
 }
 
 #if WITH_EDITOR
+bool ULandscapeHeightfieldCollisionComponent::NeedsLoadForClient() const
+{
+	ALandscapeProxy* Proxy = GetLandscapeProxy();
+	if (ensure(Proxy))
+	{
+		return !Proxy->bStripPhysicsWhenCookedClient;
+	}
+	return true;
+}
+
+bool ULandscapeHeightfieldCollisionComponent::NeedsLoadForServer() const
+{
+	ALandscapeProxy* Proxy = GetLandscapeProxy();
+	if (ensure(Proxy))
+	{
+		return !Proxy->bStripPhysicsWhenCookedServer;
+	}
+	return true;
+}
+
 void ULandscapeInfo::UpdateAllAddCollisions()
 {
 	TRACE_CPUPROFILER_EVENT_SCOPE(ULandscapeInfo::UpdateAllAddCollisions);
