Crash Prevention: Added null checks to various crashes found when running the editor.

These crashes were surfaced by CL 29712224 which exposed that these values were not being null checked.

#rnx

diff --git a/Engine/Source/Editor/BlueprintGraph/Private/K2Node_DynamicCast.cpp b/Engine/Source/Editor/BlueprintGraph/Private/K2Node_DynamicCast.cpp
--- a/Engine/Source/Editor/BlueprintGraph/Private/K2Node_DynamicCast.cpp
+++ b/Engine/Source/Editor/BlueprintGraph/Private/K2Node_DynamicCast.cpp
@@ -491,6 +491,11 @@ void UK2Node_DynamicCast::ValidateNodeDuringCompilation(FCompilerResultsLog& Mes
 
 		for (UEdGraphPin* CastInput : SourcePin->LinkedTo)
 		{
+			if (CastInput == nullptr)
+			{
+				continue;
+			}
+
 			const FEdGraphPinType& SourcePinType = CastInput->PinType;
 			if (SourcePinType.PinCategory != UEdGraphSchema_K2::PC_Object)
 			{
diff --git a/Engine/Source/Runtime/Engine/Private/EdGraph/EdGraphNode.cpp b/Engine/Source/Runtime/Engine/Private/EdGraph/EdGraphNode.cpp
--- a/Engine/Source/Runtime/Engine/Private/EdGraph/EdGraphNode.cpp
+++ b/Engine/Source/Runtime/Engine/Private/EdGraph/EdGraphNode.cpp
@@ -438,17 +438,21 @@ void UEdGraphNode::BreakAllNodeLinks()
 	NodeList.Add(this);
 
 	// Iterate over each pin and break all links
-	for(int32 PinIdx = 0; PinIdx < Pins.Num(); ++PinIdx)
+	for (int32 PinIdx = 0; PinIdx < Pins.Num(); ++PinIdx)
 	{
-		UEdGraphPin* Pin = Pins[PinIdx];
-
-		// Save all the connected nodes to be notified below
-		for (UEdGraphPin* Connection : Pin->LinkedTo)
+		if (UEdGraphPin* Pin = Pins[PinIdx])
 		{
-			NodeList.Add(Connection->GetOwningNode());
-		}
+			// Save all the connected nodes to be notified below
+			for (UEdGraphPin* Connection : Pin->LinkedTo)
+			{
+				if (Connection != nullptr)
+				{
+					NodeList.Add(Connection->GetOwningNode());
+				}
+			}
 
-		Pin->BreakAllPinLinks();
+			Pin->BreakAllPinLinks();
+		}
 	}
 
 	// Send a notification to all nodes that lost a connection
diff --git a/Engine/Source/Runtime/Landscape/Private/LandscapeEditLayers.cpp b/Engine/Source/Runtime/Landscape/Private/LandscapeEditLayers.cpp
--- a/Engine/Source/Runtime/Landscape/Private/LandscapeEditLayers.cpp
+++ b/Engine/Source/Runtime/Landscape/Private/LandscapeEditLayers.cpp
@@ -7496,13 +7496,21 @@ int32 ALandscape::PerformLayersWeightmapsGlobalMerge(FUpdateLayersContentContext
 							continue;
 						}
 
+						FTextureResource* WeightmapTextureResource = WeightmapTexture ? WeightmapTexture->GetResource() : nullptr;
+						if (WeightmapTextureResource == nullptr)
+						{
+							continue;
+						}
+
+						const int32 TextureSizeX = WeightmapTextureResource->GetSizeX();
+						const int32 TextureSizeY = WeightmapTextureResource->GetSizeY();
+
 						FIntPoint TextureTopLeftPositionInAtlas(
 							static_cast<int32>(WeightmapTextureOutputOffset[NextTextureIndexToProcess - StartTextureIndex].X),
 							static_cast<int32>(WeightmapTextureOutputOffset[NextTextureIndexToProcess - StartTextureIndex].Y));
 
 						int32 CurrentMip = 0;
-						const int32 TextureSizeX = WeightmapTexture->GetResource()->GetSizeX();
-						const int32 TextureSizeY = WeightmapTexture->GetResource()->GetSizeY();
+
 						for (int32 MipRTIndex = (int32)EWeightmapRTType::WeightmapRT_Mip0; MipRTIndex < (int32)EWeightmapRTType::WeightmapRT_Count; ++MipRTIndex)
 						{
 							CurrentRT = WeightmapRTList[MipRTIndex];
