Fixed grass not being updated by exclusion boxes during heavy building scenario in the local area

diff --git a/Engine/Source/Runtime/Landscape/Private/LandscapeGrass.cpp b/Engine/Source/Runtime/Landscape/Private/LandscapeGrass.cpp
--- a/Engine/Source/Runtime/Landscape/Private/LandscapeGrass.cpp
+++ b/Engine/Source/Runtime/Landscape/Private/LandscapeGrass.cpp
@@ -2984,7 +2984,8 @@ void ALandscapeProxy::ProcessAsyncGrassInstanceTasks(bool bWaitAsyncTasks, bool
 		{
 			SCOPE_CYCLE_COUNTER(STAT_FoliageGrassEndComp);
 			FAsyncGrassTask& Inner = Task->GetTask();
-			AsyncFoliageTasks.RemoveAtSwap(Index--);
+			// We need to preserve the order here, otherwise we'll have new jobs that are added to the end jumping up in front of the queue and an original second job would be updated the last
+			AsyncFoliageTasks.RemoveAt(Index--);
 			UGrassInstancedStaticMeshComponent* GrassISMComponent = Cast<UGrassInstancedStaticMeshComponent>(Inner.Foliage.Get());
 			int32 NumBuiltRenderInstances = Inner.Builder->InstanceBuffer.GetNumInstances();
 			//UE_LOG(LogCore, Display, TEXT("%d instances in %4.0fms     %6.0f instances / sec"), NumBuiltRenderInstances, 1000.0f * float(Inner.Builder->BuildTime), float(NumBuiltRenderInstances) / float(Inner.Builder->BuildTime));
