Fixed bad english

diff --git a/Engine/Source/Editor/LandscapeEditor/Private/LandscapeEditorDetailCustomization_Layers.cpp b/Engine/Source/Editor/LandscapeEditor/Private/LandscapeEditorDetailCustomization_Layers.cpp
--- a/Engine/Source/Editor/LandscapeEditor/Private/LandscapeEditorDetailCustomization_Layers.cpp
+++ b/Engine/Source/Editor/LandscapeEditor/Private/LandscapeEditorDetailCustomization_Layers.cpp
@@ -566,7 +566,7 @@ void FLandscapeEditorCustomNodeBuilder_Layers::SetLandscapeSplinesReservedLayer(
 			EAppReturnType::Type Result = EAppReturnType::No;
 			if (NewReservedLayer)
 			{
-				Result = FMessageDialog::Open(EAppMsgType::YesNo, FText::Format(LOCTEXT("Landscape_SetReservedForSplines_Message", "Reserving layer {0} for landscape splines will clear it from its content and no edition will be allowed.  Continue?"), FText::FromName(NewReservedLayer->Name)));
+				Result = FMessageDialog::Open(EAppMsgType::YesNo, FText::Format(LOCTEXT("Landscape_SetReservedForSplines_Message", "Reserving layer {0} for landscape splines will clear it from its content and no editing will be allowed.  Continue?"), FText::FromName(NewReservedLayer->Name)));
 			}
 			else if (CurrenReservedLayer)
 			{
diff --git a/Engine/Source/Runtime/Landscape/Classes/LandscapeProxy.h b/Engine/Source/Runtime/Landscape/Classes/LandscapeProxy.h
--- a/Engine/Source/Runtime/Landscape/Classes/LandscapeProxy.h
+++ b/Engine/Source/Runtime/Landscape/Classes/LandscapeProxy.h
@@ -871,7 +871,7 @@ public:
 	 * @param bRaiseHeights - Allow the landscape to be raised up to the level of the spline. If both bRaiseHeights and bLowerHeights are false, no height modification of the landscape will be performed
 	 * @param bLowerHeights - Allow the landscape to be lowered down to the level of the spline. If both bRaiseHeights and bLowerHeights are false, no height modification of the landscape will be performed
 	 * @param PaintLayer - LayerInfo to paint, or none to skip painting. The landscape must be configured with the same layer info in one of its layers or this will do nothing!
-	 * @param EditLayerName - Name of the landscape edition layer to affect (in Edit Layers mode)
+	 * @param EditLayerName - Name of the landscape edit layer to affect (in Edit Layers mode)
 	 */
 	UFUNCTION(BlueprintCallable, Category = "Landscape|Editor")
 	LANDSCAPE_API void EditorApplySpline(USplineComponent* InSplineComponent, float StartWidth = 200, float EndWidth = 200, float StartSideFalloff = 200, float EndSideFalloff = 200, float StartRoll = 0, float EndRoll = 0, int32 NumSubdivisions = 20, bool bRaiseHeights = true, bool bLowerHeights = true, ULandscapeLayerInfoObject* PaintLayer = nullptr, FName EditLayerName = TEXT(""));
