Fixed ensure when painting/sculpting outside a landscape
#tests editor

diff --git a/Engine/Source/Runtime/Landscape/Private/LandscapeEditInterface.cpp b/Engine/Source/Runtime/Landscape/Private/LandscapeEditInterface.cpp
--- a/Engine/Source/Runtime/Landscape/Private/LandscapeEditInterface.cpp
+++ b/Engine/Source/Runtime/Landscape/Private/LandscapeEditInterface.cpp
@@ -4283,10 +4283,10 @@ void FLandscapeEditDataInterface::GetWeightDataTempl(ULandscapeLayerInfoObject*
 			ComponentSizeX, ComponentSizeY, CornerValues,
 			NoBorderY1, NoBorderY2, ComponentDataExist, StoreData );
 		// Update valid region
-		ValidX1 = FMath::Max<int32>(X1, ValidX1);
-		ValidX2 = FMath::Min<int32>(X2, ValidX2);
-		ValidY1 = FMath::Max<int32>(Y1, ValidY1);
-		ValidY2 = FMath::Min<int32>(Y2, ValidY2);
+		ValidX1 = FMath::Min<int32>(X1, ValidX1);
+		ValidX2 = FMath::Max<int32>(X2, ValidX2);
+		ValidY1 = FMath::Min<int32>(Y1, ValidY1);
+		ValidY2 = FMath::Max<int32>(Y2, ValidY2);
 	}
 	else
 	{
