Check landscape override key exists before using it
  - Fix EngineTest Ocean_Colorless, River_Deep_EyeLevel, SkyOverview tests

diff --git a/Engine/Source/Runtime/Landscape/Private/LandscapeRender.cpp b/Engine/Source/Runtime/Landscape/Private/LandscapeRender.cpp
--- a/Engine/Source/Runtime/Landscape/Private/LandscapeRender.cpp
+++ b/Engine/Source/Runtime/Landscape/Private/LandscapeRender.cpp
@@ -399,8 +399,8 @@ static int32 GetViewLodOverride(FSceneView const& View, uint32 LandscapeKey)
 		TOptional<TMap<uint32, int32>>& LandscapeLODOverrides = *(TOptional<TMap<uint32, int32>>*)(View.CustomRenderPass->UserData);
 		if (LandscapeLODOverrides)
 		{
-			int32 LandscapeLODOverride = LandscapeLODOverrides->FindChecked(LandscapeKey);
-			LodOverride = LandscapeLODOverride > -1 ? LandscapeLODOverride : LodOverride;
+			int32* LandscapeLODOverride = LandscapeLODOverrides->Find(LandscapeKey);
+			LodOverride = LandscapeLODOverride ? *LandscapeLODOverride : LodOverride;
 		}
 	}
 
