Fixed incorrect landscape CVar names
#tests editor

diff --git a/Engine/Source/Editor/DetailCustomizations/Private/LandscapeProxyUIDetails.cpp b/Engine/Source/Editor/DetailCustomizations/Private/LandscapeProxyUIDetails.cpp
--- a/Engine/Source/Editor/DetailCustomizations/Private/LandscapeProxyUIDetails.cpp
+++ b/Engine/Source/Editor/DetailCustomizations/Private/LandscapeProxyUIDetails.cpp
@@ -54,8 +54,8 @@ class URuntimeVirtualTexture;
 
 #define LOCTEXT_NAMESPACE "FLandscapeProxyUIDetails"
 
-static TAutoConsoleVariable<int32> CVarShowCompressWeightMapsOption(
-	TEXT("landscape.ShowCompressHeightMapsOption"),
+static TAutoConsoleVariable<int32> CVarShowCompressHeightmapsOption(
+	TEXT("landscape.ShowCompressHeightmapsOption"),
 	0,
 	TEXT("Enable editing of the compressed height map option on landscape proxies (experimental)."),
 	ECVF_Default);
@@ -136,7 +136,7 @@ void FLandscapeProxyUIDetails::CustomizeDetails( IDetailLayoutBuilder& DetailBui
 		DetailBuilder.HideProperty(DetailBuilder.GetProperty(GET_MEMBER_NAME_CHECKED(ALandscape, bAreNewLandscapeActorsSpatiallyLoaded), ALandscape::StaticClass()));
 	}
 
-	if (CVarShowCompressWeightMapsOption->GetInt() == 0)
+	if (CVarShowCompressHeightmapsOption->GetInt() == 0)
 	{
 		DetailBuilder.HideProperty(DetailBuilder.GetProperty(GET_MEMBER_NAME_CHECKED(ALandscapeProxy, bUseCompressedHeightmapStorage), ALandscapeProxy::StaticClass()));
 	}
diff --git a/Engine/Source/Runtime/Landscape/Private/LandscapeRender.cpp b/Engine/Source/Runtime/Landscape/Private/LandscapeRender.cpp
--- a/Engine/Source/Runtime/Landscape/Private/LandscapeRender.cpp
+++ b/Engine/Source/Runtime/Landscape/Private/LandscapeRender.cpp
@@ -75,7 +75,7 @@ static void OnLODDistributionScaleChanged(IConsoleVariable* CVar)
 #if !UE_BUILD_SHIPPING
 float GLandscapeLOD0ScreenSizeOverride = -1.f;
 FAutoConsoleVariableRef CVarLandscapeLOD0ScreenSizeOverride(
-	TEXT("r.Landscape.Override.LOD0ScreenSize"),
+	TEXT("landscape.OverrideLOD0ScreenSize"),
 	GLandscapeLOD0ScreenSizeOverride,
 	TEXT("When > 0, force override the landscape LOD0ScreenSize property on all landscapes"),
 	FConsoleVariableDelegate::CreateStatic(&OnLODDistributionScaleChanged),
@@ -84,7 +84,7 @@ FAutoConsoleVariableRef CVarLandscapeLOD0ScreenSizeOverride(
 
 float GLandscapeLOD0DistributionOverride = -1.f;
 FAutoConsoleVariableRef CVarLandscapeLOD0DistributionOverride(
-	TEXT("r.Landscape.Override.LOD0Distribution"),
+	TEXT("landscape.OverrideLOD0Distribution"),
 	GLandscapeLOD0DistributionOverride,
 	TEXT("When > 0, force override the LOD0DistributionSetting property on all landscapes, and ignore r.LandscapeLOD0DistributionScale"),
 	FConsoleVariableDelegate::CreateStatic(&OnLODDistributionScaleChanged),
@@ -93,7 +93,7 @@ FAutoConsoleVariableRef CVarLandscapeLOD0DistributionOverride(
 
 float GLandscapeLODDistributionOverride = -1.f;
 FAutoConsoleVariableRef CVarLandscapeLODDistributionOverride(
-	TEXT("r.Landscape.Override.LODDistribution"),
+	TEXT("landscape.OverrideLODDistribution"),
 	GLandscapeLODDistributionOverride,
 	TEXT("When > 0, force override the landscape LODDistributionSetting property on all landscapes, and ignore r.LandscapeLODDistributionScale"),
 	FConsoleVariableDelegate::CreateStatic(&OnLODDistributionScaleChanged),
@@ -102,7 +102,7 @@ FAutoConsoleVariableRef CVarLandscapeLODDistributionOverride(
 
 float GLandscapeLODBlendRangeOverride = -1.f;
 FAutoConsoleVariableRef CVarLandscapeLODBlendRangeOverride(
-	TEXT("r.Landscape.Override.LODBlendRange"),
+	TEXT("landscape.OverrideLODBlendRange"),
 	GLandscapeLODBlendRangeOverride,
 	TEXT("When > 0, force the LODBlendRange property on all landscapes"),
 	FConsoleVariableDelegate::CreateStatic(&OnLODDistributionScaleChanged),
