Deprecated prototype-y function to change LOD settings on landscape directly on the render proxy : it's already doable in a normal fashion by invalidating the render proxy
#tests editor

diff --git a/Engine/Source/Runtime/Landscape/Classes/LandscapeProxy.h b/Engine/Source/Runtime/Landscape/Classes/LandscapeProxy.h
--- a/Engine/Source/Runtime/Landscape/Classes/LandscapeProxy.h
+++ b/Engine/Source/Runtime/Landscape/Classes/LandscapeProxy.h
@@ -850,7 +850,7 @@ public:
 	virtual void ChangeLODDistanceFactor(float InLODDistanceFactor);
 
 	/** Change ComponentScreenSizeToUseSubSections value on the render proxy.*/
-	UFUNCTION(BlueprintCallable, Category = "Rendering")
+	UFUNCTION(BlueprintCallable, Category = "Rendering", meta = (DeprecatedFunction, DeprecationMessage = "This value can't be changed anymore, you should edit the property ComponentScreenSizeToUseSubSections of the Landscape"))
 	virtual void ChangeComponentScreenSizeToUseSubSections(float InComponentScreenSizeToUseSubSections);
 
 	/* Setter for LandscapeMaterial. Has no effect outside the editor. */
diff --git a/Engine/Source/Runtime/Landscape/Private/LandscapeRender.cpp b/Engine/Source/Runtime/Landscape/Private/LandscapeRender.cpp
--- a/Engine/Source/Runtime/Landscape/Private/LandscapeRender.cpp
+++ b/Engine/Source/Runtime/Landscape/Private/LandscapeRender.cpp
@@ -4112,33 +4112,7 @@ void ULandscapeComponent::GetStreamingRenderAssetInfo(FStreamingTextureLevelCont
 
 void ALandscapeProxy::ChangeComponentScreenSizeToUseSubSections(float InComponentScreenSizeToUseSubSections)
 {
-	ComponentScreenSizeToUseSubSections = FMath::Clamp<float>(InComponentScreenSizeToUseSubSections, 0.01f, 1.0f);
-
-	if (LandscapeComponents.Num() > 0)
-	{
-		int32 ComponentCount = LandscapeComponents.Num();
-		FLandscapeComponentSceneProxy** RenderProxies = new FLandscapeComponentSceneProxy * [ComponentCount];
-		for (int32 Idx = 0; Idx < ComponentCount; ++Idx)
-		{
-			RenderProxies[Idx] = (FLandscapeComponentSceneProxy*)(LandscapeComponents[Idx]->SceneProxy);
-		}
-
-		float ComponentScreenSizeToUseSubSectionsLocal = ComponentScreenSizeToUseSubSections;
-		ENQUEUE_RENDER_COMMAND(LandscapeChangeComponentScreenSizeToUseSubSectionsCommand)(
-			[RenderProxies, ComponentCount, ComponentScreenSizeToUseSubSectionsLocal](FRHICommandListImmediate& RHICmdList)
-		{
-			for (int32 Idx = 0; Idx < ComponentCount; ++Idx)
-			{
-				if (RenderProxies[Idx] != nullptr)
-				{
-					RenderProxies[Idx]->ChangeComponentScreenSizeToUseSubSections_RenderThread(ComponentScreenSizeToUseSubSectionsLocal);
-				}
-			}
-
-			delete[] RenderProxies;
-		}
-		);
-	}
+	// Deprecated
 }
 
 void ALandscapeProxy::ChangeLODDistanceFactor(float InLODDistanceFactor)
@@ -4148,6 +4122,7 @@ void ALandscapeProxy::ChangeLODDistanceFactor(float InLODDistanceFactor)
 
 void FLandscapeComponentSceneProxy::ChangeComponentScreenSizeToUseSubSections_RenderThread(float InComponentScreenSizeToUseSubSections)
 {
+	// Deprecated
 	ComponentSquaredScreenSizeToUseSubSections = FMath::Square(InComponentScreenSizeToUseSubSections);
 }
 
diff --git a/Engine/Source/Runtime/Landscape/Public/LandscapeRender.h b/Engine/Source/Runtime/Landscape/Public/LandscapeRender.h
--- a/Engine/Source/Runtime/Landscape/Public/LandscapeRender.h
+++ b/Engine/Source/Runtime/Landscape/Public/LandscapeRender.h
@@ -890,6 +890,7 @@ public:
 	const FMeshBatch& GetGrassMeshBatch() const { return GrassMeshBatch; }
 
 	// FLandcapeSceneProxy
+	UE_DEPRECATED(5.4, "This function has been deprecated, use ComponentScreenSizeToUseSubSections on ALandscapeProxy instead")
 	LANDSCAPE_API void ChangeComponentScreenSizeToUseSubSections_RenderThread(float InComponentScreenSizeToUseSubSections);
 
 	LANDSCAPE_API virtual bool HeightfieldHasPendingStreaming() const override;
