Fix for assert when saving some landscapes with no grass.

Main issue was that RemoveGrassMap() would take a shortcut if there were no grass types, failing to set a proper generation hash. But I also noticed there was an issue with the amortization that could result in a partial update, which could also lead to the assert.

diff --git a/Engine/Source/Runtime/Landscape/Private/LandscapeGrass.cpp b/Engine/Source/Runtime/Landscape/Private/LandscapeGrass.cpp
--- a/Engine/Source/Runtime/Landscape/Private/LandscapeGrass.cpp
+++ b/Engine/Source/Runtime/Landscape/Private/LandscapeGrass.cpp
@@ -990,13 +990,6 @@ void ULandscapeComponent::RemoveGrassMap()
 #if WITH_EDITOR
 	GrassData->bIsDirty = true;
 #endif // WITH_EDITOR
-
-	if (GrassTypes.IsEmpty())
-	{
-		// Mark grass data as valid but empty if it just doesn't support grass because nothing will ever trigger a RenderGrassMaps on it, which would leave NumElements unset (which is considered as invalid data) :
-		// TODO [chris.tchou]: this is not technically true with bake position offset into collision...
-		GrassData->NumElements = 0;
-	}
 }
 
 void ALandscapeProxy::UpdateGrassTypeSummary()
diff --git a/Engine/Source/Runtime/Landscape/Private/LandscapeGrassMapsBuilder.h b/Engine/Source/Runtime/Landscape/Private/LandscapeGrassMapsBuilder.h
--- a/Engine/Source/Runtime/Landscape/Private/LandscapeGrassMapsBuilder.h
+++ b/Engine/Source/Runtime/Landscape/Private/LandscapeGrassMapsBuilder.h
@@ -138,13 +138,22 @@ private:
 		void StartUpdateTick(int32 UpdateElementCount, int32 MaxUpdatesThisFrame)
 		{
 			check(MaxUpdatesThisFrame >= 0);
-			FirstIndex = LastIndex;
-			if (FirstIndex >= UpdateElementCount)
+			if (MaxUpdatesThisFrame >= UpdateElementCount)
 			{
+				// if we are updating more than what we have, then just update everything
 				FirstIndex = 0;
+				LastIndex = UpdateElementCount;
+			}
+			else
+			{
+				FirstIndex = LastIndex;
+				if (FirstIndex >= UpdateElementCount)
+				{
+					FirstIndex = 0;
+				}
+				check(FirstIndex >= 0);
+				LastIndex = FirstIndex + MaxUpdatesThisFrame;
 			}
-			check(FirstIndex >= 0);
-			LastIndex = FirstIndex + MaxUpdatesThisFrame;
 		}
 
 		// return true if the given element should update
