Landscape: Emit warning during cook if Landscape needs resaving (to rebuild nanite)

diff --git a/Engine/Source/Runtime/Landscape/Private/Landscape.cpp b/Engine/Source/Runtime/Landscape/Private/Landscape.cpp
--- a/Engine/Source/Runtime/Landscape/Private/Landscape.cpp
+++ b/Engine/Source/Runtime/Landscape/Private/Landscape.cpp
@@ -382,8 +382,11 @@ FGraphEventRef ALandscapeProxy::UpdateNaniteRepresentationAsync(const ITargetPla
 
 		const FGuid ComponentNaniteContentId = GetNaniteComponentContentId();
 		const bool bNaniteContentDirty = ComponentNaniteContentId != NaniteContentId;
-
-		UE_LOG(LogLandscape, VeryVerbose, TEXT("UpdateNaniteRepresentationAsync actor: '%s' package:'%s' dirty:%i component guid:'%s' proxy guid:'%s'"), *GetActorNameOrLabel(), *GetPackage()->GetName(), bNaniteContentDirty, *ComponentNaniteContentId.ToString(), *NaniteContentId.ToString());
+	
+		if(bNaniteContentDirty && IsRunningCookCommandlet())
+		{
+			UE_LOG(LogLandscape, Display, TEXT("Landscape Nanite out of date. Map requires resaving. Actor: '%s' Package: '%s'"), *GetActorNameOrLabel(), *GetPackage()->GetName());
+		}		
 
 		TArray<ULandscapeComponent*> StableOrderComponents(LandscapeComponents);
 		ULandscapeSubsystem* Subsystem = GetWorld()->GetSubsystem<ULandscapeSubsystem>();
@@ -455,17 +458,14 @@ void ALandscapeProxy::UpdateNaniteRepresentation(const ITargetPlatform* InTarget
 		return;
 	}
 
-	UE_LOG(LogLandscape, Display, TEXT("UpdateNaniteRepresentation proxy:%p target platform:%p subsystem:%p"), this, InTargetPlatform, Subsystem);
 	if (!Subsystem->IsMultithreadedNaniteBuildEnabled() || IsRunningCookCommandlet())
 	{
-		UE_LOG(LogLandscape, Display, TEXT("Waiting for nanite build: '%s'"), *GetActorNameOrLabel());
 		while (!GraphEvent->IsComplete())
 		{
 			ENamedThreads::Type CurrentThread = FTaskGraphInterface::Get().GetCurrentThreadIfKnown();
 			FTaskGraphInterface::Get().ProcessThreadUntilIdle(CurrentThread);
 			FAssetCompilingManager::Get().ProcessAsyncTasks();	
 		}
-		UE_LOG(LogLandscape, Display, TEXT("Complete nanite build '%s'"), *GetActorNameOrLabel());
 	}
 }
 
@@ -4237,7 +4237,7 @@ void ALandscapeProxy::PostLoad()
 	}
 
 	// Display a MapCheck warning if the Nanite data is stale with the option to trigger a rebuild & Save
-	if (!IsNaniteMeshUpToDate())
+	if (!IsNaniteMeshUpToDate() && !IsRunningCookCommandlet())
 	{
 		FFormatNamedArguments Arguments;
 		Arguments.Add(TEXT("LandscapeProxyName"), FText::FromString(GetActorNameOrLabel())); 
