Fix for Landscape Grass exporting the incorrect locations due to recent bUseFauxOrtho changes in orthographic rendering.

diff --git a/Engine/Source/Runtime/Landscape/Private/LandscapeGrassWeightExporter.cpp b/Engine/Source/Runtime/Landscape/Private/LandscapeGrassWeightExporter.cpp
--- a/Engine/Source/Runtime/Landscape/Private/LandscapeGrassWeightExporter.cpp
+++ b/Engine/Source/Runtime/Landscape/Private/LandscapeGrassWeightExporter.cpp
@@ -331,6 +331,7 @@ void FLandscapeGrassWeightExporter_RenderThread::RenderLandscapeComponentToTextu
 	ViewInitOptions.ViewRotationMatrix = ViewRotationMatrix;
 	ViewInitOptions.ProjectionMatrix = ProjectionMatrix;
 	ViewInitOptions.ViewFamily = &ViewFamily;
+	ViewInitOptions.bUseFauxOrthoViewPos = false;
 
 	GetRendererModule().CreateAndInitSingleView(RHICmdList, &ViewFamily, &ViewInitOptions);
 
