Landscape : fix crash on map load

diff --git a/Engine/Source/Runtime/Landscape/Private/LandscapeEdit.cpp b/Engine/Source/Runtime/Landscape/Private/LandscapeEdit.cpp
--- a/Engine/Source/Runtime/Landscape/Private/LandscapeEdit.cpp
+++ b/Engine/Source/Runtime/Landscape/Private/LandscapeEdit.cpp
@@ -554,7 +554,8 @@ void ULandscapeComponent::UpdateMaterialInstances()
 
 	// If this component is within a Nanite Enabled Landscape we have to update the materials on the
 	// ULandscapeNaniteComponents within the Parent Actor. 
-	GetLandscapeProxy()->UpdateNaniteMaterials(); 
+	//
+	// todo [don.boogert] reinstate update of landscape nanite materials GetLandscapeProxy()->UpdateNaniteMaterials(); 
 }
 
 void ULandscapeComponent::UpdateMaterialInstances(FMaterialUpdateContext& InOutMaterialContext, TArray<FComponentRecreateRenderStateContext>& InOutRecreateRenderStateContext)
