Fix for Landscape Proxy leak when unloading a level instance

diff --git a/Engine/Source/Runtime/Landscape/Classes/LandscapeInfo.h b/Engine/Source/Runtime/Landscape/Classes/LandscapeInfo.h
--- a/Engine/Source/Runtime/Landscape/Classes/LandscapeInfo.h
+++ b/Engine/Source/Runtime/Landscape/Classes/LandscapeInfo.h
@@ -311,7 +311,7 @@ public:
 	/** Recreate all LandscapeInfo objects in given world
 	 *  @param  bMapCheck	Whether to warn about landscape errors
 	 */
-	LANDSCAPE_API static void RecreateLandscapeInfo(UWorld* InWorld, bool bMapCheck);
+	LANDSCAPE_API static void RecreateLandscapeInfo(UWorld* InWorld, bool bMapCheck, bool bKeepRegistrationStatus = false);
 
 	/** 
 	 *  Fixes up proxies relative position to landscape actor
diff --git a/Engine/Source/Runtime/Landscape/Private/Landscape.cpp b/Engine/Source/Runtime/Landscape/Private/Landscape.cpp
--- a/Engine/Source/Runtime/Landscape/Private/Landscape.cpp
+++ b/Engine/Source/Runtime/Landscape/Private/Landscape.cpp
@@ -5286,7 +5286,7 @@ void ULandscapeInfo::UpdateComponentLayerAllowList()
 	});
 }
 
-void ULandscapeInfo::RecreateLandscapeInfo(UWorld* InWorld, bool bMapCheck)
+void ULandscapeInfo::RecreateLandscapeInfo(UWorld* InWorld, bool bMapCheck, bool bKeepRegistrationStatus)
 {
 	check(InWorld);
 
@@ -5301,6 +5301,9 @@ void ULandscapeInfo::RecreateLandscapeInfo(UWorld* InWorld, bool bMapCheck)
 		if (LandscapeInfo != nullptr)
 		{
 			LandscapeInfo->Modify();
+
+			// this effectively unregisters all proxies, but does not flag them as unregistered
+			// so we can use the flags below to determine what was previously registered
 			LandscapeInfo->Reset();
 		}
 	}
@@ -5313,6 +5316,7 @@ void ULandscapeInfo::RecreateLandscapeInfo(UWorld* InWorld, bool bMapCheck)
 			Proxy->GetLevel()->bIsVisible &&
 			!Proxy->HasAnyFlags(RF_BeginDestroyed) &&
 			IsValid(Proxy) &&
+			(!bKeepRegistrationStatus || Proxy->bIsRegisteredWithLandscapeInfo) &&
 			!Proxy->IsPendingKillPending())
 		{
 			ValidLandscapesMap.FindOrAdd(Proxy->GetLandscapeGuid()).Add(Proxy);
@@ -5325,6 +5329,7 @@ void ULandscapeInfo::RecreateLandscapeInfo(UWorld* InWorld, bool bMapCheck)
 		auto& LandscapeList = ValidLandscapesPair.Value;
 		for (ALandscapeProxy* Proxy : LandscapeList)
 		{
+			// note this may re-register already registered actors
 			Proxy->CreateLandscapeInfo()->RegisterActor(Proxy, bMapCheck);
 		}
 	}
diff --git a/Engine/Source/Runtime/Landscape/Private/LandscapeEdit.cpp b/Engine/Source/Runtime/Landscape/Private/LandscapeEdit.cpp
--- a/Engine/Source/Runtime/Landscape/Private/LandscapeEdit.cpp
+++ b/Engine/Source/Runtime/Landscape/Private/LandscapeEdit.cpp
@@ -4765,7 +4765,7 @@ void ALandscapeProxy::PostEditMove(bool bFinished)
 
 	if (bFinished)
 	{
-		ULandscapeInfo::RecreateLandscapeInfo(GetWorld(), true);
+		ULandscapeInfo::RecreateLandscapeInfo(GetWorld(), /* bMapCheck = */ true, /* bKeepRegistrationStatus = */ true);
 		RecreateComponentsState();
 
 		if (SplineComponent)
