use GetRenderTargetImage in LandscapeEdit
make GetRenderTargetImage work with rect
affects LandscapeImportWeightmapFromRenderTarget used by BP_Landmass_LayerStack

#rnx

diff --git a/Engine/Source/Runtime/Engine/Private/ImageUtils.cpp b/Engine/Source/Runtime/Engine/Private/ImageUtils.cpp
--- a/Engine/Source/Runtime/Engine/Private/ImageUtils.cpp
+++ b/Engine/Source/Runtime/Engine/Private/ImageUtils.cpp
@@ -331,6 +331,11 @@ bool FImageUtils::GetRawData(UTextureRenderTarget2D* TexRT, TArray64<uint8>& Raw
 }
 
 bool FImageUtils::GetRenderTargetImage(UTextureRenderTarget2D* TexRT, FImage & Image)
+{
+	return GetRenderTargetImage(TexRT,Image,FIntRect(0,0,0,0));
+}
+
+bool FImageUtils::GetRenderTargetImage(UTextureRenderTarget2D* TexRT, FImage & Image, const FIntRect & InRectOrZero)
 {
 	Image = FImage();
 	
@@ -341,17 +346,36 @@ bool FImageUtils::GetRenderTargetImage(UTextureRenderTarget2D* TexRT, FImage & I
 		
 	// RCM_MinMax means don't renormalize, just read the pixels as they are
 	//	default RCM_UNorm does funny scalings
-	FReadSurfaceDataFlags InFlags(RCM_MinMax, CubeFace_MAX);
+	FReadSurfaceDataFlags ReadFlags(RCM_MinMax, CubeFace_MAX);
+	
+	FIntRect Rect = InRectOrZero;
+	if (InRectOrZero == FIntRect(0, 0, 0, 0))
+	{
+		Rect = FIntRect(0, 0, TexRT->SizeX, TexRT->SizeY);
+	}
+
+	int64 RectSizeX = Rect.Width();
+	int64 RectSizeY = Rect.Height();
+	if ( ! ensure( Rect.Min.X >= 0 && Rect.Min.Y >= 0 && 
+		Rect.Max.X <= TexRT->SizeX && Rect.Max.Y <= TexRT->SizeY &&
+		RectSizeX >= 0 && RectSizeY >= 0 ) )
+	{
+		return false;
+	}
+	if ( RectSizeX == 0 || RectSizeY == 0 )
+	{
+		return false; // or is that a success to grab zero pixels?
+	}
 
 	if ( ReadFormat == ERawImageFormat::RGBA16F )
 	{
 		// ReadFloat16Pixels does no conversions
 		//	must be used only exactly with FloatRGBA type
 
-		Image.Init(TexRT->SizeX,TexRT->SizeY,ERawImageFormat::RGBA16F,EGammaSpace::Linear);
+		Image.Init(RectSizeX,RectSizeY,ERawImageFormat::RGBA16F,EGammaSpace::Linear);
 		
 		TArray<FFloat16Color> Colors;
-		if ( ! RenderTarget->ReadFloat16Pixels(Colors,InFlags) )
+		if ( ! RenderTarget->ReadFloat16Pixels(Colors,ReadFlags,Rect) )
 		{
 			return false;
 		}
@@ -365,15 +389,15 @@ bool FImageUtils::GetRenderTargetImage(UTextureRenderTarget2D* TexRT, FImage & I
 		//	mainly we are trying to catch the check for whether the _SRGB or non _SRGB BGRA8 format as chosen
 		EGammaSpace GammaSpace = TexRT->IsSRGB() ? EGammaSpace::sRGB : EGammaSpace::Linear;
 		
-		Image.Init(TexRT->SizeX,TexRT->SizeY,ERawImageFormat::BGRA8,GammaSpace);
+		Image.Init(RectSizeX,RectSizeY,ERawImageFormat::BGRA8,GammaSpace);
 		
 		// "LinearToGamma" is basically moot; that would only be used if we were reading float pixels to FColor
 		//	but in that case the ReadFormat should have been float, so we won't be here
 		//	gamma conversion will be handled by FImage after the pixel read, not inside RHI
-		InFlags.SetLinearToGamma( GammaSpace == EGammaSpace::sRGB );
+		ReadFlags.SetLinearToGamma( GammaSpace == EGammaSpace::sRGB );
 
 		TArray<FColor> Colors;
-		if ( ! RenderTarget->ReadPixels(Colors,InFlags) )
+		if ( ! RenderTarget->ReadPixels(Colors,ReadFlags,Rect) )
 		{
 			return false;
 		}
@@ -383,10 +407,10 @@ bool FImageUtils::GetRenderTargetImage(UTextureRenderTarget2D* TexRT, FImage & I
 	}
 	else if ( ReadFormat == ERawImageFormat::RGBA32F )
 	{
-		Image.Init(TexRT->SizeX,TexRT->SizeY,ERawImageFormat::RGBA32F,EGammaSpace::Linear);
+		Image.Init(RectSizeX,RectSizeY,ERawImageFormat::RGBA32F,EGammaSpace::Linear);
 		
 		TArray<FLinearColor> Colors;
-		if ( ! RenderTarget->ReadLinearColorPixels(Colors,InFlags) )
+		if ( ! RenderTarget->ReadLinearColorPixels(Colors,ReadFlags,Rect) )
 		{
 			return false;
 		}
diff --git a/Engine/Source/Runtime/Engine/Public/ImageUtils.h b/Engine/Source/Runtime/Engine/Public/ImageUtils.h
--- a/Engine/Source/Runtime/Engine/Public/ImageUtils.h
+++ b/Engine/Source/Runtime/Engine/Public/ImageUtils.h
@@ -439,5 +439,6 @@ public:
 	* @param OutImage	Filled with the image, allocated as needed
 	*/
 	ENGINE_API static bool GetRenderTargetImage(UTextureRenderTarget2D* TexRT, FImage & OutImage);
+	ENGINE_API static bool GetRenderTargetImage(UTextureRenderTarget2D* TexRT, FImage & OutImage, const FIntRect & Rect);
 
 };
diff --git a/Engine/Source/Runtime/Landscape/Landscape.Build.cs b/Engine/Source/Runtime/Landscape/Landscape.Build.cs
--- a/Engine/Source/Runtime/Landscape/Landscape.Build.cs
+++ b/Engine/Source/Runtime/Landscape/Landscape.Build.cs
@@ -21,7 +21,8 @@ public class Landscape : ModuleRules
 				"CoreUObject",
 				"ApplicationCore",
 				"Engine",
-				"RenderCore", 
+				"RenderCore",
+				"ImageCore",
 				"RHI",
 				"Renderer",
 				"Foliage",
diff --git a/Engine/Source/Runtime/Landscape/Private/LandscapeEdit.cpp b/Engine/Source/Runtime/Landscape/Private/LandscapeEdit.cpp
--- a/Engine/Source/Runtime/Landscape/Private/LandscapeEdit.cpp
+++ b/Engine/Source/Runtime/Landscape/Private/LandscapeEdit.cpp
@@ -46,6 +46,7 @@ LandscapeEdit.cpp: Landscape editing
 #include "Serialization/MemoryWriter.h"
 #include "MaterialCachedData.h"
 #include "Math/UnrealMathUtility.h"
+#include "ImageUtils.h"
 
 #if WITH_EDITOR
 #include "Engine/World.h"
@@ -7718,7 +7719,6 @@ void ALandscapeProxy::RemoveOverlappingComponent(ULandscapeComponent* Component)
 
 TArray<FLinearColor> ALandscapeProxy::SampleRTData(UTextureRenderTarget2D* InRenderTarget, FLinearColor InRect)
 {
-
 	if (!InRenderTarget)
 	{
 		FMessageLog("Blueprint").Warning(LOCTEXT("SampleRTData_InvalidRenderTarget", "SampleRTData: Render Target must be non-null."));
@@ -7731,48 +7731,25 @@ TArray<FLinearColor> ALandscapeProxy::SampleRTData(UTextureRenderTarget2D* InRen
 	}
 	else
 	{
-		ETextureRenderTargetFormat format = (InRenderTarget->RenderTargetFormat);
+		InRect.R = static_cast<float>(FMath::Clamp(int(InRect.R), 0, InRenderTarget->SizeX - 1));
+		InRect.G = static_cast<float>(FMath::Clamp(int(InRect.G), 0, InRenderTarget->SizeY - 1));
+		InRect.B = static_cast<float>(FMath::Clamp(int(InRect.B), int(InRect.R + 1), InRenderTarget->SizeX));
+		InRect.A = static_cast<float>(FMath::Clamp(int(InRect.A), int(InRect.G + 1), InRenderTarget->SizeY));
+		FIntRect Rect = FIntRect(static_cast<int32>(InRect.R), static_cast<int32>(InRect.G), static_cast<int32>(InRect.B), static_cast<int32>(InRect.A));
 
-		if ((format == (RTF_RGBA16f)) || (format == (RTF_RGBA32f)) || (format == (RTF_RGBA8)))
+		FImage Image;
+		if ( ! FImageUtils::GetRenderTargetImage(InRenderTarget,Image,Rect) )
 		{
+			FMessageLog("Blueprint").Warning(LOCTEXT("SampleRTData_FailedGetRenderTarget", "SampleRTData: GetRenderTargetImage failed."));
+			return { FLinearColor(0,0,0,0) };
+		}
 
-			FTextureRenderTargetResource* RTResource = InRenderTarget->GameThread_GetRenderTargetResource();
-
-			InRect.R = static_cast<float>(FMath::Clamp(int(InRect.R), 0, InRenderTarget->SizeX - 1));
-			InRect.G = static_cast<float>(FMath::Clamp(int(InRect.G), 0, InRenderTarget->SizeY - 1));
-			InRect.B = static_cast<float>(FMath::Clamp(int(InRect.B), int(InRect.R + 1), InRenderTarget->SizeX));
-			InRect.A = static_cast<float>(FMath::Clamp(int(InRect.A), int(InRect.G + 1), InRenderTarget->SizeY));
-			FIntRect Rect = FIntRect(static_cast<int32>(InRect.R), static_cast<int32>(InRect.G), static_cast<int32>(InRect.B), static_cast<int32>(InRect.A));
-
-			FReadSurfaceDataFlags ReadPixelFlags(RCM_MinMax);
-
-			TArray<FColor> OutLDR;
-			TArray<FLinearColor> OutHDR;
-
-			TArray<FLinearColor> OutVals;
-
-			bool ishdr = ((format == (RTF_R16f)) || (format == (RTF_RG16f)) || (format == (RTF_RGBA16f)) || (format == (RTF_R32f)) || (format == (RTF_RG32f)) || (format == (RTF_RGBA32f)));
-
-			if (!ishdr)
-			{
-				RTResource->ReadPixels(OutLDR, ReadPixelFlags, Rect);
-				for (auto i : OutLDR)
-				{
-					OutVals.Add(FLinearColor(float(i.R), float(i.G), float(i.B), float(i.A)) / 255.0f);
-				}
-			}
-			else
-			{
-				RTResource->ReadLinearColorPixels(OutHDR, ReadPixelFlags, Rect);
-				return OutHDR;
-			}
+		Image.ChangeFormat(ERawImageFormat::RGBA32F,EGammaSpace::Linear);
+		
+		TArrayView64<FLinearColor> Colors = Image.AsRGBA32F();
 
-			return OutVals;
-		}
+		return TArray<FLinearColor>( Colors );
 	}
-	FMessageLog("Blueprint").Warning(LOCTEXT("SampleRTData_InvalidTexture", "SampleRTData: Currently only 4 channel formats are supported: RTF_RGBA8, RTF_RGBA16f, and RTF_RGBA32f."));
-
-	return { FLinearColor(0,0,0,0) };
 }
 
 bool ALandscapeProxy::LandscapeImportHeightmapFromRenderTarget(UTextureRenderTarget2D* InRenderTarget, bool InImportHeightFromRGChannel)
@@ -8074,10 +8051,11 @@ bool ALandscapeProxy::LandscapeImportWeightmapFromRenderTarget(UTextureRenderTar
 				RTData = SampleRTData(InRenderTarget, SampleRect);
 
 				TArray<uint8> LayerData;
+				LayerData.Reserve( RTData.Num() );
 
-				for (auto i : RTData)
+				for (const FLinearColor & RTColor : RTData)
 				{
-					LayerData.Add((uint8)(FMath::Clamp((float)i.R, 0.0f, 1.0f) * 255));
+					LayerData.Add( FColor::QuantizeUNormFloatTo8(RTColor.R) );
 				}
 
 				FLandscapeInfoLayerSettings CurWeightmapInfo;
