Fixed incorrect landscape components being rendered when doing a partial update of landscape edit layers : we need components whose "final" weightmap is shared with the one on which GetLandscapeComponentWeightmapsToRender is called, instead of the ones whose edit layers' weightmaps are shared)
#tests editor

diff --git a/Engine/Source/Runtime/Landscape/Private/LandscapeEditLayers.cpp b/Engine/Source/Runtime/Landscape/Private/LandscapeEditLayers.cpp
--- a/Engine/Source/Runtime/Landscape/Private/LandscapeEditLayers.cpp
+++ b/Engine/Source/Runtime/Landscape/Private/LandscapeEditLayers.cpp
@@ -8213,19 +8213,7 @@ void ULandscapeComponent::GetLandscapeComponentWeightmapsToRender(TSet<ULandscap
 			if (Usage != nullptr && Usage->ChannelUsage[Channel] != nullptr)
 			{
 				ULandscapeComponent* Component = Usage->ChannelUsage[Channel];
-				Component->ForEachLayer([&](const FGuid& LayerGuid, FLandscapeLayerComponentData& LayerData)
-				{
-					for (ULandscapeWeightmapUsage* Usage : LayerData.WeightmapData.TextureUsages)
-					{
-						for (int32 Channel = 0; Channel < ULandscapeWeightmapUsage::NumChannels; ++Channel)
-						{
-							if (Usage != nullptr && Usage->ChannelUsage[Channel] != nullptr)
-							{
-								OutWeightmapComponents.Add(Usage->ChannelUsage[Channel]);
-							}
-						}
-					}
-				});
+				OutWeightmapComponents.Add(Component);
 			}
 		}
 	}
