Fixes to make Runtime Grass Generation work on all platforms
- Use the correct landscape material on every platform.
- Use pooled render targets for render and readback, instead of UObject based render targets that may be deallocated by GC.
- Added an optional async fetch stage controlled by a cvar and disabled by default.
- Incorporated the check for zero into the grassmap readback copy to reduce the time spent in fetch.
- Small optimization on Populated components.
- Runtime grass is disabled by default, controlled by cvar.
- Converted editable opt-in to hidden opt-out on landscape proxy settings.

diff --git a/Engine/Source/Runtime/Landscape/Classes/Landscape.h b/Engine/Source/Runtime/Landscape/Classes/Landscape.h
--- a/Engine/Source/Runtime/Landscape/Classes/Landscape.h
+++ b/Engine/Source/Runtime/Landscape/Classes/Landscape.h
@@ -408,6 +408,8 @@ public:
 
 	LANDSCAPE_API void EnableNaniteSkirts(bool bInEnable, float InSkirtDepth, bool bInShouldDirtyPackage);
 
+	LANDSCAPE_API void SetDisableRuntimeGrassMapGeneration(bool bInDisableRuntimeGrassMapGeneration);
+
 protected:
 	FName GenerateUniqueLayerName(FName InName = NAME_None) const;
 
diff --git a/Engine/Source/Runtime/Landscape/Classes/LandscapeComponent.h b/Engine/Source/Runtime/Landscape/Classes/LandscapeComponent.h
--- a/Engine/Source/Runtime/Landscape/Classes/LandscapeComponent.h
+++ b/Engine/Source/Runtime/Landscape/Classes/LandscapeComponent.h
@@ -190,6 +190,9 @@ struct IBuffer2DView
 	// copy up to Count elements to Dest, in X then Y order (standard image order)
 	virtual void CopyTo(T* Dest, int32 Count) const = 0;
 
+	// copy up to Count elements to Dest, in X then Y order (standard image order)
+	virtual bool CopyToAndCalcIsAllZero(T* Dest, int32 Count) const = 0;
+
 	// return the total number of elements
 	virtual int32 Num() const = 0;
 };
@@ -237,7 +240,7 @@ struct FLandscapeComponentGrassData
 	bool HasData() const;
 
 	void InitializeFrom(const TArray<uint16>& HeightData, const TMap<ULandscapeGrassType*, TArray<uint8>>& WeightData);
-	void InitializeFrom(IBuffer2DView<uint16>* HeightData, TMap<ULandscapeGrassType*, IBuffer2DView<uint8>*>& WeightData);
+	void InitializeFrom(IBuffer2DView<uint16>* HeightData, TMap<ULandscapeGrassType*, IBuffer2DView<uint8>*>& WeightData, bool bStripEmptyWeights);
 
 	bool HasWeightData() const;
 	TArrayView<uint8> GetWeightData(const ULandscapeGrassType* GrassType);
@@ -906,6 +909,9 @@ public:
 	virtual TSubclassOf<class UHLODBuilder> GetCustomHLODBuilderClass() const override;
 #endif
 
+	int32 GetCurrentRuntimeMaterialInstanceCount() const;
+	class UMaterialInterface* GetCurrentRuntimeMaterialInterface(int32 InIndex);
+
 	LANDSCAPE_API int32 GetMaterialInstanceCount(bool InDynamic = true) const;
 	LANDSCAPE_API class UMaterialInstance* GetMaterialInstance(int32 InIndex, bool InDynamic = true) const;
 
diff --git a/Engine/Source/Runtime/Landscape/Classes/LandscapeProxy.h b/Engine/Source/Runtime/Landscape/Classes/LandscapeProxy.h
--- a/Engine/Source/Runtime/Landscape/Classes/LandscapeProxy.h
+++ b/Engine/Source/Runtime/Landscape/Classes/LandscapeProxy.h
@@ -412,6 +412,10 @@ protected:
 	
 #endif // WITH_EDITORONLY_DATA
 
+	/** Disable runtime grass data generation.  If disabled, the grass maps will be serialized at cook time. Do not set directly, use ALandscape::SetDisableRuntimeGrassMapGeneration to ensure it is set on all loaded proxies. */
+	UPROPERTY(meta = (LandscapeInherited))
+	bool bDisableRuntimeGrassMapGeneration = false;
+
 public:
 	/** Offset in quads from global components grid origin (in quads) **/
 	UPROPERTY()
@@ -626,6 +630,8 @@ public:
 	inline void InvalidateGrassTypeSummary() { GrassTypeSummary.LandscapeComponentCount = -1; }
 	void UpdateGrassTypeSummary();
 
+	inline bool GetDisableRuntimeGrassMapGeneration() { return bDisableRuntimeGrassMapGeneration; }
+
 	/**
 	 *	The resolution to cache lighting at, in texels/quad in one axis
 	 *  Total resolution would be changed by StaticLightingResolution*StaticLightingResolution
@@ -811,11 +817,6 @@ public:
 	UPROPERTY(EditAnywhere, Category = Landscape, AdvancedDisplay, meta = (LandscapeInherited))
 	bool bUseCompressedHeightmapStorage = false;
 
-	/** Enable runtime grass data generation to save memory, when grass.GrassMap.UseRuntimeGeneration is true (this is checked per platform).
-		When enabled grass data is not serialized during cook and will be regenerated at runtime when the landscape is loaded. */
-	UPROPERTY(EditAnywhere, Category = Landscape, AdvancedDisplay, meta = (LandscapeInherited))
-	bool bUseRuntimeGrassMapGeneration = false;
-
 	/** Strip Physics/collision components when cooked for client */
 	UPROPERTY(EditAnywhere, Category = Landscape, AdvancedDisplay, meta = (LandscapeOverridable))
 	bool bStripPhysicsWhenCookedClient = false;
diff --git a/Engine/Source/Runtime/Landscape/Private/Landscape.cpp b/Engine/Source/Runtime/Landscape/Private/Landscape.cpp
--- a/Engine/Source/Runtime/Landscape/Private/Landscape.cpp
+++ b/Engine/Source/Runtime/Landscape/Private/Landscape.cpp
@@ -315,6 +315,33 @@ UMaterialInstance* ULandscapeComponent::GetMaterialInstance(int32 InIndex, bool
 	return MaterialInstances[InIndex];
 }
 
+int32 ULandscapeComponent::GetCurrentRuntimeMaterialInstanceCount() const
+{
+	ALandscapeProxy* Proxy = GetLandscapeProxy();
+	const ERHIFeatureLevel::Type FeatureLevel = Proxy->GetWorld()->GetFeatureLevel();
+	if (FeatureLevel == ERHIFeatureLevel::ES3_1)
+	{
+		return MobileMaterialInterfaces.Num();
+	}
+
+	bool bDynamic = Proxy->bUseDynamicMaterialInstance;
+	return GetMaterialInstanceCount(bDynamic);
+}
+
+class UMaterialInterface* ULandscapeComponent::GetCurrentRuntimeMaterialInterface(int32 InIndex)
+{
+	ALandscapeProxy* Proxy = GetLandscapeProxy();
+	const ERHIFeatureLevel::Type FeatureLevel = GetLandscapeProxy()->GetWorld()->GetFeatureLevel();
+
+	if (FeatureLevel == ERHIFeatureLevel::ES3_1)
+	{
+		return MobileMaterialInterfaces[InIndex];
+	}
+
+	bool bDynamic = Proxy->bUseDynamicMaterialInstance;
+	return GetMaterialInstance(InIndex, bDynamic);
+}
+
 UMaterialInstanceDynamic* ULandscapeComponent::GetMaterialInstanceDynamic(int32 InIndex) const
 {
 	ALandscapeProxy* Actor = GetLandscapeProxy();
@@ -3535,6 +3562,22 @@ void ALandscape::EnableNaniteSkirts(bool bInEnable, float InSkirtDepth, bool bIn
 	}
 }
 
+void ALandscape::SetDisableRuntimeGrassMapGeneration(bool bInDisableRuntimeGrassMapGeneration)
+{
+	bDisableRuntimeGrassMapGeneration = bInDisableRuntimeGrassMapGeneration;
+	if (ULandscapeInfo* LandscapeInfo = GetLandscapeInfo())
+	{
+		LandscapeInfo->ForEachLandscapeProxy([bInDisableRuntimeGrassMapGeneration](ALandscapeProxy* Proxy) -> bool
+		{
+			if (Proxy != nullptr)
+			{
+				Proxy->bDisableRuntimeGrassMapGeneration = bInDisableRuntimeGrassMapGeneration;
+			}
+			return true;
+		});
+	}
+}
+
 void ALandscapeProxy::OnFeatureLevelChanged(ERHIFeatureLevel::Type NewFeatureLevel)
 {
 	FlushGrassComponents();
@@ -3571,11 +3614,19 @@ void ALandscapeProxy::PreSave(FObjectPreSaveContext ObjectSaveContext)
 		{
 			// generate all of the grass data
 			BuildGrassMaps();
+
+			int32 ValidGrassCount = 0;
 			for (ULandscapeComponent* Component : LandscapeComponents)
 			{
 				// Manually reset dirty flag (for post save)
 				Component->GrassData->bIsDirty = false;
+				if (Component->GrassData->HasValidData())
+				{
+					ValidGrassCount++;
+				}
 			}
+
+			UE_LOG(LogGrass, Verbose, TEXT("PRESAVE: landscape %s has %d / %d valid grass components (UseRuntimeGeneration %d Disable %d)"), *GetName(), ValidGrassCount, LandscapeComponents.Num(), GGrassMapUseRuntimeGeneration, bDisableRuntimeGrassMapGeneration);
 		}
 	}
 
@@ -4105,6 +4156,36 @@ void ALandscapeProxy::PostLoad()
 		BodyInstance.FixupData(this);
 	}
 
+	for (ULandscapeComponent* Comp : LandscapeComponents)
+	{
+		if (Comp == nullptr)
+		{
+			continue;
+		}
+
+		UE_LOG(LogGrass, Verbose, TEXT("POSTLOAD: component %s on landscape %s UseRuntimeGeneration %d Disable: %d data: %d"),
+			*Comp->GetName(),
+			*GetName(),
+			GGrassMapUseRuntimeGeneration, bDisableRuntimeGrassMapGeneration, Comp->GrassData->NumElements);
+
+#if !WITH_EDITOR
+		// if using runtime grass gen, it should have been cleared out in PreSave
+		if (GGrassMapUseRuntimeGeneration && !bDisableRuntimeGrassMapGeneration)
+		{
+			if (Comp->GrassData->HasData())
+			{
+				UE_LOG(LogGrass, Warning, TEXT("grass.GrassMap.UseRuntimeGeneration is enabled, but component %s on landscape %s has unnecessary grass data saved.  Ensure grass.GrassMap.UseRuntimeGeneration is enabled at cook time to reduce cooked data size."),
+					*Comp->GetName(),
+					*GetName());
+
+				// Free the memory, so at least we will save the space at runtime.
+				TUniquePtr<FLandscapeComponentGrassData> NewGrassData = MakeUnique<FLandscapeComponentGrassData>();
+				Comp->GrassData = MakeShareable(NewGrassData.Release());
+			}
+		}
+#endif // !WITH_EDITOR
+	}
+
 #if WITH_EDITOR
 	PRAGMA_DISABLE_DEPRECATION_WARNINGS;
 	if (!LandscapeMaterialsOverride_DEPRECATED.IsEmpty())
@@ -4171,23 +4252,6 @@ void ALandscapeProxy::PostLoad()
 			continue;
 		}
 
-#if !WITH_EDITOR
-		// if using runtime grass gen, it should have been cleared out in PreSave
-		if (GGrassMapUseRuntimeGeneration)
-		{
-			if (Comp->GrassData->HasValidData() && bUseRuntimeGrassMapGeneration)
-			{
-				UE_LOG(LogGrass, Warning, TEXT("grass.GrassMap.UseRuntimeGeneration is enabled, but component %s on landscape %s has unnecessary grass data saved.  Ensure grass.GrassMap.UseRuntimeGeneration is enabled at cook time to reduce cooked data size."),
-					*Comp->GetName(),
-					*GetName());
-
-				// Free the memory, so at least we will save the space at runtime.
-				TUniquePtr<FLandscapeComponentGrassData> NewGrassData = MakeUnique<FLandscapeComponentGrassData>();
-				Comp->GrassData = MakeShareable(NewGrassData.Release());
-			}
-		}
-#endif // !WITH_EDITOR
-
 		// Validate the layer combination and store it in the MaterialInstanceConstantMap
 		UMaterialInstance* MaterialInstance = Comp->GetMaterialInstance(0, false);
 
diff --git a/Engine/Source/Runtime/Landscape/Private/LandscapeGrass.cpp b/Engine/Source/Runtime/Landscape/Private/LandscapeGrass.cpp
--- a/Engine/Source/Runtime/Landscape/Private/LandscapeGrass.cpp
+++ b/Engine/Source/Runtime/Landscape/Private/LandscapeGrass.cpp
@@ -965,7 +965,7 @@ TArray<uint16> ULandscapeComponent::RenderWPOHeightmap(int32 LOD)
 		{
 			FLandscapeGrassWeightExporter Exporter(GetLandscapeProxy(), { this }, /*bInNeedsGrassmap = */ false, /*bInNeedsHeightmap =*/ true, {});
 			TMap<ULandscapeComponent*, TUniquePtr<FLandscapeComponentGrassData>, TInlineSetAllocator<1>> TempGrassData;
-			TempGrassData = Exporter.FetchResults();
+			TempGrassData = Exporter.FetchResults(/* bFreeAsyncReadback= */ true);
 			Results = TArray<uint16>(TempGrassData[this]->GetHeightData());
 		}
 		else
@@ -974,7 +974,7 @@ TArray<uint16> ULandscapeComponent::RenderWPOHeightmap(int32 LOD)
 			HeightMips.Add(LOD);
 			FLandscapeGrassWeightExporter Exporter(GetLandscapeProxy(), { this }, /*bInNeedsGrassmap = */ false, /*bInNeedsHeightmap =*/ false, MoveTemp(HeightMips));
 			TMap<ULandscapeComponent*, TUniquePtr<FLandscapeComponentGrassData>, TInlineSetAllocator<1>> TempGrassData;
-			TempGrassData = Exporter.FetchResults();
+			TempGrassData = Exporter.FetchResults(/* bFreeAsyncReadback= */ true);
 			Results = MoveTemp(TempGrassData[this]->HeightMipData[LOD]);
 		}
 	}
@@ -1483,12 +1483,13 @@ void FLandscapeComponentGrassData::InitializeFrom(const TArray<uint16>& HeightDa
 }
 
 // Note that this is a relatively expensive operation (in the runtime sense), as it's making a copy of the buffers
-void FLandscapeComponentGrassData::InitializeFrom(IBuffer2DView<uint16>* HeightData, TMap<ULandscapeGrassType*, IBuffer2DView<uint8>*>& WeightData)
+void FLandscapeComponentGrassData::InitializeFrom(IBuffer2DView<uint16>* HeightData, TMap<ULandscapeGrassType*, IBuffer2DView<uint8>*>& WeightData, bool bStripEmptyWeights)
 {
-	WeightOffsets.Empty(WeightData.Num());
+	const int32 OriginalGrassTypeCount = WeightData.Num();
+	WeightOffsets.Empty(OriginalGrassTypeCount);
 
 	// If weight data is empty make sure we don't have any memory allocated to grass
-	if (WeightData.Num() == 0)
+	if (OriginalGrassTypeCount == 0)
 	{
 		NumElements = 0;
 		HeightWeightData.Empty();
@@ -1496,7 +1497,9 @@ void FLandscapeComponentGrassData::InitializeFrom(IBuffer2DView<uint16>* HeightD
 	}
 
 	NumElements = HeightData->Num();
-	HeightWeightData.SetNumUninitialized(NumElements * sizeof(uint16) + NumElements * WeightData.Num() * sizeof(uint8));
+
+	// reserve space for up to OriginalGrassTypeCount weight elements
+	HeightWeightData.SetNumUninitialized(NumElements * sizeof(uint16) + NumElements * OriginalGrassTypeCount * sizeof(uint8));
 
 	uint8* CopyDest = HeightWeightData.GetData();
 	int32 CopyOffset = 0;
@@ -1504,20 +1507,32 @@ void FLandscapeComponentGrassData::InitializeFrom(IBuffer2DView<uint16>* HeightD
 
 	check((CopyOffset + CopySize) <= HeightWeightData.Num());
 
+	// copy the height data
 	HeightData->CopyTo((uint16*)&CopyDest[CopyOffset], CopySize / sizeof(uint16));
-
 	CopyOffset += CopySize;
 	CopySize = NumElements * sizeof(uint8);
 
 	for (const TPair<ULandscapeGrassType*, IBuffer2DView<uint8>*>& Pair : WeightData)
 	{
-		WeightOffsets.Add(Pair.Key, CopyOffset);
 		check(Pair.Value->Num() == NumElements);
 		check((CopyOffset + CopySize) <= HeightWeightData.Num());
 
-		Pair.Value->CopyTo(&CopyDest[CopyOffset], CopySize);
-		CopyOffset += CopySize;
+		const bool bIsAllZero = Pair.Value->CopyToAndCalcIsAllZero(&CopyDest[CopyOffset], CopySize);
+
+		if (bStripEmptyWeights && bIsAllZero)
+		{
+			// strip all zero weights
+		}
+		else
+		{
+			// record the weight data
+			WeightOffsets.Add(Pair.Key, CopyOffset);
+			CopyOffset += CopySize;
+		}
 	}
+
+	// shrink if necessary to account for stripped weight data
+	HeightWeightData.SetNum(CopyOffset);
 }
 
 FArchive& operator<<(FArchive& Ar, FLandscapeComponentGrassData& Data)
@@ -2983,6 +2998,15 @@ void ALandscapeProxy::ProcessAsyncGrassInstanceTasks(bool bWaitAsyncTasks, bool
 			int32 NumBuiltRenderInstances = Inner.Builder->InstanceBuffer.GetNumInstances();
 			//UE_LOG(LogCore, Display, TEXT("%d instances in %4.0fms     %6.0f instances / sec"), NumBuiltRenderInstances, 1000.0f * float(Inner.Builder->BuildTime), float(NumBuiltRenderInstances) / float(Inner.Builder->BuildTime));
 
+			FCachedLandscapeFoliage::FGrassCompKey& InnerKey = Inner.Key;
+
+			UE_LOG(LogGrass, Verbose, TEXT("ASYNC GRASS INSTANCES COMPLETE %s %s %d %d (%d)"),
+				*InnerKey.BasedOn->GetName(),
+				*(InnerKey.GrassType->GrassVarieties[InnerKey.VarietyIndex].GrassMesh->GetName()),
+				InnerKey.SubsectionX,
+				InnerKey.SubsectionY,
+				NumBuiltRenderInstances);
+
 			if (GrassISMComponent && StillUsed.Contains(GrassISMComponent))
 			{
 				if (NumBuiltRenderInstances > 0)
diff --git a/Engine/Source/Runtime/Landscape/Private/LandscapeGrassMapsBuilder.cpp b/Engine/Source/Runtime/Landscape/Private/LandscapeGrassMapsBuilder.cpp
--- a/Engine/Source/Runtime/Landscape/Private/LandscapeGrassMapsBuilder.cpp
+++ b/Engine/Source/Runtime/Landscape/Private/LandscapeGrassMapsBuilder.cpp
@@ -19,6 +19,20 @@
 
 #define GRASS_DEBUG_LOG(...) UE_LOG(LogGrass, Verbose, __VA_ARGS__)
 
+#define DEBUG_TRANSITION(StateRef, StageBefore, StageAfter) \
+	GRASS_DEBUG_LOG(TEXT("%s %s -> %s (after %d ticks) Pend:%d Strm:%d Rend:%d Fetch:%d Pop:%d NR:%d Total:%d"), \
+		StateRef.Component ? *StateRef.Component->GetName() : TEXT("<REMOVED>"), \
+		TEXT(#StageBefore), \
+		TEXT(#StageAfter), \
+		StateRef.TickCount, \
+		PendingCount, \
+		StreamingCount, \
+		RenderingCount, \
+		AsyncFetchCount, \
+		PopulatedCount, \
+		NotReadyCount, \
+		ComponentStates.Num())
+
 extern int32 GGrassEnable;
 extern float GGrassCullDistanceScale;
 
@@ -28,6 +42,12 @@ FAutoConsoleVariableRef CVarGrassMapUseRuntimeGeneration(
 	GGrassMapUseRuntimeGeneration,
 	TEXT("Enable runtime grass map generation to save disk space and runtime memory.  When enabled the grass density maps are not serialized and are built on the fly at runtime."));
 
+int32 GGrassMapUseAsyncFetch = 0;
+FAutoConsoleVariableRef CVarGrassMapUseAsyncFetch(
+	TEXT("grass.GrassMap.UseAsyncFetch"),
+	GGrassMapUseAsyncFetch,
+	TEXT("Enable async fetch tasks to readback the runtime grass maps from the GPU.  When disabled, it the fetch is performed on the game thread, when enabled it uses an async task instead."));
+
 int32 GGrassMapAlwaysBuildRuntimeGenerationResources = 0;
 static FAutoConsoleVariableRef CVarGrassMapAlwaysBuildRuntimeGenerationResources(
 	TEXT("grass.GrassMap.AlwaysBuildRuntimeGenerationResources"),
@@ -110,7 +130,7 @@ namespace UE::Landscape
 		Hash = FCrc::TypeCrc32(Material->ComputeAllStateCRC(), Hash);
 
 		// hash the heightmap texture (we use Source Id as it is a content hash, lighting guid is random)
-		UTexture2D* Heightmap = Component->GetHeightmap();
+		const UTexture2D* Heightmap = Component->GetHeightmap();
 		check(Heightmap->Source.IsValid());
 		Hash = FCrc::TypeCrc32(Heightmap->Source.GetId(), Hash);
 
@@ -141,7 +161,7 @@ namespace UE::Landscape
 	}
 
 #if WITH_EDITOR
-	void CompileGrassMapShader(ULandscapeComponent* Component)
+	void CompileGrassMapShader(const ULandscapeComponent* Component)
 	{
 		if (Component->GetMaterialInstanceCount(false) > 0)
 		{
@@ -159,18 +179,38 @@ namespace UE::Landscape
 	bool CanRenderGrassMap(ULandscapeComponent *Component)
 	{
 		// Check we can render
-		UWorld* ComponentWorld = Component->GetWorld();
+		const UWorld* ComponentWorld = Component->GetWorld();
 		if (GUsingNullRHI || !ComponentWorld || !Component->SceneProxy)
 		{
+			GRASS_DEBUG_LOG(TEXT("GrassMap No SceneProxy for %s"), *Component->GetName());
 			return false;
 		}
 
-		UMaterialInstance* MaterialInstance = Component->GetMaterialInstanceCount(false) > 0 ? Component->GetMaterialInstance(0) : nullptr;
-		FMaterialResource* MaterialResource = MaterialInstance != nullptr ? MaterialInstance->GetMaterialResource(ComponentWorld->GetFeatureLevel()) : nullptr;
-
 		// Check we can render the material
+		const int32 MatInstCount = Component->GetCurrentRuntimeMaterialInstanceCount();
+		if (MatInstCount <= 0)
+		{
+			GRASS_DEBUG_LOG(TEXT("GrassMap MaterialInstanceCount %d <= 0 for %s (%d dyn:%d inst:%d mob:%d)"),
+				MatInstCount,
+				*Component->GetName(),
+				Component->GetLandscapeProxy()->bUseDynamicMaterialInstance ? 1 : 0,
+				Component->MaterialInstancesDynamic.Num(),
+				Component->MaterialInstances.Num(),
+				Component->MobileMaterialInterfaces.Num());
+			return false;
+		}
+
+		UMaterialInterface* MaterialInterface = Component->GetCurrentRuntimeMaterialInterface(0);
+		if (MaterialInterface == nullptr)
+		{
+			GRASS_DEBUG_LOG(TEXT("GrassMap MaterialInterface NULL for %s"), *Component->GetName());
+			return false;
+		}
+
+		const FMaterialResource* MaterialResource = MaterialInterface->GetMaterialResource(ComponentWorld->GetFeatureLevel());
 		if (MaterialResource == nullptr)
 		{
+			GRASS_DEBUG_LOG(TEXT("GrassMap MaterialResource NULL for %s (%s feature level %d)"), *Component->GetName(), *MaterialInterface->GetName(), ComponentWorld->GetFeatureLevel());
 			return false;
 		}
 
@@ -178,9 +218,10 @@ namespace UE::Landscape
 		FMaterialShaderTypes ShaderTypes;
 		UE::Landscape::Grass::AddGrassWeightShaderTypes(ShaderTypes);
 
-		FVertexFactoryType* LandscapeGrassVF = FindVertexFactoryType(FName(TEXT("FLandscapeFixedGridVertexFactory"), FNAME_Find));
+		const FVertexFactoryType* LandscapeGrassVF = FindVertexFactoryType(FName(TEXT("FLandscapeFixedGridVertexFactory"), FNAME_Find));
 		if (!MaterialResource->HasShaders(ShaderTypes, LandscapeGrassVF))
 		{
+			GRASS_DEBUG_LOG(TEXT("GrassMap MaterialResource does not have FixedGridVF for %s"), *Component->GetName());
 			return false;
 		}
 		return true;
@@ -306,6 +347,13 @@ void FLandscapeGrassMapsBuilder::FPendingComponent::UpdatePriorityDistance(const
 	PriorityKey = MinSqrDistanceToComponent;
 }
 
+void FAsyncFetchTask::DoWork()
+{
+	// do not delete the async readback resources (it must be done on the game thread, after the task completes)
+	constexpr bool bFreeAsyncReadback = false;
+	Results = ActiveRender->FetchResults(bFreeAsyncReadback);
+}
+
 bool FLandscapeGrassMapsBuilder::UpdateTrackedComponents(const TArray<FVector>& Cameras, int32 LocalMaxRendering, int32 MaxExpensiveUpdateChecksToPerform, bool bCancelAndEvictAll)
 {
 	SCOPE_CYCLE_COUNTER(STAT_UpdateComponentGrassMaps);
@@ -322,6 +370,8 @@ bool FLandscapeGrassMapsBuilder::UpdateTrackedComponents(const TArray<FVector>&
 
 	AmortizedUpdate.StartUpdateTick(ComponentStates.Num(), MaxExpensiveUpdateChecksToPerform);
 
+	const bool bIsGameWorld = World->IsGameWorld();
+
 	// Iterate our components, removing invalid ones and counting how many are in each state.
 	// We can also immediately process any components in the populated or rendering states
 	// (those states don't need to consider throttling when moving to the next state, and processing them early frees up slots)
@@ -343,24 +393,18 @@ bool FLandscapeGrassMapsBuilder::UpdateTrackedComponents(const TArray<FVector>&
 			switch (State->Stage)
 			{
 			case EComponentStage::Pending:
-				if (World->IsGameWorld())
-				{
-					// if runtime grass generation is disabled, go straight to not ready.
-					if (!Component->GetLandscapeProxy()->bUseRuntimeGrassMapGeneration)
-					{
-						PendingToNotReady(*State);
-						continue; // next!
-					}
-				}
 				continue; // next!
 			case EComponentStage::NotReady:
 				// in game, any not ready component will never become ready.
 				// in editor, check to see if the conditions changed.
-				if (World->IsGameWorld() || !UE::Landscape::CanRenderGrassMap(Component))
+				if (!bIsGameWorld && AmortizedUpdate.ShouldUpdate(ComponentStateIndex))
 				{
-					continue; // still not ready - next!
+					if (UE::Landscape::CanRenderGrassMap(Component))
+					{
+						break; // cancel and evict to restart build process
+					}
 				}
-				break; // cancel and evict to restart build process
+				continue; // next!
 			case EComponentStage::TextureStreaming:
 				// don't process streaming states yet -- first process rendering states to free up slots
 				StreamingStatesToProcess.Add(State);
@@ -370,12 +414,29 @@ bool FLandscapeGrassMapsBuilder::UpdateTrackedComponents(const TArray<FVector>&
 				check(State->ActiveRender != nullptr);
 				if (State->ActiveRender->CheckAndUpdateAsyncReadback())
 				{
-					// TODO [chris.tchou] We could move this to an async task, it does have a decent cost to it
-					PopulateGrassDataFromReadback(*State);
+					if (GGrassMapUseAsyncFetch != 0)
+					{
+						LaunchAsyncFetchTask(*State);
+					}
+					else
+					{
+						PopulateGrassDataFromReadback(*State);
+					}
 					bChanged = true;
 				}
 				continue; // next!
 
+			case EComponentStage::AsyncFetch:
+				{
+					FAsyncTask<FAsyncFetchTask>* Task = State->AsyncFetchTask.Get();
+					check(Task);
+					if (Task->IsDone())
+					{
+						PopulateGrassDataFromAsyncFetchTask(*State);
+					}
+				}
+				continue; // next!
+
 			case EComponentStage::GrassMapsPopulated:
 				// only check for invalidation of populated grass maps once in a while
 				if (AmortizedUpdate.ShouldUpdate(ComponentStateIndex))
@@ -395,27 +456,30 @@ bool FLandscapeGrassMapsBuilder::UpdateTrackedComponents(const TArray<FVector>&
 					}
 
 #if WITH_EDITOR
-					UMaterialInterface* Material = Component->GetLandscapeMaterial();
-					if (Material == nullptr)
-					{
-						break; // cancel and evict
-					}
-
-					// check if any dependencies changed
-					uint32 CurGrassMapGenerationHash = UE::Landscape::ComputeGrassMapGenerationHash(Component, Material);
-					if (State->GrassMapGenerationHash != CurGrassMapGenerationHash)
-					{
-						break; // cancel and evict to restart process
-					}
-
-					// check if any grass types have changed -- this invalidates foliage instances but not the grass maps
-					const TArray<TObjectPtr<ULandscapeGrassType>>& GrassTypes = Component->GetGrassTypes();
-					uint32 CurGrassInstanceGenerationHash = UE::Landscape::ComputeGrassInstanceGenerationHash(CurGrassMapGenerationHash, GrassTypes);
-					if (State->GrassInstanceGenerationHash != CurGrassInstanceGenerationHash)
+					if (!bIsGameWorld)
 					{
-						ComponentsToRemoveFoliageInstances.Add(Component);
-						Component->InvalidateGrassTypeSummary();
-						State->GrassInstanceGenerationHash = CurGrassInstanceGenerationHash;
+						UMaterialInterface* Material = Component->GetLandscapeMaterial();
+						if (Material == nullptr)
+						{
+							break; // cancel and evict
+						}
+
+						// check if any dependencies changed
+						const uint32 CurGrassMapGenerationHash = UE::Landscape::ComputeGrassMapGenerationHash(Component, Material);
+						if (State->GrassMapGenerationHash != CurGrassMapGenerationHash)
+						{
+							break; // cancel and evict to restart process
+						}
+
+						// check if any grass types have changed -- this invalidates foliage instances but not the grass maps
+						const TArray<TObjectPtr<ULandscapeGrassType>>& GrassTypes = Component->GetGrassTypes();
+						const uint32 CurGrassInstanceGenerationHash = UE::Landscape::ComputeGrassInstanceGenerationHash(CurGrassMapGenerationHash, GrassTypes);
+						if (State->GrassInstanceGenerationHash != CurGrassInstanceGenerationHash)
+						{
+							ComponentsToRemoveFoliageInstances.Add(Component);
+							Component->InvalidateGrassTypeSummary();
+							State->GrassInstanceGenerationHash = CurGrassInstanceGenerationHash;
+						}
 					}
 #endif // WITH_EDITOR
 				}
@@ -487,7 +551,6 @@ void FLandscapeGrassMapsBuilder::StartPrioritizedGrassMapGeneration(const TArray
 	// no point in calling this if there are no cameras -- we can't calculate priority
 	check(Cameras.Num() > 0);
 
-	float DiscardDistanceScale = GGrassMapGuardBandDiscardMultiplier * GGrassCullDistanceScale;
 	float MustHaveDistanceScale = GGrassMapGuardBandMultiplier * GGrassCullDistanceScale;
 
 	// update pending component priorities (distances)
@@ -531,15 +594,15 @@ void FLandscapeGrassMapsBuilder::StartPrioritizedGrassMapGeneration(const TArray
 	// now pull as many elements off of the heap as we need
 	while ((MaxComponentsToStart > 0) && PendingComponentsHeap.Num() > 0)
 	{
-		FPendingComponent Pending = PendingComponentsHeap[0];
+		const FPendingComponent Pending = PendingComponentsHeap[0];
 		FComponentState* State = Pending.State;
-		ULandscapeComponent* Component = State->Component;
+		const ULandscapeComponent* Component = State->Component;
 		check(State->Stage == EComponentStage::Pending);
 
 		// runtime generation doesn't generate past this distance
 		if (GGrassMapUseRuntimeGeneration)
 		{
-			double MaxSpawnDistance = Component->GrassTypeSummary.MaxInstanceDiscardDistance * MustHaveDistanceScale;
+			const double MaxSpawnDistance = Component->GrassTypeSummary.MaxInstanceDiscardDistance * MustHaveDistanceScale;
 			if (Pending.PriorityKey > MaxSpawnDistance * MaxSpawnDistance)
 			{
 				// the closest component is not close enough to generate yet
@@ -547,7 +610,7 @@ void FLandscapeGrassMapsBuilder::StartPrioritizedGrassMapGeneration(const TArray
 			}
 		}
 
-		PendingComponentsHeap.HeapPop(Pending, false);
+		PendingComponentsHeap.HeapPopDiscard(false);
 		if (StartGrassMapGeneration(*State, false))
 		{
 			MaxComponentsToStart--;
@@ -580,7 +643,7 @@ void FLandscapeGrassMapsBuilder::RegisterComponent(ULandscapeComponent* Componen
 	}
 }
 
-void FLandscapeGrassMapsBuilder::UnregisterComponent(ULandscapeComponent* Component)
+void FLandscapeGrassMapsBuilder::UnregisterComponent(const ULandscapeComponent* Component)
 {
 	check(Component);
 	if (FComponentState* State = ComponentStates.FindRef(Component))
@@ -642,14 +705,14 @@ void FLandscapeGrassMapsBuilder::AmortizedUpdateGrassMaps(
 		AmortizedMaxRendering *= GGrassMapPrioritizedMultiplier;
 	}
 
-	bool bCancelAndEvictAll = !GGrassEnable;
+	const bool bCancelAndEvictAll = !GGrassEnable;
 	UpdateTrackedComponents(Cameras, AmortizedMaxRendering, GGrassMapMaxDiscardChecksPerFrame, bCancelAndEvictAll);
 
 	// no point in looking to start new grass map generation if nothing is pending, if grass is disabled or there are no cameras
 	if (bAllowStartGrassMapGeneration && PendingCount > 0 && GGrassEnable && Cameras.Num() > 0)
 	{
 		// check our pipeline limits to make sure we have room to start components
-		int32 AvailableStreamingSlots = AmortizedMaxStreaming - StreamingCount;
+		const int32 AvailableStreamingSlots = AmortizedMaxStreaming - StreamingCount;
 		StartPrioritizedGrassMapGeneration(Cameras, AvailableStreamingSlots);
 	}
 }
@@ -659,7 +722,7 @@ bool FLandscapeGrassMapsBuilder::BuildGrassMapsNowForComponents(
 	TArrayView<TObjectPtr<ULandscapeComponent>> LandscapeComponents, FScopedSlowTask* SlowTask, bool bMarkDirty)
 {
 	TRACE_CPUPROFILER_EVENT_SCOPE(BuildGrassMapsNowForComponents);
-	int32 MaxStreamingRendering = FMath::Max(GGrassMapMaxComponentsForBlockingUpdate, 1);
+	const int32 MaxStreamingRendering = FMath::Max(GGrassMapMaxComponentsForBlockingUpdate, 1);
 
 	if (LandscapeComponents.IsEmpty())
 	{
@@ -680,16 +743,16 @@ bool FLandscapeGrassMapsBuilder::BuildGrassMapsNowForComponents(
 	};
 
 	// using an empty cameras array causes all distance checks to return 0, so grass maps won't be evicted
-	TArray<FVector> EmptyCamerasArray;
+	const TArray<FVector> EmptyCamerasArray;
 	int32 LastUpToDateCount = 0;
 	int32 UpToDateCount = 0;
 	
-	double StartTime = FPlatformTime::Seconds();
+	const double StartTime = FPlatformTime::Seconds();
 	double LastChangeTime = StartTime;
 	while (UpToDateCount != LandscapeComponents.Num())
 	{
 		// ensure we are making progress within a reasonable amount of time TODO [chris.tchou] there should be a better way to detect non-progress here
-		double CurTime = FPlatformTime::Seconds();
+		const double CurTime = FPlatformTime::Seconds();
 		if (CurTime > LastChangeTime + 15.0)
 		{
 			UE_LOG(LogGrass, Error, TEXT("ERROR: BuildGrassMapsNowForComponents() took too long, grass maps are not up to date"));
@@ -697,7 +760,7 @@ bool FLandscapeGrassMapsBuilder::BuildGrassMapsNowForComponents(
 		}
 
 		// update all components that are tracked (without evicting)
-		int32 UpdateAllComponentCount = ComponentStates.Num();
+		const int32 UpdateAllComponentCount = ComponentStates.Num();
 		bool bChanged = UpdateTrackedComponents(EmptyCamerasArray, MaxStreamingRendering, UpdateAllComponentCount, /* bCancelAndEvictAll= */ false);
 
 		UpToDateCount = 0;
@@ -710,7 +773,7 @@ bool FLandscapeGrassMapsBuilder::BuildGrassMapsNowForComponents(
 			if (State->Stage == EComponentStage::Pending)
 			{
 				// Start tracking to kick off the build process
-				const bool bForceCompileShaders = true;
+				constexpr bool bForceCompileShaders = true;
 				if ((AvailableStreamingSlots > 0) && StartGrassMapGeneration(*State, bForceCompileShaders))
 				{
 					// modification isn't complete yet, but convenient to dirty the package when starting the process here
@@ -780,6 +843,7 @@ bool FLandscapeGrassMapsBuilder::CancelAndEvict(FComponentState& State)
 			check(NotReadyCount > 0);
 			NotReadyCount--;
 			PendingCount++;
+			DEBUG_TRANSITION(State, NotReady, Pending);
 		}
 		break;
 	case EComponentStage::TextureStreaming:
@@ -788,6 +852,7 @@ bool FLandscapeGrassMapsBuilder::CancelAndEvict(FComponentState& State)
 			check(StreamingCount > 0);
 			StreamingCount--;
 			PendingCount++;
+			DEBUG_TRANSITION(State, TextureStreaming, Pending);
 		}
 		break;
 	case EComponentStage::Rendering:
@@ -801,12 +866,33 @@ bool FLandscapeGrassMapsBuilder::CancelAndEvict(FComponentState& State)
 					// we can't cancel yet.. must wait for the readback to complete
 					return false;
 				}
-				delete State.ActiveRender;
-				State.ActiveRender = nullptr;
+				State.ActiveRender.Reset();
 			}
 			check(RenderingCount > 0);
 			RenderingCount--;
 			PendingCount++;
+			DEBUG_TRANSITION(State, Rendering, Pending);
+		}
+		break;
+	case EComponentStage::AsyncFetch:
+		{
+			FAsyncTask<FAsyncFetchTask>* Task = State.AsyncFetchTask.Get();
+			if (Task != nullptr)
+			{
+				if (!Task->IsDone())
+				{
+					// can't cancel, async task is still in flight
+					return false;
+				}
+
+				State.ActiveRender->FreeAsyncReadback();
+				State.ActiveRender.Reset();
+				State.AsyncFetchTask.Reset();
+			}
+			check(AsyncFetchCount > 0);
+			AsyncFetchCount--;
+			PendingCount++;
+			DEBUG_TRANSITION(State, AsyncFetch, Pending);
 		}
 		break;
 	case EComponentStage::GrassMapsPopulated:
@@ -830,6 +916,7 @@ bool FLandscapeGrassMapsBuilder::CancelAndEvict(FComponentState& State)
 			check(PopulatedCount > 0);
 			PopulatedCount--;
 			PendingCount++;
+			DEBUG_TRANSITION(State, Populated, Pending);
 		}
 		break;
 	default:
@@ -842,17 +929,6 @@ bool FLandscapeGrassMapsBuilder::CancelAndEvict(FComponentState& State)
 
 	if (State.Stage != EComponentStage::Pending)
 	{
-		GRASS_DEBUG_LOG(TEXT("CancelAndEvict %s (from %d after %d ticks) p:%d S:%d R:%d P:%d NR:%d T:%d"),
-			State.Component ? *State.Component->GetName() : TEXT("<REMOVED>"),
-			State.Stage,
-			State.TickCount,
-			PendingCount,
-			StreamingCount,
-			RenderingCount,
-			PopulatedCount,
-			NotReadyCount,
-			ComponentStates.Num());
-
 		// back to pending state with you!
 		State.Stage = EComponentStage::Pending;
 		State.TickCount = 0;
@@ -876,8 +952,9 @@ bool FLandscapeGrassMapsBuilder::StartGrassMapGeneration(FComponentState& State,
 	if (World->IsGameWorld()) // including PIE
 	{
 		// if runtime grass generation is disabled, go straight to not ready
-		if (!Component->GetLandscapeProxy()->bUseRuntimeGrassMapGeneration)
+		if (Component->GetLandscapeProxy()->GetDisableRuntimeGrassMapGeneration())
 		{
+			GRASS_DEBUG_LOG(TEXT("GrassMap Proxy DisableRuntimeGeneration for %s"), *Component->GetName());
 			PendingToNotReady(State);
 			return false;
 		}
@@ -886,8 +963,7 @@ bool FLandscapeGrassMapsBuilder::StartGrassMapGeneration(FComponentState& State,
 	{
 #if WITH_EDITOR
 		// recalculate hashes
-		UMaterialInterface* Material = Component->GetLandscapeMaterial();
-		if (Material)
+		if (UMaterialInterface* Material = Component->GetLandscapeMaterial())
 		{
 			State.GrassMapGenerationHash = UE::Landscape::ComputeGrassMapGenerationHash(Component, Material);
 			const TArray<TObjectPtr<ULandscapeGrassType>>& GrassTypes = Material->GetCachedExpressionData().GrassTypes;
@@ -941,16 +1017,7 @@ void FLandscapeGrassMapsBuilder::PendingToNotReady(FComponentState& State)
 
 	RemoveFromPendingComponentHeap(&State);
 
-	GRASS_DEBUG_LOG(TEXT("%s Pending -> NotReady (after %d ticks) p:%d S:%d R:%d P:%d NR:%d T:%d"),
-		*State.Component->GetName(),
-		State.TickCount,
-		PendingCount,
-		StreamingCount,
-		RenderingCount,
-		PopulatedCount,
-		NotReadyCount,
-		ComponentStates.Num());
-
+	DEBUG_TRANSITION(State, Pending, NotReady);
 	State.TickCount = 0;
 }
 
@@ -963,16 +1030,7 @@ void FLandscapeGrassMapsBuilder::PendingToPopulatedFastPathAlreadyHasData(FCompo
 
 	RemoveFromPendingComponentHeap(&State);
 
-	GRASS_DEBUG_LOG(TEXT("%s Pending -> Populated(Existing) (after %d ticks) p:%d S:%d R:%d P:%d NR:%d T:%d"),
-		*State.Component->GetName(),
-		State.TickCount,
-		PendingCount,
-		StreamingCount,
-		RenderingCount,
-		PopulatedCount,
-		NotReadyCount,
-		ComponentStates.Num());
-
+	DEBUG_TRANSITION(State, Pending, Populated_Existing);
 	State.TickCount = 0;
 }
 
@@ -994,16 +1052,7 @@ void FLandscapeGrassMapsBuilder::PendingToPopulatedFastPathNoGrass(FComponentSta
 
 	RemoveFromPendingComponentHeap(&State);
 
-	GRASS_DEBUG_LOG(TEXT("%s Pending -> Populated(Empty) (after %d ticks) p:%d S:%d R:%d P:%d NR:%d T:%d"),
-		*Component->GetName(),
-		State.TickCount,
-		PendingCount,
-		StreamingCount,
-		RenderingCount,
-		PopulatedCount,
-		NotReadyCount,
-		ComponentStates.Num());
-
+	DEBUG_TRANSITION(State, Pending, Populated_Empty);
 	State.TickCount = 0;
 }
 
@@ -1016,7 +1065,7 @@ void FLandscapeGrassMapsBuilder::PendingToStreaming(FComponentState& State)
 	// TODO [chris.tchou] : also grab asset textures from the material.
 	State.TexturesToStream.Add(Component->GetHeightmap());
 
-	ERHIFeatureLevel::Type FeatureLevel = World->GetFeatureLevel();
+	const ERHIFeatureLevel::Type FeatureLevel = World->GetFeatureLevel();
 	for (UTexture2D* WeightmapTexture : Component->GetRenderedWeightmapTexturesForFeatureLevel(FeatureLevel))
 	{
 		State.TexturesToStream.Add(WeightmapTexture);
@@ -1035,16 +1084,7 @@ void FLandscapeGrassMapsBuilder::PendingToStreaming(FComponentState& State)
 
 	RemoveFromPendingComponentHeap(&State);
 
-	GRASS_DEBUG_LOG(TEXT("%s Pending -> Streaming (after %d ticks) p:%d S:%d R:%d P:%d NR:%d U:%d"),
-		*Component->GetName(),
-		State.TickCount,
-		PendingCount,
-		StreamingCount,
-		RenderingCount,
-		PopulatedCount,
-		NotReadyCount,
-		ComponentStates.Num());
-
+	DEBUG_TRANSITION(State, Pending, Streaming);
 	State.TickCount = 0;
 }
 
@@ -1088,23 +1128,55 @@ void FLandscapeGrassMapsBuilder::KickOffRenderAndReadback(FComponentState& State
 	check(State.ActiveRender == nullptr);
 	constexpr bool bInNeedsGrassmap = true;
 	constexpr bool bInNeedsHeightmap = true;
-	constexpr bool bUseAsyncReadback = true;
-	State.ActiveRender = new FLandscapeGrassWeightExporter(Component->GetLandscapeProxy(), { Component }, bInNeedsGrassmap, bInNeedsHeightmap, MoveTemp(HeightMips), bUseAsyncReadback);
+	State.ActiveRender.Reset(new FLandscapeGrassWeightExporter(Component->GetLandscapeProxy(), { Component }, bInNeedsGrassmap, bInNeedsHeightmap, MoveTemp(HeightMips)));
 	check(State.ActiveRender != nullptr);
 
 	State.Stage = EComponentStage::Rendering;
 	RenderingCount++;
 
-	GRASS_DEBUG_LOG(TEXT("%s Streaming -> Rendering (after %d ticks) p:%d S:%d R:%d P:%d NR:%d T:%d"),
-		*Component->GetName(),
-		State.TickCount,
-		PendingCount,
-		StreamingCount,
-		RenderingCount,
-		PopulatedCount,
-		NotReadyCount,
-		ComponentStates.Num());
+	DEBUG_TRANSITION(State, Streaming, Rendering);
+	State.TickCount = 0;
+}
+
+void FLandscapeGrassMapsBuilder::LaunchAsyncFetchTask(FComponentState& State)
+{
+	SCOPE_CYCLE_COUNTER(STAT_PopulateGrassMap);
+
+	check(State.Stage == EComponentStage::Rendering);
+	check(RenderingCount > 0);
+	RenderingCount--;
+
+	// now that render is complete, we can drop the texture streaming requests and allow textures to stream out
+	RemoveTextureStreamingRequests(State);
+
+	check(State.ActiveRender != nullptr);
+
+	State.AsyncFetchTask.Reset(new FAsyncTask<FAsyncFetchTask>(State.ActiveRender.Get()));
+	State.AsyncFetchTask->StartBackgroundTask();
+	State.Stage = EComponentStage::AsyncFetch;
+	AsyncFetchCount++;
+
+	DEBUG_TRANSITION(State, Rendering, AsyncFetch);
+	State.TickCount = 0;
+}
+
+void FLandscapeGrassMapsBuilder::PopulateGrassDataFromAsyncFetchTask(FComponentState& State)
+{
+	check(AsyncFetchCount > 0);
+	AsyncFetchCount--;
+
+	State.ActiveRender->FreeAsyncReadback();
 
+	FAsyncFetchTask& Inner = State.AsyncFetchTask->GetTask();
+	FLandscapeGrassWeightExporter::ApplyResults(Inner.Results);
+
+	State.ActiveRender.Reset();
+	State.AsyncFetchTask.Reset();
+
+	State.Stage = EComponentStage::GrassMapsPopulated;
+	PopulatedCount++;
+
+	DEBUG_TRANSITION(State, AsyncFetch, Populated);
 	State.TickCount = 0;
 }
 
@@ -1121,22 +1193,12 @@ void FLandscapeGrassMapsBuilder::PopulateGrassDataFromReadback(FComponentState&
 
 	check(State.ActiveRender != nullptr);
 	State.ActiveRender->ApplyResults();
-	delete State.ActiveRender;
-	State.ActiveRender = nullptr;
+	State.ActiveRender.Reset();
 
 	State.Stage = EComponentStage::GrassMapsPopulated;
 	PopulatedCount++;
 
-	GRASS_DEBUG_LOG(TEXT("%s Rendering -> Populated (after %d ticks) p:%d S:%d R:%d P:%d NR:%d T:%d"),
-		*State.Component->GetName(),
-		State.TickCount,
-		PendingCount,
-		StreamingCount,
-		RenderingCount,
-		PopulatedCount,
-		NotReadyCount,
-		ComponentStates.Num());
-
+	DEBUG_TRANSITION(State, Rendering, Populated);
 	State.TickCount = 0;
 }
 
@@ -1168,10 +1230,10 @@ bool FLandscapeGrassMapsBuilder::FComponentState::IsRenderReadbackComplete() con
 bool FLandscapeGrassMapsBuilder::FComponentState::IsBeyondEvictionRange(const TArray<FVector>& Cameras) const
 {
 	check(Stage == EComponentStage::GrassMapsPopulated);
-	FBoxSphereBounds WorldBounds = Component->CalcBounds(Component->GetComponentTransform());
-	float MinSqrDistanceToComponent = UE::Landscape::CalculateMinDistanceToCameras(Cameras, WorldBounds);
-	float DiscardDistanceScale = GGrassMapGuardBandDiscardMultiplier * GGrassCullDistanceScale;
-	float MinEvictDistance = Component->GrassTypeSummary.MaxInstanceDiscardDistance * DiscardDistanceScale;
+	const FBoxSphereBounds WorldBounds = Component->CalcBounds(Component->GetComponentTransform());
+	const float MinSqrDistanceToComponent = UE::Landscape::CalculateMinDistanceToCameras(Cameras, WorldBounds);
+	const float DiscardDistanceScale = GGrassMapGuardBandDiscardMultiplier * GGrassCullDistanceScale;
+	const float MinEvictDistance = Component->GrassTypeSummary.MaxInstanceDiscardDistance * DiscardDistanceScale;
 	return (MinSqrDistanceToComponent > MinEvictDistance * MinEvictDistance);
 }
 
@@ -1182,35 +1244,47 @@ void FLandscapeGrassMapsBuilder::Build()
 {
 	if (World)
 	{
+		int32 ValidCount = 0;
+		int32 TotalCount = 0;
+
 		// iterate proxies, update grass types and create the list of components to build
 		TArray<TObjectPtr<ULandscapeComponent>> LandscapeComponentsToBuild;
 		for (TActorIterator<ALandscapeProxy> ProxyIt(World); ProxyIt; ++ProxyIt)
 		{
-			check(!ProxyIt->HasAnyFlags(RF_ClassDefaultObject));
+			ALandscapeProxy* Proxy = *ProxyIt;
+			check(!Proxy->HasAnyFlags(RF_ClassDefaultObject));
 
-			for (ULandscapeComponent* Component : ProxyIt->LandscapeComponents)
+			for (ULandscapeComponent* Component : Proxy->LandscapeComponents)
 			{
 				FComponentState* State = ComponentStates.FindRef(Component);
 
-				if (Component->UpdateGrassTypes() ||
-					State->Stage != EComponentStage::GrassMapsPopulated)	// TODO check up to date slow?
+				ValidCount += Component->GrassData->HasValidData() ? 1 : 0;
+				TotalCount ++;
+
+				bool bBuildThisComponent =
+					Component->UpdateGrassTypes() ||							// grass types changed
+					(State->Stage != EComponentStage::GrassMapsPopulated);		// not in a populated state yet
+
+				if (bBuildThisComponent)
 				{
-					
 					Component->MarkPackageDirty();
 					LandscapeComponentsToBuild.Add(Component);
 				}
 			}
 
-			ProxyIt->UpdateGrassTypeSummary();
+			Proxy->UpdateGrassTypeSummary();
 		}
 
+		UE_LOG(LogGrass, Verbose, TEXT("FLandscapeGrassMapsBuilder::Build() building %d component grass maps (%d / %d valid)"),
+			LandscapeComponentsToBuild.Num(), ValidCount, TotalCount);
+
 		// build the grass maps
 		if (LandscapeComponentsToBuild.Num() > 0)
 		{
 			FScopedSlowTask SlowTask(static_cast<float>(LandscapeComponentsToBuild.Num()), (LOCTEXT("GrassMaps_BuildGrassMaps", "Building Grass maps")));
 			SlowTask.MakeDialog();
 
-			bool bMarkDirty = true;
+			constexpr bool bMarkDirty = true;
 			BuildGrassMapsNowForComponents(LandscapeComponentsToBuild, &SlowTask, bMarkDirty);
 		}
 	}
@@ -1240,7 +1314,7 @@ int32 FLandscapeGrassMapsBuilder::GetOutdatedGrassMapCount(bool bInForceUpdate)
 		bool bUpdate = bInForceUpdate || GLandscapeEditModeActive;
 		if (!bUpdate)
 		{
-			double GrassMapsTimeNow = FPlatformTime::Seconds();
+			const double GrassMapsTimeNow = FPlatformTime::Seconds();
 			// Recheck every 20 secs to handle the case where levels may have been Streamed in/out
 			if ((GrassMapsTimeNow - GrassMapsLastCheckTime) > 20)
 			{
diff --git a/Engine/Source/Runtime/Landscape/Private/LandscapeGrassMapsBuilder.h b/Engine/Source/Runtime/Landscape/Private/LandscapeGrassMapsBuilder.h
--- a/Engine/Source/Runtime/Landscape/Private/LandscapeGrassMapsBuilder.h
+++ b/Engine/Source/Runtime/Landscape/Private/LandscapeGrassMapsBuilder.h
@@ -2,6 +2,7 @@
 #pragma once
 
 #include "CoreMinimal.h"
+#include "Async/AsyncWork.h"
 #include "UObject/WeakObjectPtrTemplates.h"
 #include "LandscapeComponent.h"
 #include "LandscapeTextureStreamingManager.h"
@@ -10,6 +11,26 @@
 class FLandscapeGrassWeightExporter;
 struct FScopedSlowTask;
 
+class FAsyncFetchTask : public FNonAbandonableTask
+{
+public:
+	// non-owned pointer. The lifetime is managed externally and must guarantee this pointer is valid while the task is in flight
+	TNonNullPtr< FLandscapeGrassWeightExporter> ActiveRender;
+	TMap<ULandscapeComponent*, TUniquePtr<FLandscapeComponentGrassData>, TInlineSetAllocator<1>> Results;
+
+	FAsyncFetchTask(FLandscapeGrassWeightExporter* ActiveRender)
+		: ActiveRender(ActiveRender)
+	{
+	}
+
+	void DoWork();
+
+	FORCEINLINE TStatId GetStatId() const
+	{
+		RETURN_QUICK_DECLARE_CYCLE_STAT(FAsyncFetchTask, STATGROUP_ThreadPoolAsyncTasks);
+	}
+};
+
 /**
  * Helper class used to Build or monitor outdated Grass maps of a world
  */
@@ -33,7 +54,7 @@ public:
 	// called when components are registered to the world
 	void RegisterComponent(ULandscapeComponent* Component);
 	// called when components are unregistered from the world.
-	void UnregisterComponent(ULandscapeComponent* Component);
+	void UnregisterComponent(const ULandscapeComponent* Component);
 
 	// get the number of grass maps that are still waiting to render, as of the last AmortizedUpdateGrassMaps()
 	int32 GetTotalGrassMapsWaitingToRender() const { return TotalComponentsWaitingCount; }
@@ -71,6 +92,7 @@ private:
 		NotReady,							// tried to start generation process, but either no grass types exist or the material is not ready. wait for that to change.
 		TextureStreaming,					// texture streaming was requested.  wait for the mips to be available
 		Rendering,							// GPU render commands were sent -- waiting for async readback to complete
+		AsyncFetch,							// Waiting for the async fetch task to complete
 		GrassMapsPopulated,					// grass maps are built and are ready to create instances
 	};
 
@@ -96,7 +118,10 @@ private:
 		TArray<UTexture*> TexturesToStream;
 		
 		// the active render (valid only in Rendering stage)
-		FLandscapeGrassWeightExporter* ActiveRender = nullptr;
+		TUniquePtr<FLandscapeGrassWeightExporter> ActiveRender;
+
+		// when in AsyncFetch stage, this is async task that we are waiting for
+		TUniquePtr<FAsyncTask<FAsyncFetchTask>> AsyncFetchTask;
 
 		FComponentState(ULandscapeComponent* Component);
 
@@ -188,6 +213,7 @@ private:
 	int32 NotReadyCount = 0;
 	int32 StreamingCount = 0;
 	int32 RenderingCount = 0;
+	int32 AsyncFetchCount = 0;
 	int32 PopulatedCount = 0;
 
 	// number of components that need to render but are waiting (as of the last call to StartTrackingComponents())
@@ -198,6 +224,7 @@ private:
 	// store the grass map state of each registered (or recently unregistered) component
 	TMap<ULandscapeComponent*, FComponentState*> ComponentStates;
 
+	// Pending components, in a min heap by distance to streaming cameras
 	TArray<FPendingComponent> PendingComponentsHeap;
 	int32 PendingUpdateAmortizationCounter = 0;
 
@@ -217,7 +244,9 @@ private:
 	void KickOffRenderAndReadback(FComponentState& State);
 
 	// once the GPU readback is complete, this starts processing the data to generate a GrassData structure
-	void ProcessGrassDataAsync(FComponentState& State);
+	void LaunchAsyncFetchTask(FComponentState& State);
+
+	void PopulateGrassDataFromAsyncFetchTask(FComponentState& State);
 
 	// once the GPU readback is complete, this populates the grass data on the component (and cancels texture stream requests)
 	void PopulateGrassDataFromReadback(FComponentState& State);
diff --git a/Engine/Source/Runtime/Landscape/Private/LandscapeGrassWeightExporter.cpp b/Engine/Source/Runtime/Landscape/Private/LandscapeGrassWeightExporter.cpp
--- a/Engine/Source/Runtime/Landscape/Private/LandscapeGrassWeightExporter.cpp
+++ b/Engine/Source/Runtime/Landscape/Private/LandscapeGrassWeightExporter.cpp
@@ -14,6 +14,7 @@
 #include "TextureResource.h"
 #include "RenderCaptureInterface.h"
 #include "ShaderPlatformCachedIniValue.h"
+#include "LandscapePrivate.h"
 
 #if UE_BUILD_DEBUG
 #include "Misc/FileHelper.h"
@@ -29,13 +30,14 @@ static FAutoConsoleVariableRef CVarRenderCaptureNextGrassmapDraws(
 	TEXT("Trigger render captures during the next N grassmap draw calls."));
 
 extern int32 GGrassMapAlwaysBuildRuntimeGenerationResources;
+extern int32 GGrassMapUseRuntimeGeneration;
 
 BEGIN_SHADER_PARAMETER_STRUCT(FLandscapeGrassPassParameters, )
 	SHADER_PARAMETER_STRUCT_REF(FViewUniformShaderParameters, View)
 	SHADER_PARAMETER_RDG_UNIFORM_BUFFER(FSceneUniformParameters, Scene)
 	SHADER_PARAMETER_STRUCT_INCLUDE(FInstanceCullingDrawParams, InstanceCullingDrawParams)
 	RENDER_TARGET_BINDING_SLOTS()
-	END_SHADER_PARAMETER_STRUCT()
+END_SHADER_PARAMETER_STRUCT()
 
 class FLandscapeGrassWeightShaderElementData : public FMeshMaterialShaderElementData
 {
@@ -46,25 +48,35 @@ public:
 
 static bool ShouldCacheLandscapeGrassShaders(const FMeshMaterialShaderPermutationParameters& Parameters)
 {
-	static FShaderPlatformCachedIniValue<int32> GrassMapsUseRuntimeGenerationPerPlatform(TEXT("grass.GrassMap.UseRuntimeGeneration"));
-
-	bool bIsEditorPlatform = 
+	const bool bIsEditorPlatform = 
 		IsFeatureLevelSupported(Parameters.Platform, ERHIFeatureLevel::SM5) &&
 		EnumHasAllFlags(Parameters.Flags, EShaderPermutationFlags::HasEditorOnlyData);
 
-	bool bShouldBuildForPlatform = 
+#if WITH_EDITOR
+	static FShaderPlatformCachedIniValue<int32> GrassMapsUseRuntimeGenerationPerPlatform(TEXT("grass.GrassMap.UseRuntimeGeneration"));
+	const bool bPlatformUsesRuntimeGen = (GrassMapsUseRuntimeGenerationPerPlatform.Get(Parameters.Platform) != 0);
+#else
+	const bool bPlatformUsesRuntimeGen = (GGrassMapUseRuntimeGeneration != 0);
+#endif // WITH_EDITOR
+
+	const bool bShouldBuildForPlatform = 
 		bIsEditorPlatform ||
 		GGrassMapAlwaysBuildRuntimeGenerationResources ||
-		(GrassMapsUseRuntimeGenerationPerPlatform.Get(Parameters.Platform) != 0);
+		bPlatformUsesRuntimeGen;
 
-	bool bIsFixedGridVertexFactory =
+	const bool bIsFixedGridVertexFactory =
 		Parameters.VertexFactoryType == FindVertexFactoryType(FName(TEXT("FLandscapeFixedGridVertexFactory"), FNAME_Find));
 
 	// We only need grass weight shaders for Landscape fixed grid vertex factories
 	// And only for platforms that have runtime generation enabled or are editor platforms (or if we are always building resources)
-	return  (Parameters.MaterialParameters.bIsUsedWithLandscape || Parameters.MaterialParameters.bIsSpecialEngineMaterial) &&
-			bIsFixedGridVertexFactory &&
-			bShouldBuildForPlatform;
+	const bool bIsLandscapeRelated = (Parameters.MaterialParameters.bIsUsedWithLandscape || Parameters.MaterialParameters.bIsSpecialEngineMaterial);
+
+	const bool bShouldCache =
+		bIsLandscapeRelated &&
+		bIsFixedGridVertexFactory &&
+		bShouldBuildForPlatform;
+
+	return bShouldCache;
 }
 
 class FLandscapeGrassWeightVS : public FMeshMaterialShader
@@ -341,7 +353,13 @@ void FLandscapeGrassWeightExporter_RenderThread::RenderLandscapeComponentToTextu
 
 	const FSceneView* View = ViewFamily.Views[0];
 
-	FRDGTextureRef OutputTexture = GraphBuilder.RegisterExternalTexture(CreateRenderTarget(RenderTargetResource->GetTextureRHI(), TEXT("LandscapeGrass")));
+	FRDGTextureDesc TextureDesc = FRDGTextureDesc::Create2D(
+		TargetSize,
+		PF_B8G8R8A8,
+		FClearValueBinding(),
+		ETextureCreateFlags::RenderTargetable);
+
+	FRDGTextureRef OutputTexture = GraphBuilder.CreateTexture(TextureDesc, TEXT("LandscapeGrassMapRenderTarget"), ERDGTextureFlags::None);
 
 	auto* PassParameters = GraphBuilder.AllocParameters<FLandscapeGrassPassParameters>();
 	PassParameters->View = View->ViewUniformBuffer;
@@ -380,13 +398,12 @@ void FLandscapeGrassWeightExporter_RenderThread::RenderLandscapeComponentToTextu
 }
 
 
-FLandscapeGrassWeightExporter::FLandscapeGrassWeightExporter(ALandscapeProxy* InLandscapeProxy, TArrayView<ULandscapeComponent* const> InLandscapeComponents, bool bInNeedsGrassmap, bool bInNeedsHeightmap, const TArray<int32>& InHeightMips, bool bInUseAsyncReadback)
-	: FLandscapeGrassWeightExporter_RenderThread(InHeightMips, bInUseAsyncReadback)
+FLandscapeGrassWeightExporter::FLandscapeGrassWeightExporter(ALandscapeProxy* InLandscapeProxy, TArrayView<ULandscapeComponent* const> InLandscapeComponents, bool bInNeedsGrassmap, bool bInNeedsHeightmap, const TArray<int32>& InHeightMips)
+	: FLandscapeGrassWeightExporter_RenderThread(InHeightMips)
 	, LandscapeProxy(InLandscapeProxy)
 	, ComponentSizeVerts(InLandscapeProxy->ComponentSizeQuads + 1)
 	, SubsectionSizeQuads(InLandscapeProxy->SubsectionSizeQuads)
 	, NumSubsections(InLandscapeProxy->NumSubsections)
-	, RenderTargetTexture(nullptr)
 {
 	check(InLandscapeComponents.Num() > 0);
 	SceneInterface = InLandscapeComponents[0]->GetScene();
@@ -442,14 +459,6 @@ FLandscapeGrassWeightExporter::FLandscapeGrassWeightExporter(ALandscapeProxy* In
 		0.5f / ZOffset,
 		ZOffset);
 
-	RenderTargetTexture = NewObject<UTextureRenderTarget2D>();
-	check(RenderTargetTexture);
-	RenderTargetTexture->ClearColor = FLinearColor::White;
-	RenderTargetTexture->TargetGamma = 1.0f;
-	const bool bForceLinearGamma = true;
-	RenderTargetTexture->InitCustomFormat(TargetSize.X, TargetSize.Y, PF_B8G8R8A8, bForceLinearGamma);
-	RenderTargetResource = RenderTargetTexture->GameThread_GetRenderTargetResource()->GetTextureRenderTarget2DResource();
-
 	UE::RenderCommandPipe::FSyncScope SyncScope;
 
 	RenderCaptureInterface::FScopedCapture RenderCapture((GRenderCaptureNextGrassmapDraws != 0), TEXT("LandscapeGrassmapCapture"));
@@ -492,23 +501,26 @@ struct FByteBuffer2DView : public IBuffer2DView<uint8>
 		}
 	}
 
-	bool IsAllZero() const
+	// copy elements from buffer to Dest, in X then Y order, return true if the copied data is all zero
+	virtual bool CopyToAndCalcIsAllZero(uint8* Dest, int32 SizeInBytes) const override
 	{
-		for (int Y = 0; Y < NumY; Y++)
+		uint8 MaxBits = 0;
+		for (int Y = 0; SizeInBytes > 0 && Y < NumY; Y++)
 		{
-			uint8 MaxBits = 0;
 			uint8* Src = BufferStart + Y * ByteStrideY;
-			for (int X = 0; X < NumX; X++)
+			int32 CopyCountX = FMath::Min(SizeInBytes, NumX);
+			// we can't use memcpy because of the ByteStride
+			while (CopyCountX--)
 			{
-				MaxBits = MaxBits | *Src;
+				uint8 Value = *Src;
+				MaxBits = MaxBits | Value;
+				*Dest = Value;
+				Dest++;
 				Src += ByteStrideX;
 			}
-			if (MaxBits != 0)
-			{
-				return false;
-			}
+			SizeInBytes -= NumX;
 		}
-		return true;
+		return (MaxBits == 0);
 	}
 
 	virtual int32 Num() const override { return NumX * NumY; }
@@ -551,33 +563,42 @@ struct FHeightBuffer2DView : IBuffer2DView<uint16>
 		}
 	}
 
+	// copy elements from buffer to Dest, in X then Y order
+	virtual bool CopyToAndCalcIsAllZero(uint16* Dest, int32 Count) const override
+	{
+		unimplemented()
+		return true;
+	}
+
 	virtual int32 Num() const override { return NumX * NumY; }
 };
 
-TMap<ULandscapeComponent*, TUniquePtr<FLandscapeComponentGrassData>, TInlineSetAllocator<1>> FLandscapeGrassWeightExporter::FetchResults()
+
+void FLandscapeGrassWeightExporter::FreeAsyncReadback()
+{
+	check(AsyncReadbackPtr != nullptr);
+	AsyncReadbackPtr->QueueDeletionFromGameThread();
+	AsyncReadbackPtr = nullptr;
+}
+
+
+TMap<ULandscapeComponent*, TUniquePtr<FLandscapeComponentGrassData>, TInlineSetAllocator<1>> FLandscapeGrassWeightExporter::FetchResults(bool bFreeAsyncReadback)
 {
 	TRACE_CPUPROFILER_EVENT_SCOPE(FetchResults);
 	TMap<ULandscapeComponent*, TUniquePtr<FLandscapeComponentGrassData>, TInlineSetAllocator<1>> Results;
 	TArray<FColor> Samples;
 
-	if (AsyncReadbackPtr != nullptr)
+	check(AsyncReadbackPtr != nullptr);
+
 	{
 		FIntPoint Size;
 		Samples = AsyncReadbackPtr->TakeResults(&Size);
-		AsyncReadbackPtr->QueueDeletionFromGameThread();
-		AsyncReadbackPtr = nullptr;
+		if (bFreeAsyncReadback)
+		{
+			FreeAsyncReadback();
+		}
 		check(Size == TargetSize);
 	}
-	else
-	{
-		Samples.SetNumUninitialized(TargetSize.X * TargetSize.Y);
-
-		// Copy the contents of the remote texture to system memory (SamplesDataBuffer)
-		// This is a synchronous operation that may stall the cpu, waiting for the GPU data to be generated and copied
-		FReadSurfaceDataFlags ReadSurfaceDataFlags;
-		ReadSurfaceDataFlags.SetLinearToGamma(false);
-		RenderTargetResource->ReadPixels(Samples, ReadSurfaceDataFlags, FIntRect(0, 0, TargetSize.X, TargetSize.Y));
-	}
 
 	Results.Reserve(ComponentInfos.Num());
 	FHeightBuffer2DView HeightData;
@@ -696,21 +717,7 @@ TMap<ULandscapeComponent*, TUniquePtr<FLandscapeComponentGrassData>, TInlineSetA
 		#undef BGRA_AS_FCOLOR_RED
 		#undef BGRA_AS_FCOLOR_ALPHA
 		
-		{
-			TRACE_CPUPROFILER_EVENT_SCOPE(RemoveZeroWeight);
-			for (auto Iter(WeightData.CreateIterator()); Iter; ++Iter)
-			{
-				// Remove null grass type if we had one (can occur if the node has null entries)
-				// Remove any grass data that is entirely weight 0
-				if (Iter->Key == nullptr ||
-					((FByteBuffer2DView*)Iter->Value)->IsAllZero())
-				{
-					Iter.RemoveCurrent();
-				}
-			}
-		}
-
-		NewGrassData->InitializeFrom(&HeightData, WeightData);
+		NewGrassData->InitializeFrom(&HeightData, WeightData, /* bStripEmptyWeights = */ true);
 		Results.Add(Component, MoveTemp(NewGrassData));
 	}
 
@@ -721,14 +728,20 @@ void FLandscapeGrassWeightExporter::ApplyResults()
 {
 	TRACE_CPUPROFILER_EVENT_SCOPE(ApplyResults);
 
-	TMap<ULandscapeComponent*, TUniquePtr<FLandscapeComponentGrassData>, TInlineSetAllocator<1>> NewGrassData = FetchResults();
+	TMap<ULandscapeComponent*, TUniquePtr<FLandscapeComponentGrassData>, TInlineSetAllocator<1>> NewGrassData = FetchResults(/* bFreeAsyncReadback = */ true);
+	ApplyResults(NewGrassData);
+}
 
-	for (auto&& GrassDataPair : NewGrassData)
+void FLandscapeGrassWeightExporter::ApplyResults(TMap<ULandscapeComponent*, TUniquePtr<FLandscapeComponentGrassData>, TInlineSetAllocator<1>>& Results)
+{
+	for (auto&& GrassDataPair : Results)
 	{
 		ULandscapeComponent* Component = GrassDataPair.Key;
 		FLandscapeComponentGrassData* ComponentGrassData = GrassDataPair.Value.Release();
 		ALandscapeProxy* Proxy = Component->GetLandscapeProxy();
 
+		UE_LOG(LogGrass, Verbose, TEXT("Populating component %s with grass data, size: %d"), *Component->GetName(), ComponentGrassData->NumElements);
+
 		// Assign the new data (thread-safe)
 		Component->GrassData = MakeShareable(ComponentGrassData);
 
@@ -745,26 +758,3 @@ void FLandscapeGrassWeightExporter::ApplyResults()
 #endif // WITH_EDITOR
 	}
 }
-
-void FLandscapeGrassWeightExporter::AddReferencedObjects(UObject* InThis, FReferenceCollector& Collector)
-{
-	if (RenderTargetTexture)
-	{
-		Collector.AddReferencedObject(RenderTargetTexture);
-	}
-
-	if (LandscapeProxy)
-	{
-		Collector.AddReferencedObject(LandscapeProxy);
-	}
-
-	for (auto& Info : ComponentInfos)
-	{
-		if (Info.Component)
-		{
-			Collector.AddReferencedObject(Info.Component);
-		}
-
-		Collector.AddReferencedObjects(Info.RequestedGrassTypes);
-	}
-}
diff --git a/Engine/Source/Runtime/Landscape/Private/LandscapeGrassWeightExporter.h b/Engine/Source/Runtime/Landscape/Private/LandscapeGrassWeightExporter.h
--- a/Engine/Source/Runtime/Landscape/Private/LandscapeGrassWeightExporter.h
+++ b/Engine/Source/Runtime/Landscape/Private/LandscapeGrassWeightExporter.h
@@ -14,13 +14,11 @@ class UTextureRenderTarget2D;
 // data also accessible by render thread
 class FLandscapeGrassWeightExporter_RenderThread
 {
-	FLandscapeGrassWeightExporter_RenderThread(const TArray<int32>& InHeightMips, bool bInUseAsyncReadback)
+	FLandscapeGrassWeightExporter_RenderThread(const TArray<int32>& InHeightMips)
 		: HeightMips(InHeightMips)
 	{
-		if (bInUseAsyncReadback)
-		{
-			AsyncReadbackPtr = new FLandscapeAsyncTextureReadback();
-		}
+		// even when doing a synchronous readback, we use the async readback structure
+		AsyncReadbackPtr = new FLandscapeAsyncTextureReadback();
 	}
 
 	friend class FLandscapeGrassWeightExporter;
@@ -67,7 +65,6 @@ public:
 	};
 
 	FSceneInterface* SceneInterface = nullptr;
-	FTextureRenderTarget2DResource* RenderTargetResource = nullptr;
 	TArray<FComponentInfo, TInlineAllocator<1>> ComponentInfos;
 	FIntPoint TargetSize;
 	TArray<int32> HeightMips;
@@ -88,10 +85,9 @@ class FLandscapeGrassWeightExporter : public FLandscapeGrassWeightExporter_Rende
 	int32 SubsectionSizeQuads;
 	int32 NumSubsections;
 	TArray<TObjectPtr<ULandscapeGrassType>> GrassTypes;
-	TObjectPtr<UTextureRenderTarget2D> RenderTargetTexture;
 
 public:
-	FLandscapeGrassWeightExporter(ALandscapeProxy* InLandscapeProxy, TArrayView<ULandscapeComponent* const> InLandscapeComponents, bool bInNeedsGrassmap = true, bool bInNeedsHeightmap = true, const TArray<int32>& InHeightMips = {}, bool bUseAsyncReadback = false);
+	FLandscapeGrassWeightExporter(ALandscapeProxy* InLandscapeProxy, TArrayView<ULandscapeComponent* const> InLandscapeComponents, bool bInNeedsGrassmap = true, bool bInNeedsHeightmap = true, const TArray<int32>& InHeightMips = {});
 
 	// If using the async readback path, check its status and update if needed. Return true when the AsyncReadbackResults are available.
 	// You must call this periodically, or the async readback may not complete.
@@ -110,13 +106,17 @@ public:
 
 	// Fetches the results from the GPU texture and translates them into FLandscapeComponentGrassDatas.
 	// If using async readback, requires AsyncReadback to be complete before calling this.
-	TMap<ULandscapeComponent*, TUniquePtr<FLandscapeComponentGrassData>, TInlineSetAllocator<1>> FetchResults();
+	// bFreeAsyncReadback if true will call FreeAsyncReadback() to free the readback resource (otherwise you must do it manually)
+	TMap<ULandscapeComponent*, TUniquePtr<FLandscapeComponentGrassData>, TInlineSetAllocator<1>> FetchResults(bool bFreeAsyncReadback);
+
+	void FreeAsyncReadback();
+
+	// Applies the results using pre-fetched data.
+	static void ApplyResults(TMap<ULandscapeComponent*, TUniquePtr<FLandscapeComponentGrassData>, TInlineSetAllocator<1>>& Results);
 
 	// Fetches the results and applies them to the landscape components
 	// If using async readback, requires AsyncReadback to be complete before calling this.
 	void ApplyResults();
-
-	void AddReferencedObjects(UObject* InThis, FReferenceCollector& Collector);
 };
 
 namespace UE::Landscape::Grass
diff --git a/Engine/Source/Runtime/Landscape/Private/LandscapeSubsystem.cpp b/Engine/Source/Runtime/Landscape/Private/LandscapeSubsystem.cpp
--- a/Engine/Source/Runtime/Landscape/Private/LandscapeSubsystem.cpp
+++ b/Engine/Source/Runtime/Landscape/Private/LandscapeSubsystem.cpp
@@ -363,7 +363,7 @@ void ULandscapeSubsystem::Tick(float DeltaTime)
 			
 			// Update the proxies proxy
 			{
-				if (Proxy->bUseRuntimeGrassMapGeneration)
+				if (!Proxy->GetDisableRuntimeGrassMapGeneration())
 				{
 					bAllProxiesRuntimeGrassMapsDisabled = false;
 				}
@@ -814,7 +814,7 @@ void ULandscapeSubsystem::DisplayMessages(FCanvas* Canvas, float& XPos, float& Y
 		if (ModifiedNotDirtyCount > 0)
 		{
 			SmallTextItem.SetColor(FLinearColor::Red);
-			SmallTextItem.Text = FText::Format(LOCTEXT("LANDSCAPE_NEED_TO_BE_SAVED", "LANDSCAPE: NEED TO BE SAVED ({0} {0}|plural(one=object,other=objects))"), ModifiedNotDirtyCount);
+			SmallTextItem.Text = FText::Format(LOCTEXT("LANDSCAPE_NEED_TO_BE_SAVED", "LANDSCAPE: NEED TO SAVE TO SHOW CHANGES IN A COOKED GAME ({0} {0}|plural(one=object,other=objects))"), ModifiedNotDirtyCount);
 			Canvas->DrawItem(SmallTextItem, FVector2D(XPos, YPos));
 			YPos += FontSizeY;
 		}
