Landscape: Export (for HLOD and Nanite)

No need to compact the FMeshDescription as now constructed without holes.

diff --git a/Engine/Source/Runtime/Landscape/Private/LandscapeEdit.cpp b/Engine/Source/Runtime/Landscape/Private/LandscapeEdit.cpp
--- a/Engine/Source/Runtime/Landscape/Private/LandscapeEdit.cpp
+++ b/Engine/Source/Runtime/Landscape/Private/LandscapeEdit.cpp
@@ -4287,10 +4287,6 @@ bool ALandscapeProxy::ExportToRawMeshDataCopy(const FRawMeshExportParams& InExpo
 		}
 	}
 
-	//Compact the MeshDescription, if there was visibility mask or some bounding box clip, it need to be compacted so the sparse array are from 0 to n with no invalid data in between. 
-	TRACE_CPUPROFILER_EVENT_SCOPE(ALandscapeProxy::ExportToRawMesh - Compact);
-	FElementIDRemappings ElementIDRemappings;
-	OutRawMesh.Compact(ElementIDRemappings);
 	return OutRawMesh.Polygons().Num() > 0;
 }
 
