Fix for asserts when LandscapeGuid is invalid.
The calling location should handle the null return value -- in this specific bug it already does.

diff --git a/Engine/Source/Runtime/Landscape/Private/Landscape.cpp b/Engine/Source/Runtime/Landscape/Private/Landscape.cpp
--- a/Engine/Source/Runtime/Landscape/Private/Landscape.cpp
+++ b/Engine/Source/Runtime/Landscape/Private/Landscape.cpp
@@ -5424,9 +5424,7 @@ void ULandscapeInfo::RecreateLandscapeInfo(UWorld* InWorld, bool bMapCheck, bool
 ULandscapeInfo* ULandscapeInfo::Find(UWorld* InWorld, const FGuid& LandscapeGuid)
 {
 	ULandscapeInfo* LandscapeInfo = nullptr;
-
-	check(LandscapeGuid.IsValid());
-	if (InWorld != nullptr)
+	if (InWorld != nullptr && LandscapeGuid.IsValid())
 	{
 		auto& LandscapeInfoMap = ULandscapeInfoMap::GetLandscapeInfoMap(InWorld);
 		LandscapeInfo = LandscapeInfoMap.Map.FindRef(LandscapeGuid);
