Added profile tags to landscape code called several times in editor

diff --git a/Engine/Source/Runtime/Landscape/Private/LandscapeGrassMapsBuilder.cpp b/Engine/Source/Runtime/Landscape/Private/LandscapeGrassMapsBuilder.cpp
--- a/Engine/Source/Runtime/Landscape/Private/LandscapeGrassMapsBuilder.cpp
+++ b/Engine/Source/Runtime/Landscape/Private/LandscapeGrassMapsBuilder.cpp
@@ -340,7 +340,7 @@ FLandscapeGrassMapsBuilder::FComponentState::FComponentState(ULandscapeComponent
 
 void FLandscapeGrassMapsBuilder::FPendingComponent::UpdatePriorityDistance(const TArray<FVector>&Cameras)
 {
-	TRACE_CPUPROFILER_EVENT_SCOPE(UpdatePriorityDistance);
+	TRACE_CPUPROFILER_EVENT_SCOPE(FLandscapeGrassMapsBuilder::FPendingComponent::UpdatePriorityDistance);
 	ULandscapeComponent* Component = State->Component;
 	FBoxSphereBounds WorldBounds = Component->CalcBounds(Component->GetComponentTransform());
 	double MinSqrDistanceToComponent = UE::Landscape::CalculateMinDistanceToCameras(Cameras, WorldBounds);
@@ -356,6 +356,7 @@ void FAsyncFetchTask::DoWork()
 
 bool FLandscapeGrassMapsBuilder::UpdateTrackedComponents(const TArray<FVector>& Cameras, int32 LocalMaxRendering, int32 MaxExpensiveUpdateChecksToPerform, bool bCancelAndEvictAll)
 {
+	TRACE_CPUPROFILER_EVENT_SCOPE(FLandscapeGrassMapsBuilder::UpdateTrackedComponents);
 	SCOPE_CYCLE_COUNTER(STAT_UpdateComponentGrassMaps);
 
 	bool bChanged = false;
@@ -721,7 +722,7 @@ void FLandscapeGrassMapsBuilder::AmortizedUpdateGrassMaps(
 bool FLandscapeGrassMapsBuilder::BuildGrassMapsNowForComponents(
 	TArrayView<TObjectPtr<ULandscapeComponent>> LandscapeComponents, FScopedSlowTask* SlowTask, bool bMarkDirty)
 {
-	TRACE_CPUPROFILER_EVENT_SCOPE(BuildGrassMapsNowForComponents);
+	TRACE_CPUPROFILER_EVENT_SCOPE(FLandscapeGrassMapsBuilder::BuildGrassMapsNowForComponents);
 	const int32 MaxStreamingRendering = FMath::Max(GGrassMapMaxComponentsForBlockingUpdate, 1);
 
 	if (LandscapeComponents.IsEmpty())
@@ -830,7 +831,7 @@ bool FLandscapeGrassMapsBuilder::BuildGrassMapsNowForComponents(
 
 bool FLandscapeGrassMapsBuilder::CancelAndEvict(FComponentState& State)
 {
-	TRACE_CPUPROFILER_EVENT_SCOPE(CancelAndEvict);
+	TRACE_CPUPROFILER_EVENT_SCOPE(FLandscapeGrassMapsBuilder::CancelAndEvict);
 
 	// handle transitioning from any stage to EComponentState::Pending
 	switch (State.Stage)
@@ -944,7 +945,7 @@ bool FLandscapeGrassMapsBuilder::CancelAndEvict(FComponentState& State)
 
 bool FLandscapeGrassMapsBuilder::StartGrassMapGeneration(FComponentState& State, bool bForceCompileShaders)
 {
-	TRACE_CPUPROFILER_EVENT_SCOPE(StartGrassMapGeneration);
+	TRACE_CPUPROFILER_EVENT_SCOPE(FLandscapeGrassMapsBuilder::StartGrassMapGeneration);
 	check(State.Stage == EComponentStage::Pending);
 
 	ULandscapeComponent* Component = State.Component;
@@ -1058,7 +1059,7 @@ void FLandscapeGrassMapsBuilder::PendingToPopulatedFastPathNoGrass(FComponentSta
 
 void FLandscapeGrassMapsBuilder::PendingToStreaming(FComponentState& State)
 {
-	TRACE_CPUPROFILER_EVENT_SCOPE(PendingToStreaming);
+	TRACE_CPUPROFILER_EVENT_SCOPE(FLandscapeGrassMapsBuilder::PendingToStreaming);
 	ULandscapeComponent* Component = State.Component;
 
 	// determine which textures we need to stream by inspecting the material, kick off streaming requests for them
diff --git a/Engine/Source/Runtime/Landscape/Private/LandscapeSubsystem.cpp b/Engine/Source/Runtime/Landscape/Private/LandscapeSubsystem.cpp
--- a/Engine/Source/Runtime/Landscape/Private/LandscapeSubsystem.cpp
+++ b/Engine/Source/Runtime/Landscape/Private/LandscapeSubsystem.cpp
@@ -207,7 +207,7 @@ void ULandscapeSubsystem::UnregisterComponent(ULandscapeComponent* Component)
 
 void ULandscapeSubsystem::RemoveGrassInstances(const TSet<ULandscapeComponent*>* ComponentsToRemoveGrassInstances)
 {
-	TRACE_CPUPROFILER_EVENT_SCOPE(RemoveGrassInstances);
+	TRACE_CPUPROFILER_EVENT_SCOPE(ULandscapeSubsystem::RemoveGrassInstances);
 	for (TWeakObjectPtr<ALandscapeProxy> ProxyPtr : Proxies)
 	{
 		if (ALandscapeProxy* Proxy = ProxyPtr.Get())
@@ -775,6 +775,8 @@ ALandscapeProxy* ULandscapeSubsystem::FindOrAddLandscapeProxy(ULandscapeInfo* La
 
 void ULandscapeSubsystem::DisplayMessages(FCanvas* Canvas, float& XPos, float& YPos)
 {
+	TRACE_CPUPROFILER_EVENT_SCOPE(ULandscapeSubsystem::DisplayMessages);
+
 	const int32 FontSizeY = 20;
 	FCanvasTextItem SmallTextItem(FVector2D(0, 0), FText::GetEmpty(), GEngine->GetSmallFont(), FLinearColor::White);
 	SmallTextItem.EnableShadow(FLinearColor::Black);
