WorldPartition: Add UActorDescContainerInstance/FWorldPartitionActorDescInstance which take transient/context data from UActorDescContainer/FWorldPartitionActorDesc leaving them as pure on disk repres

diff --git a/Engine/Source/Runtime/Engine/Public/WorldPartition/WorldPartitionSubsystem.h b/Engine/Source/Runtime/Engine/Public/WorldPartition/WorldPartitionSubsystem.h
--- a/Engine/Source/Runtime/Engine/Public/WorldPartition/WorldPartitionSubsystem.h
+++ b/Engine/Source/Runtime/Engine/Public/WorldPartition/WorldPartitionSubsystem.h
@@ -6,6 +6,7 @@
 #include "Subsystems/WorldSubsystem.h"
 #include "WorldPartition/Filter/WorldPartitionActorFilter.h"
 #include "WorldPartition/WorldPartitionActorContainerID.h"
+#include "WorldPartition/WorldPartitionHandle.h"
 #include "Containers/Map.h"
 #include "Containers/Set.h"
 #include "Misc/Guid.h"
@@ -56,12 +57,6 @@ class UWorldPartitionSubsystem : public UTickableWorldSubsystem
 public:
 	ENGINE_API UWorldPartitionSubsystem();
 
-	//~ Begin UObject Interface
-#if WITH_EDITOR
-	static ENGINE_API void AddReferencedObjects(UObject* InThis, FReferenceCollector& Collector);
-#endif
-	//~ End UObject Interface
-
 	//~ Begin USubsystem Interface.
 	ENGINE_API virtual void Initialize(FSubsystemCollectionBase& Collection) override;
 	ENGINE_API virtual void Deinitialize() override;
@@ -113,45 +108,16 @@ public:
 
 	static ENGINE_API bool IsRunningConvertWorldPartitionCommandlet();
 
-	UActorDescContainer* RegisterContainer(FName PackageName) { return ActorDescContainerInstanceManager.RegisterContainer(PackageName, GetWorld()); }
-	void UnregisterContainer(UActorDescContainer* Container) { ActorDescContainerInstanceManager.UnregisterContainer(Container); }
-	FBox GetContainerBounds(FName PackageName) const { return ActorDescContainerInstanceManager.GetContainerBounds(PackageName); }
-	void UpdateContainerBounds(FName PackageName) { ActorDescContainerInstanceManager.UpdateContainerBounds(PackageName); }
-
-	TSet<FWorldPartitionActorDesc*> SelectedActorDescs;
-
-	class FActorDescContainerInstanceManager
-	{
-		friend class UWorldPartitionSubsystem;
-
-		struct FActorDescContainerInstance
-		{
-			FActorDescContainerInstance()
-				: Container(nullptr)
-				, RefCount(0)
-				, Bounds(ForceInit)
-			{}
-
-			void AddReferencedObjects(FReferenceCollector& Collector);
-			void UpdateBounds();
-
-			TObjectPtr<UActorDescContainer> Container;
-			uint32 RefCount;
-			FBox Bounds;
-		};
-
-		void AddReferencedObjects(FReferenceCollector& Collector);
-
-	public:
-		UActorDescContainer* RegisterContainer(FName PackageName, UWorld* InWorld);
-		void UnregisterContainer(UActorDescContainer* Container);
-
-		FBox GetContainerBounds(FName PackageName) const;
-		void UpdateContainerBounds(FName PackageName);
+	UE_DEPRECATED(5.4, "Use UActorDescContainerSubsystem instead")
+	UActorDescContainer* RegisterContainer(FName PackageName) { return nullptr; }
+	UE_DEPRECATED(5.4, "Use UActorDescContainerSubsystem instead")
+	void UnregisterContainer(UActorDescContainer* Container) { }
+	UE_DEPRECATED(5.4, "Use UActorDescContainerSubsystem instead")
+	FBox GetContainerBounds(FName PackageName) const { return FBox(); }
+	UE_DEPRECATED(5.4, "Use UActorDescContainerSubsystem instead")
+	void UpdateContainerBounds(FName PackageName) { }
 
-	private:
-		TMap<FName, FActorDescContainerInstance> ActorDescContainers;
-	};
+	TSet<FWorldPartitionHandle> SelectedActorHandles;
 private:
 	ENGINE_API FWorldPartitionActorFilter GetWorldPartitionActorFilterInternal(const FString& InWorldPackage, EWorldPartitionActorFilterType InFilterTypes, TSet<FString>& InOutVisitedPackages) const;
 #endif
@@ -229,6 +195,5 @@ private:
 
 #if WITH_EDITOR
 	bool bIsRunningConvertWorldPartitionCommandlet;
-	mutable FActorDescContainerInstanceManager ActorDescContainerInstanceManager;
 #endif
 };
diff --git a/Engine/Source/Runtime/Landscape/Private/Landscape.cpp b/Engine/Source/Runtime/Landscape/Private/Landscape.cpp
--- a/Engine/Source/Runtime/Landscape/Private/Landscape.cpp
+++ b/Engine/Source/Runtime/Landscape/Private/Landscape.cpp
@@ -15,6 +15,7 @@ Landscape.cpp: Terrain rendering
 #include "UObject/ConstructorHelpers.h"
 #include "UObject/DevObjectVersion.h"
 #include "UObject/LinkerLoad.h"
+#include "UObject/Package.h"
 #include "Framework/Application/SlateApplication.h"
 #include "LandscapePrivate.h"
 #include "LandscapeStreamingProxy.h"
@@ -94,6 +95,7 @@ Landscape.cpp: Terrain rendering
 #include "WorldPartition/WorldPartitionHelpers.h"
 #include "WorldPartition/WorldPartitionHandle.h"
 #include "WorldPartition/Landscape/LandscapeActorDesc.h"
+#include "WorldPartition/WorldPartitionActorDescInstance.h"
 #include "Engine/Texture2DArray.h"
 
 #if WITH_EDITOR
@@ -6141,10 +6143,10 @@ FBox ULandscapeInfo::GetCompleteBounds() const
 
 	FBox Bounds(EForceInit::ForceInit);
 
-	FWorldPartitionHelpers::ForEachActorDesc<ALandscapeProxy>(Landscape->GetWorld()->GetWorldPartition(), [this, &Bounds, Landscape](const FWorldPartitionActorDesc* ActorDesc)
+	FWorldPartitionHelpers::ForEachActorDescInstance<ALandscapeProxy>(Landscape->GetWorld()->GetWorldPartition(), [this, &Bounds, Landscape](const FWorldPartitionActorDescInstance* ActorDescInstance)
 	{
-		FLandscapeActorDesc* LandscapeActorDesc = (FLandscapeActorDesc*)ActorDesc;
-		ALandscapeProxy* LandscapeProxy = Cast<ALandscapeProxy>(ActorDesc->GetActor());
+		FLandscapeActorDesc* LandscapeActorDesc = (FLandscapeActorDesc*)ActorDescInstance->GetActorDesc();
+		ALandscapeProxy* LandscapeProxy = Cast<ALandscapeProxy>(ActorDescInstance->GetActor());
 
 		// Prioritize loaded bounds, as the bounds in the actor desc might not be up-to-date
 		if(LandscapeProxy && (LandscapeProxy->GetGridGuid() == LandscapeGuid))
@@ -6153,7 +6155,7 @@ FBox ULandscapeInfo::GetCompleteBounds() const
 		}
 		else if (LandscapeActorDesc->GridGuid == LandscapeGuid)
 		{
-			Bounds += ActorDesc->GetEditorBounds();
+			Bounds += LandscapeActorDesc->GetEditorBounds();
 		}
 
 		return true;
diff --git a/Engine/Source/Runtime/Landscape/Private/LandscapeEdit.cpp b/Engine/Source/Runtime/Landscape/Private/LandscapeEdit.cpp
--- a/Engine/Source/Runtime/Landscape/Private/LandscapeEdit.cpp
+++ b/Engine/Source/Runtime/Landscape/Private/LandscapeEdit.cpp
@@ -75,12 +75,12 @@ LandscapeEdit.cpp: Landscape editing
 #include "ScopedTransaction.h"
 #include "Editor.h"
 #include "PhysicalMaterials/PhysicalMaterial.h"
-#include "WorldPartition/ActorDescContainer.h"
 #include "WorldPartition/WorldPartition.h"
 #include "WorldPartition/WorldPartitionHandle.h"
 #include "WorldPartition/WorldPartitionHelpers.h"
 #include "WorldPartition/WorldPartitionActorDesc.h"
 #include "WorldPartition/Landscape/LandscapeActorDesc.h"
+#include "WorldPartition/WorldPartitionActorDescInstance.h"
 #include "ActorPartition/ActorPartitionSubsystem.h"
 #include "LandscapeUtils.h"
 #include "LandscapeSplineActor.h"
@@ -2675,16 +2675,16 @@ bool ULandscapeInfo::HasUnloadedComponentsInRegion(int32 X1, int32 Y1, int32 X2,
 
 		if (UWorldPartition* WorldPartition = World->GetWorldPartition())
 		{
-			FWorldPartitionHelpers::ForEachActorDesc<ALandscapeProxy>(WorldPartition, [this, World, &MinCoord, &MaxCoord, &bResult](const FWorldPartitionActorDesc* ActorDesc)
+			FWorldPartitionHelpers::ForEachActorDescInstance<ALandscapeProxy>(WorldPartition, [this, World, &MinCoord, &MaxCoord, &bResult](const FWorldPartitionActorDescInstance* ActorDescInstance)
 			{
-				FLandscapeActorDesc* LandscapeActorDesc = (FLandscapeActorDesc*)ActorDesc;
+				FLandscapeActorDesc* LandscapeActorDesc = (FLandscapeActorDesc*)ActorDescInstance->GetActorDesc();
 
 				if (LandscapeActorDesc->GridGuid == LandscapeGuid)
 				{
 					const UActorPartitionSubsystem::FCellCoord ActorCoord(LandscapeActorDesc->GridIndexX, LandscapeActorDesc->GridIndexY, LandscapeActorDesc->GridIndexZ, World->PersistentLevel);
 					if (ActorCoord.X >= MinCoord.X && ActorCoord.Y >= MinCoord.Y && ActorCoord.X <= MaxCoord.X && ActorCoord.Y <= MaxCoord.Y)
 					{
-						if (!LandscapeActorDesc->IsLoaded())
+						if (!ActorDescInstance->IsLoaded())
 						{
 							bResult = true;
 							return false;
@@ -4834,13 +4834,13 @@ void ALandscape::PostRegisterAllComponents()
 	{
 		if (UWorldPartition* WorldPartition = GetWorld()->GetWorldPartition())
 		{
-			FWorldPartitionHelpers::ForEachActorDesc<ALandscapeProxy>(WorldPartition, [this, WorldPartition](const FWorldPartitionActorDesc* ActorDesc)
+			FWorldPartitionHelpers::ForEachActorDescInstance<ALandscapeProxy>(WorldPartition, [this, WorldPartition](const FWorldPartitionActorDescInstance* ActorDescInstance)
 			{
-				FLandscapeActorDesc* LandscapeActorDesc = (FLandscapeActorDesc*)ActorDesc;
+				FLandscapeActorDesc* LandscapeActorDesc = (FLandscapeActorDesc*)ActorDescInstance->GetActorDesc();
 
 				if (LandscapeActorDesc->GridGuid == LandscapeGuid)
 				{
-					ActorDescReferences.Add(FWorldPartitionReference(WorldPartition, ActorDesc->GetGuid()));
+					ActorDescReferences.Add(FWorldPartitionReference(WorldPartition, ActorDescInstance->GetGuid()));
 				}
 				return true;
 			});
@@ -6056,15 +6056,15 @@ void ALandscapeStreamingProxy::PostRegisterAllComponents()
 				Bounds.Min.Z = -HALF_WORLD_MAX;
 				Bounds.Max.Z = HALF_WORLD_MAX;
 
-				FWorldPartitionHelpers::ForEachIntersectingActorDesc(WorldPartition, Bounds, [this, WorldPartition](const FWorldPartitionActorDesc* ActorDesc) mutable
+				FWorldPartitionHelpers::ForEachIntersectingActorDescInstance(WorldPartition, Bounds, [this, WorldPartition](const FWorldPartitionActorDescInstance* ActorDescInstance) mutable
 				{
 					FName PropertyValue;
-					if (ActorDesc->GetProperty(ALandscape::AffectsLandscapeActorDescProperty, &PropertyValue))
+					if (ActorDescInstance->GetProperty(ALandscape::AffectsLandscapeActorDescProperty, &PropertyValue))
 					{
 						// If no Guid specified then consider actor as affecting all landscapes
 						if(FGuid ParsedGuid; PropertyValue.IsNone() || (FGuid::Parse(PropertyValue.ToString(), ParsedGuid) && ParsedGuid == LandscapeGuid))
 						{
-							ActorDescReferences.Add(FWorldPartitionReference(WorldPartition, ActorDesc->GetGuid()));
+							ActorDescReferences.Add(FWorldPartitionReference(WorldPartition, ActorDescInstance->GetGuid()));
 						}
 					}
 					return true;
@@ -6162,13 +6162,13 @@ bool ULandscapeInfo::CanDeleteLandscape(FText& OutReason) const
 		UWorld* World = Actor->GetWorld();
 		if (UWorldPartition* WorldPartition = World->GetWorldPartition())
 		{
-			FWorldPartitionHelpers::ForEachActorDesc<ALandscapeProxy>(WorldPartition, [this, &UndeletedProxyCount](const FWorldPartitionActorDesc* ActorDesc)
+			FWorldPartitionHelpers::ForEachActorDescInstance<ALandscapeProxy>(WorldPartition, [this, &UndeletedProxyCount](const FWorldPartitionActorDescInstance* ActorDescInstance)
 			{
-				FLandscapeActorDesc* LandscapeActorDesc = (FLandscapeActorDesc*)ActorDesc;
+				FLandscapeActorDesc* LandscapeActorDesc = (FLandscapeActorDesc*)ActorDescInstance->GetActorDesc();
 
 				if (LandscapeActorDesc->GridGuid == LandscapeGuid)
 				{
-					ALandscapeProxy* LandscapeProxy = Cast<ALandscapeProxy>(ActorDesc->GetActor());
+					ALandscapeProxy* LandscapeProxy = Cast<ALandscapeProxy>(ActorDescInstance->GetActor());
 					if (LandscapeProxy != LandscapeActor)
 					{
 						// If LandscapeProxy is null then it is not loaded so not deleted.
@@ -6207,15 +6207,15 @@ bool ULandscapeInfo::CanDeleteLandscape(FText& OutReason) const
 		UWorld* World = Actor->GetWorld();
 		if (UWorldPartition* WorldPartition = World->GetWorldPartition())
 		{
-			FWorldPartitionHelpers::ForEachActorDesc<ALandscapeSplineActor>(WorldPartition, [this, &UndeletedSplineCount](const FWorldPartitionActorDesc* ActorDesc)
+			FWorldPartitionHelpers::ForEachActorDescInstance<ALandscapeSplineActor>(WorldPartition, [this, &UndeletedSplineCount](const FWorldPartitionActorDescInstance* ActorDescInstance)
 			{
 				FName AffectsLandscapeProperty;
-				if (ActorDesc->GetProperty(ALandscape::AffectsLandscapeActorDescProperty, &AffectsLandscapeProperty))
+				if (ActorDescInstance->GetProperty(ALandscape::AffectsLandscapeActorDescProperty, &AffectsLandscapeProperty))
 				{
 					FGuid ParsedLandscapeGuid;
 					if (FGuid::Parse(AffectsLandscapeProperty.ToString(), ParsedLandscapeGuid) && ParsedLandscapeGuid == LandscapeGuid)
 				{
-						ALandscapeSplineActor* SplineActor = Cast<ALandscapeSplineActor>(ActorDesc->GetActor());
+						ALandscapeSplineActor* SplineActor = Cast<ALandscapeSplineActor>(ActorDescInstance->GetActor());
 		
 					// If SplineActor is null then it is not loaded/deleted. If it's loaded then it needs to be pending kill.
 					if (!SplineActor)
