* Automated LandscapeInherited/LandscapeOverridable enforcement of properties on the proxies so that using the meta-data on the landscape proxy property is sufficient to get the required behavior
* Changed list of overridden properties to a list of FName (memory optimization)
* Made bBakeMaterialPositionOffsetIntoCollision LandscapeInherited since that's how it is used
* Various consolidations to landscape-related PostEditChangeProperty functions (use GET_MEMBER_NAME_CHECKED when possible, removed redundant clamping
* Removed useless recomputation of normals when changing the scale when edit layers are enabled
* Fixed LandscapeInherited/LandscapeOverride properties which were not detected as such because we were using Property->Name instead of MemberProperty->Name
* Solidified all landscape PostEditChangeProperty accordingly
* Tagged all remaining ALandscapeProxy properties as either LandscapeInherited or LandscapeOverride
* Fixed landscape proxy overridable/inherited property names that didn't always appear readonly (note, we use IsValueEnabled here but there seems to be a bug with it, which is tracked with UE-202147, that makes LightingChannels and PerLODMaterialOverrides look editable when overridable but not overridden)

#tests editor

diff --git a/Engine/Source/Editor/DetailCustomizations/Private/LandscapeProxyUIDetails.cpp b/Engine/Source/Editor/DetailCustomizations/Private/LandscapeProxyUIDetails.cpp
--- a/Engine/Source/Editor/DetailCustomizations/Private/LandscapeProxyUIDetails.cpp
+++ b/Engine/Source/Editor/DetailCustomizations/Private/LandscapeProxyUIDetails.cpp
@@ -461,13 +461,18 @@ void FLandscapeProxyUIDetails::CustomizeDetails( IDetailLayoutBuilder& DetailBui
 				if (PropertyHandle->IsValidHandle())
 				{
 					const FText TooltipText = NSLOCTEXT("Landscape", "OverriddenProperty", "Check this box to override the parent landscape's property.");
-					FString PropertyName = Property->GetName();
+					FName PropertyName = Property->GetFName();
 					IDetailPropertyRow* DetailRow = DetailBuilder.EditDefaultProperty(PropertyHandle);
 					TSharedPtr<SWidget> NameWidget = nullptr;
 					TSharedPtr<SWidget> ValueWidget = nullptr;
 
 					DetailRow->GetDefaultWidgets(NameWidget, ValueWidget);
 
+					TAttribute<bool> EnabledAttribute = TAttribute<bool>::Create([LandscapeStreamingProxy, PropertyName]() -> bool
+					{
+						return LandscapeStreamingProxy->IsSharedPropertyOverridden(PropertyName);
+					});
+
 					DetailRow->CustomWidget(/*bShowChildren = */true)
 					.NameContent()
 					[
@@ -525,17 +530,22 @@ void FLandscapeProxyUIDetails::CustomizeDetails( IDetailLayoutBuilder& DetailBui
 						+ SHorizontalBox::Slot()
 						.AutoWidth()
 						[
-							NameWidget->AsShared()
+							SNew(SBox)
+							.IsEnabled(EnabledAttribute)
+							[
+								NameWidget->AsShared()
+							]
 						]
 					]
 					.ValueContent()
 					[
 						SNew(SBox)
-						.IsEnabled_Lambda([LandscapeStreamingProxy, PropertyName]() { return LandscapeStreamingProxy->IsSharedPropertyOverridden(PropertyName); })
+						.IsEnabled(EnabledAttribute)
 						[
 							ValueWidget->AsShared()
 						]
-					];
+					]
+					.IsValueEnabled(EnabledAttribute);
 				}
 			}
 		}
diff --git a/Engine/Source/Runtime/Landscape/Classes/LandscapeProxy.h b/Engine/Source/Runtime/Landscape/Classes/LandscapeProxy.h
--- a/Engine/Source/Runtime/Landscape/Classes/LandscapeProxy.h
+++ b/Engine/Source/Runtime/Landscape/Classes/LandscapeProxy.h
@@ -323,7 +323,7 @@ struct UE_DEPRECATED(5.1, "FLandscapeProxyMaterialOverride is deprecated; please
 
 #if WITH_EDITOR
 // Tracks delayed updates triggered by landscape updates.
-// It's resposible for preventing expensive async operations while the user is still editing the landscape
+// It's responsible for preventing expensive async operations while the user is still editing the landscape
 struct FAsyncWorkMonitor 
 {
 public:
@@ -470,7 +470,7 @@ public:
 
 #if WITH_EDITORONLY_DATA
 	/** LOD level to use when exporting the landscape to obj or FBX */
-	UPROPERTY(EditAnywhere, Category=LOD, AdvancedDisplay)
+	UPROPERTY(EditAnywhere, Category=LOD, AdvancedDisplay, meta = (LandscapeOverridable))
 	int32 ExportLOD;
 
 	/** Display Order of the targets */
@@ -483,8 +483,8 @@ public:
 #endif
 
 	/** LOD level to use when running lightmass (increase to 1 or 2 for large landscapes to stop lightmass crashing) */
-	UPROPERTY(EditAnywhere, Category=Lighting)
-	int32 StaticLightingLOD;
+	UPROPERTY(EditAnywhere, Category=Lighting, meta = (LandscapeOverridable))
+	int32 StaticLightingLOD; 
 
 	/** Default physical material, used when no per-layer values physical materials */
 	UPROPERTY(EditAnywhere, Category=Landscape, meta = (LandscapeOverridable))
@@ -495,7 +495,7 @@ public:
 	 * 1.0 is the default, whereas a higher value increases the streamed-in resolution.
 	 * Value can be < 0 (from legcay content, or code changes)
 	 */
-	UPROPERTY(EditAnywhere, Category=Landscape)
+	UPROPERTY(EditAnywhere, Category=Landscape, meta = (LandscapeOverridable))
 	float StreamingDistanceMultiplier;
 
 	/** Combined material used to render the landscape */
@@ -637,7 +637,7 @@ public:
 	 *  Total resolution would be changed by StaticLightingResolution*StaticLightingResolution
 	 *	Automatically calculate proper value for removing seams
 	 */
-	UPROPERTY(EditAnywhere, Category=Lighting)
+	UPROPERTY(EditAnywhere, Category=Lighting, meta = (LandscapeOverridable))
 	float StaticLightingResolution;
 
 	/** Controls whether the primitive component should cast a shadow or not. */
@@ -653,7 +653,7 @@ public:
 	uint8 bCastStaticShadow : 1;
 
 	/** Control shadow invalidation behavior, in particular with respect to Virtual Shadow Maps and material effects like World Position Offset. */
-	UPROPERTY(EditAnywhere, BlueprintReadOnly, Category=Lighting, AdvancedDisplay, meta=(EditCondition="CastShadow"))
+	UPROPERTY(EditAnywhere, BlueprintReadOnly, Category=Lighting, AdvancedDisplay, meta=(EditCondition="CastShadow", LandscapeOverridable))
 	EShadowCacheInvalidationBehavior ShadowCacheInvalidationBehavior;
 
 	/** Whether the object should cast contact shadows. This flag is only used if CastShadow is true. */
@@ -686,7 +686,7 @@ public:
 	/** Whether to use the landscape material's vertical world position offset when calculating static lighting.
 		Note: Only z (vertical) offset is supported. XY offsets are ignored.
 		Does not work correctly with an XY offset map (mesh collision) */
-	UPROPERTY(EditAnywhere, AdvancedDisplay, Category=Lighting)
+	UPROPERTY(EditAnywhere, AdvancedDisplay, Category=Lighting, meta = (LandscapeInherited))
 	uint32 bUseMaterialPositionOffsetInStaticLighting:1;
 
 	/** If true, the Landscape will be rendered in the CustomDepth pass (usually used for outlines) */
@@ -729,13 +729,13 @@ public:
 	 * @see [Overlap Events](https://docs.unrealengine.com/InteractiveExperiences/Physics/Collision/Overview#overlapandgenerateoverlapevents)
 	 * @see UpdateOverlaps(), BeginComponentOverlap(), EndComponentOverlap()
 	 */
-	UPROPERTY(EditAnywhere, BlueprintReadOnly, Category = Collision)
+	UPROPERTY(EditAnywhere, BlueprintReadOnly, Category = Collision, meta = (LandscapeOverridable))
 	uint32 bGenerateOverlapEvents : 1;
 
 	/** Whether to bake the landscape material's vertical world position offset into the collision heightfield.
 		Note: Only z (vertical) offset is supported. XY offsets are ignored.
 		Does not work with an XY offset map (mesh collision) */
-	UPROPERTY(EditAnywhere, AdvancedDisplay, Category=Collision)
+	UPROPERTY(EditAnywhere, AdvancedDisplay, Category=Collision, meta = (LandscapeInherited))
 	uint32 bBakeMaterialPositionOffsetIntoCollision:1;
 
 #if WITH_EDITORONLY_DATA
@@ -786,27 +786,27 @@ public:
 	int32 NumSubsections;    // Number of subsections in X and Y axis
 
 	/** Hints navigation system whether this landscape will ever be navigated on. true by default, but make sure to set it to false for faraway, background landscapes */
-	UPROPERTY(EditAnywhere, Category=Navigation)
+	UPROPERTY(EditAnywhere, Category=Navigation, meta = (LandscapeOverridable))
 	uint32 bUsedForNavigation:1;
 
 	/** Set to true to prevent navmesh generation under the terrain geometry */
-	UPROPERTY(EditAnywhere, Category = Navigation)
+	UPROPERTY(EditAnywhere, Category = Navigation, meta = (LandscapeOverridable))
 	uint32 bFillCollisionUnderLandscapeForNavmesh:1;
 
-	UPROPERTY(EditAnywhere, Category = Navigation, AdvancedDisplay)
+	UPROPERTY(EditAnywhere, Category = Navigation, AdvancedDisplay, meta = (LandscapeOverridable))
 	ENavDataGatheringMode NavigationGeometryGatheringMode;
 
 	/** When set to true it will generate MaterialInstanceDynamic for each components, so material can be changed at runtime */
-	UPROPERTY(EditAnywhere, Category = Landscape)
+	UPROPERTY(EditAnywhere, Category = Landscape, meta = (LandscapeOverridable))
 	bool bUseDynamicMaterialInstance;
 
 #if WITH_EDITORONLY_DATA
-	UPROPERTY(EditAnywhere, Category=Landscape)
+	UPROPERTY(EditAnywhere, Category=Landscape, meta = (LandscapeOverridable))
 	int32 MaxPaintedLayersPerComponent; // 0 = disabled
 #endif // WITH_EDITORONLY_DATA
 
 	/** Flag whether or not this Landscape's surface can be used for culling hidden triangles **/
-	UPROPERTY(EditAnywhere, Category = HLOD)
+	UPROPERTY(EditAnywhere, Category = HLOD, meta = (LandscapeOverridable))
 	bool bUseLandscapeForCullingInvisibleHLODVertices;
 
 	/** Flag that tell if we have some layers content **/
@@ -1124,10 +1124,10 @@ public:
 	LANDSCAPE_API bool IsPropertyOverridable(const FProperty* InProperty) const;
 	
 	// Returns true if the shared property is overridden by the object.
-	virtual bool IsSharedPropertyOverridden(const FString& InPropertyName) const { return false; }
+	virtual bool IsSharedPropertyOverridden(const FName& InPropertyName) const { return false; }
 
 	// Modifies the override state of the property given as argument.
-	virtual void SetSharedPropertyOverride(const FString& InPropertyName, const bool bIsOverriden) { }
+	virtual void SetSharedPropertyOverride(const FName& InPropertyName, const bool bIsOverriden) { }
 #endif // WITH_EDITOR
 
 	// Get Landscape Material assigned to this Landscape
diff --git a/Engine/Source/Runtime/Landscape/Classes/LandscapeStreamingProxy.h b/Engine/Source/Runtime/Landscape/Classes/LandscapeStreamingProxy.h
--- a/Engine/Source/Runtime/Landscape/Classes/LandscapeStreamingProxy.h
+++ b/Engine/Source/Runtime/Landscape/Classes/LandscapeStreamingProxy.h
@@ -33,7 +33,7 @@ private:
 	TSoftObjectPtr<ALandscape> LandscapeActorRef;
 
 	UPROPERTY()
-	TSet<FString> OverriddenSharedProperties;
+	TSet<FName> OverriddenSharedProperties;
 
 public:
 	//~ Begin UObject Interface
@@ -63,8 +63,8 @@ public:
 
 #if WITH_EDITOR
 	//~ Begin ALandscapeProxy Interface
-	virtual bool IsSharedPropertyOverridden(const FString& InPropertyName) const override;
-	virtual void SetSharedPropertyOverride(const FString& InPropertyName, const bool bIsOverridden) override;
+	virtual bool IsSharedPropertyOverridden(const FName& InPropertyName) const override;
+	virtual void SetSharedPropertyOverride(const FName& InPropertyName, const bool bIsOverridden) override;
 	//~ End ALandscapeProxy Interface
 
 protected:
diff --git a/Engine/Source/Runtime/Landscape/Private/Landscape.cpp b/Engine/Source/Runtime/Landscape/Private/Landscape.cpp
--- a/Engine/Source/Runtime/Landscape/Private/Landscape.cpp
+++ b/Engine/Source/Runtime/Landscape/Private/Landscape.cpp
@@ -526,6 +526,7 @@ void ALandscapeProxy::InvalidateOrUpdateNaniteRepresentation(bool bInCheckConten
 
 FGuid ALandscapeProxy::GetNaniteContentId() const
 {
+	TRACE_CPUPROFILER_EVENT_SCOPE(ALandscapeProxy::GetNaniteContentId);
 	if (!IsNaniteEnabled())
 	{
 		return FGuid();
@@ -4570,7 +4571,7 @@ namespace UE::Landscape::Private
 		return true;
 	}
 
-	static bool CopyPostEditPropertyByName(const TWeakObjectPtr<ALandscapeProxy>& InLandscapeProxy, const TWeakObjectPtr<ALandscape>& InParentLandscape, const FString& InPropertyName)
+	static bool CopyPostEditPropertyByName(const TWeakObjectPtr<ALandscapeProxy>& InLandscapeProxy, const TWeakObjectPtr<ALandscape>& InParentLandscape, const FName& InPropertyName)
 	{
 		if (!InLandscapeProxy.IsValid() || !InParentLandscape.IsValid())
 		{
@@ -4584,7 +4585,7 @@ namespace UE::Landscape::Private
 			return false;
 		}
 
-		FProperty* PropertyToCopy = LandscapeProxyClass->FindPropertyByName(FName(InPropertyName));
+		FProperty* PropertyToCopy = LandscapeProxyClass->FindPropertyByName(InPropertyName);
 
 		if (PropertyToCopy == nullptr)
 		{
@@ -4676,7 +4677,7 @@ TArray<FName> ALandscapeProxy::SynchronizeSharedProperties(ALandscapeProxy* InLa
 		}
 
 		if ((IsPropertyInherited(Property) ||
-			(IsPropertyOverridable(Property) && !IsSharedPropertyOverridden(Property->GetName()))) &&
+			(IsPropertyOverridable(Property) && !IsSharedPropertyOverridden(Property->GetFName()))) &&
 			!Property->Identical_InContainer(this, InLandscape))
 		{
 			SynchronizedProperties.Emplace(Property->GetFName());
@@ -4793,14 +4794,14 @@ UMaterialInterface* ALandscapeStreamingProxy::GetLandscapeHoleMaterial() const
 
 #if WITH_EDITOR
 
-bool ALandscapeStreamingProxy::IsSharedPropertyOverridden(const FString& InPropertyName) const
+bool ALandscapeStreamingProxy::IsSharedPropertyOverridden(const FName& InPropertyName) const
 {
 	return OverriddenSharedProperties.Contains(InPropertyName);
 }
 
-void ALandscapeStreamingProxy::SetSharedPropertyOverride(const FString& InPropertyName, const bool bIsOverridden)
+void ALandscapeStreamingProxy::SetSharedPropertyOverride(const FName& InPropertyName, const bool bIsOverridden)
 {
-	check(IsSharedProperty(FName(InPropertyName)));
+	check(IsSharedProperty(InPropertyName));
 
 	Modify();
 
@@ -4816,7 +4817,7 @@ void ALandscapeStreamingProxy::SetSharedPropertyOverride(const FString& InProper
 		if (!ParentLandscape.IsValid())
 		{
 			UE_LOG(LogLandscape, Warning, TEXT("Unable to retrieve the parent landscape's shared property value (ALandscapeStreamingProxy: %s, Property: %s). The proper value will be fixedup when reloading this proxy."),
-				   *GetFullName(), *InPropertyName);
+				   *GetFullName(), *InPropertyName.ToString());
 		}
 		else
 		{
@@ -4831,9 +4832,9 @@ void ALandscapeStreamingProxy::FixupOverriddenSharedProperties()
 {
 	const UClass* StreamingProxyClass = StaticClass();
 
-	for (const FString& PropertyName : OverriddenSharedProperties)
+	for (const FName& PropertyName : OverriddenSharedProperties)
 	{
-		const FProperty* Property = StreamingProxyClass->FindPropertyByName(FName(PropertyName));
+		const FProperty* Property = StreamingProxyClass->FindPropertyByName(PropertyName);
 		checkf(Property != nullptr, TEXT("An overridden property is referenced but cannot be found. Please check this property hasn't been renamed or deprecated and/or provide the proper adapting mechanism."));
 	}
 }
@@ -4859,18 +4860,18 @@ void ALandscapeProxy::UpgradeSharedProperties(ALandscape* InParentLandscape)
 			SynchronizedProperties.Emplace(Property->GetFName());
 			UE::Landscape::Private::CopyProperty(Property, InParentLandscape, this);
 		}
-		else if (IsPropertyOverridable(Property) && !IsSharedPropertyOverridden(Property->GetName()) && !Property->Identical_InContainer(this, InParentLandscape))
+		else if (IsPropertyOverridable(Property) && !IsSharedPropertyOverridden(Property->GetFName()) && !Property->Identical_InContainer(this, InParentLandscape))
 		{
 			if (CVarSilenceSharedPropertyDeprecationFixup->GetBool())
 			{
-				SetSharedPropertyOverride(Property->GetName(), true);
+				SetSharedPropertyOverride(Property->GetFName(), true);
 			}
 			else
 			{
 				FFormatNamedArguments Arguments;
 				TWeakObjectPtr<ALandscapeProxy> LandscapeProxy = this;
 				TWeakObjectPtr<ALandscape> ParentLandscape = InParentLandscape;
-				const FString PropertyName = Property->GetName();
+				const FName PropertyName = Property->GetFName();
 
 				bOpenMapCheckWindow = true;
 
@@ -4878,7 +4879,7 @@ void ALandscapeProxy::UpgradeSharedProperties(ALandscape* InParentLandscape)
 				Arguments.Add(TEXT("Landscape"), FText::FromString(InParentLandscape->GetActorNameOrLabel()));
 				FMessageLog("MapCheck").Warning()
 					->AddToken(FUObjectToken::Create(this, FText::FromString(GetActorNameOrLabel())))
-					->AddToken(FTextToken::Create(FText::Format(LOCTEXT("MapCheck_Message_LandscapeProxy_UpgradeSharedProperties", "Contains a property ({0}) different from parent's landscape actor. Please select between "), FText::FromString(PropertyName))))
+					->AddToken(FTextToken::Create(FText::Format(LOCTEXT("MapCheck_Message_LandscapeProxy_UpgradeSharedProperties", "Contains a property ({0}) different from parent's landscape actor. Please select between "), FText::FromString(PropertyName.ToString()))))
 					->AddToken(FActionToken::Create(LOCTEXT("MapCheck_OverrideProperty", "Override property"), LOCTEXT("MapCheck_OverrideProperty_Desc", "Keeping the current value and marking the property as overriding the parent landscape's value."),
 						FOnActionTokenExecuted::CreateLambda([LandscapeProxy, PropertyName]()
 							{
diff --git a/Engine/Source/Runtime/Landscape/Private/LandscapeEdit.cpp b/Engine/Source/Runtime/Landscape/Private/LandscapeEdit.cpp
--- a/Engine/Source/Runtime/Landscape/Private/LandscapeEdit.cpp
+++ b/Engine/Source/Runtime/Landscape/Private/LandscapeEdit.cpp
@@ -1779,6 +1779,7 @@ void ULandscapeComponent::UpdateCollisionLayerData()
 
 uint32 ULandscapeComponent::CalculatePhysicalMaterialTaskHash() const
 {
+	TRACE_CPUPROFILER_EVENT_SCOPE(ULandscapeComponent::CalculatePhysicalMaterialTaskHash);
 	uint32 Hash = 0;
 	
 	// Take into account any material changes.
@@ -4948,66 +4949,56 @@ FLinearColor ULandscapeLayerInfoObject::GenerateLayerUsageDebugColor() const
 #if WITH_EDITOR
 void ULandscapeLayerInfoObject::PostEditChangeProperty(FPropertyChangedEvent& PropertyChangedEvent)
 {
-	static const FName NAME_PhysMaterial = GET_MEMBER_NAME_CHECKED(ULandscapeLayerInfoObject, PhysMaterial);
-	static const FName NAME_LayerUsageDebugColor = GET_MEMBER_NAME_CHECKED(ULandscapeLayerInfoObject, LayerUsageDebugColor);
-	static const FName NAME_MinimumCollisionRelevanceWeight = GET_MEMBER_NAME_CHECKED(ULandscapeLayerInfoObject, MinimumCollisionRelevanceWeight);
-	static const FName NAME_R = FName(TEXT("R"));
-	static const FName NAME_G = FName(TEXT("G"));
-	static const FName NAME_B = FName(TEXT("B"));
-	static const FName NAME_A = FName(TEXT("A"));
-
 	Super::PostEditChangeProperty(PropertyChangedEvent);
 
-	const FName PropertyName = PropertyChangedEvent.Property ? PropertyChangedEvent.Property->GetFName() : NAME_None;
+	const FName MemberPropertyName = PropertyChangedEvent.MemberProperty ? PropertyChangedEvent.MemberProperty->GetFName() : NAME_None;
 
-	if (GIsEditor)
+	if ((MemberPropertyName == GET_MEMBER_NAME_CHECKED(ULandscapeLayerInfoObject, PhysMaterial))
+		|| (MemberPropertyName == GET_MEMBER_NAME_CHECKED(ULandscapeLayerInfoObject, MinimumCollisionRelevanceWeight)))
 	{
-		if (PropertyName == NAME_PhysMaterial || PropertyName == NAME_MinimumCollisionRelevanceWeight)
+		for (TObjectIterator<ALandscapeProxy> It(/*AdditionalExclusionFlags = */RF_ClassDefaultObject, /*bIncludeDerivedClasses = */true, /*InInternalExclusionFlags = */EInternalObjectFlags::Garbage); It; ++It)
 		{
-			for (TObjectIterator<ALandscapeProxy> It(/*AdditionalExclusionFlags = */RF_ClassDefaultObject, /*bIncludeDerivedClasses = */true, /*InInternalExclusionFlags = */EInternalObjectFlags::Garbage); It; ++It)
+			ALandscapeProxy* Proxy = *It;
+			if (Proxy->GetWorld() && !Proxy->GetWorld()->IsPlayInEditor())
 			{
-				ALandscapeProxy* Proxy = *It;
-				if (Proxy->GetWorld() && !Proxy->GetWorld()->IsPlayInEditor())
+				ULandscapeInfo* Info = Proxy->GetLandscapeInfo();
+				if (Info)
 				{
-					ULandscapeInfo* Info = Proxy->GetLandscapeInfo();
-					if (Info)
+					for (int32 i = 0; i < Info->Layers.Num(); ++i)
 					{
-						for (int32 i = 0; i < Info->Layers.Num(); ++i)
+						if (Info->Layers[i].LayerInfoObj == this)
 						{
-							if (Info->Layers[i].LayerInfoObj == this)
-							{
-								Proxy->ChangedPhysMaterial();
-								break;
-							}
+							Proxy->ChangedPhysMaterial();
+							break;
 						}
 					}
 				}
 			}
 		}
-		else if (PropertyName == NAME_LayerUsageDebugColor || PropertyName == NAME_R || PropertyName == NAME_G || PropertyName == NAME_B || PropertyName == NAME_A)
+	}
+	else if (MemberPropertyName == GET_MEMBER_NAME_CHECKED(ULandscapeLayerInfoObject, LayerUsageDebugColor))
+	{
+		LayerUsageDebugColor.A = 1.0f;
+		for (TObjectIterator<ALandscapeProxy> It(/*AdditionalExclusionFlags = */RF_ClassDefaultObject, /*bIncludeDerivedClasses = */true, /*InInternalExclusionFlags = */EInternalObjectFlags::Garbage); It; ++It)
 		{
-			LayerUsageDebugColor.A = 1.0f;
-			for (TObjectIterator<ALandscapeProxy> It(/*AdditionalExclusionFlags = */RF_ClassDefaultObject, /*bIncludeDerivedClasses = */true, /*InInternalExclusionFlags = */EInternalObjectFlags::Garbage); It; ++It)
+			ALandscapeProxy* Proxy = *It;
+			if (Proxy->GetWorld() && !Proxy->GetWorld()->IsPlayInEditor())
 			{
-				ALandscapeProxy* Proxy = *It;
-				if (Proxy->GetWorld() && !Proxy->GetWorld()->IsPlayInEditor())
-				{
-					Proxy->MarkComponentsRenderStateDirty();
-				}
+				Proxy->MarkComponentsRenderStateDirty();
 			}
 		}
-		else if (PropertyName == GET_MEMBER_NAME_CHECKED(ULandscapeLayerInfoObject, SplineFalloffModulationTexture) ||
-				PropertyName == GET_MEMBER_NAME_CHECKED(ULandscapeLayerInfoObject, SplineFalloffModulationColorMask) ||
-				PropertyName == GET_MEMBER_NAME_CHECKED(ULandscapeLayerInfoObject, SplineFalloffModulationBias) ||
-				PropertyName == GET_MEMBER_NAME_CHECKED(ULandscapeLayerInfoObject, SplineFalloffModulationScale) ||
-				PropertyName == GET_MEMBER_NAME_CHECKED(ULandscapeLayerInfoObject, SplineFalloffModulationTiling))
+	}
+	else if ((MemberPropertyName == GET_MEMBER_NAME_CHECKED(ULandscapeLayerInfoObject, SplineFalloffModulationTexture))
+		|| (MemberPropertyName == GET_MEMBER_NAME_CHECKED(ULandscapeLayerInfoObject, SplineFalloffModulationColorMask))
+		|| (MemberPropertyName == GET_MEMBER_NAME_CHECKED(ULandscapeLayerInfoObject, SplineFalloffModulationBias))
+		|| (MemberPropertyName == GET_MEMBER_NAME_CHECKED(ULandscapeLayerInfoObject, SplineFalloffModulationScale))
+		|| (MemberPropertyName == GET_MEMBER_NAME_CHECKED(ULandscapeLayerInfoObject, SplineFalloffModulationTiling)))
+	{
+		for (TObjectIterator<ULandscapeInfo> It(/*AdditionalExclusionFlags = */RF_ClassDefaultObject, /*bIncludeDerivedClasses = */true, /*InInternalExclusionFlags = */EInternalObjectFlags::Garbage); It; ++It)
 		{
-			for (TObjectIterator<ULandscapeInfo> It(/*AdditionalExclusionFlags = */RF_ClassDefaultObject, /*bIncludeDerivedClasses = */true, /*InInternalExclusionFlags = */EInternalObjectFlags::Garbage); It; ++It)
+			if(ALandscape* Landscape = It->LandscapeActor.Get())
 			{
-				if(ALandscape* Landscape = It->LandscapeActor.Get())
-				{
-					Landscape->OnLayerInfoSplineFalloffModulationChanged(this);
-				}
+				Landscape->OnLayerInfoSplineFalloffModulationChanged(this);
 			}
 		}
 	}
@@ -5646,32 +5637,11 @@ bool ALandscapeProxy::CanEditChange(const FProperty* InProperty) const
 		return true;
 	}
 
-	// Don't allow edition of properties that are shared with the parent landscape properties
+	// Don't allow editing of properties that are shared with the parent landscape properties
 	// See ALandscapeProxy::FixupSharedData(ALandscape* Landscape)
 	if (GetLandscapeActor() != this)
 	{
-		FName PropertyName = InProperty ? InProperty->GetFName() : NAME_None;
-
-		if (   (PropertyName == GET_MEMBER_NAME_CHECKED(ALandscapeProxy, MaxLODLevel))
-			|| (PropertyName == GET_MEMBER_NAME_CHECKED(ALandscapeProxy, ComponentScreenSizeToUseSubSections))
-			|| (PropertyName == GET_MEMBER_NAME_CHECKED(ALandscapeProxy, LODDistributionSetting))
-			|| (PropertyName == GET_MEMBER_NAME_CHECKED(ALandscapeProxy, LOD0DistributionSetting))
-			|| (PropertyName == GET_MEMBER_NAME_CHECKED(ALandscapeProxy, LOD0ScreenSize))
-			|| (PropertyName == GET_MEMBER_NAME_CHECKED(ALandscapeProxy, ScalableLODDistributionSetting))
-			|| (PropertyName == GET_MEMBER_NAME_CHECKED(ALandscapeProxy, ScalableLOD0DistributionSetting))
-			|| (PropertyName == GET_MEMBER_NAME_CHECKED(ALandscapeProxy, ScalableLOD0ScreenSize))
-			|| (PropertyName == GET_MEMBER_NAME_CHECKED(ALandscapeProxy, LODGroupKey))
-			|| (PropertyName == GET_MEMBER_NAME_CHECKED(ALandscapeProxy, bUseCompressedHeightmapStorage))
-			|| (PropertyName == GET_MEMBER_NAME_CHECKED(ALandscapeProxy, TargetDisplayOrder))
-			|| (PropertyName == GET_MEMBER_NAME_CHECKED(ALandscapeProxy, TargetDisplayOrderList))
-			|| (PropertyName == GET_MEMBER_NAME_CHECKED(ALandscapeProxy, bEnableNanite))
-			|| (PropertyName == GET_MEMBER_NAME_CHECKED(ALandscapeProxy, bNaniteSkirtEnabled))
-			|| (PropertyName == GET_MEMBER_NAME_CHECKED(ALandscapeProxy, NaniteSkirtDepth))
-			|| (PropertyName == GET_MEMBER_NAME_CHECKED(ALandscapeProxy, NaniteLODIndex))
-			|| (PropertyName == GET_MEMBER_NAME_CHECKED(ALandscapeProxy, NanitePositionPrecision)))
-		{
-			return false;
-		}
+		return !IsPropertyInherited(InProperty);
 	}
 
 	return true;
@@ -5679,12 +5649,12 @@ bool ALandscapeProxy::CanEditChange(const FProperty* InProperty) const
 
 void ALandscapeProxy::PostEditChangeProperty(FPropertyChangedEvent& PropertyChangedEvent)
 {
-	const FName PropertyName = PropertyChangedEvent.MemberProperty ? PropertyChangedEvent.MemberProperty->GetFName() : NAME_None;
+	const FName MemberPropertyName = PropertyChangedEvent.MemberProperty ? PropertyChangedEvent.MemberProperty->GetFName() : NAME_None;
 	const FName SubPropertyName = PropertyChangedEvent.Property ? PropertyChangedEvent.Property->GetFName() : NAME_None;
 
 	bool bChangedPhysMaterial = false;
 
-	if (PropertyName == FName(TEXT("RelativeScale3D")))
+	if (MemberPropertyName == FName(TEXT("RelativeScale3D")))
 	{
 		// RelativeScale3D isn't even a property of ALandscapeProxy, it's a property of the root component
 		if (RootComponent)
@@ -5704,7 +5674,7 @@ void ALandscapeProxy::PostEditChangeProperty(FPropertyChangedEvent& PropertyChan
 			else if (SubPropertyName != FName("Z"))
 			{
 				// When changing all axis values at once (e.g. when copy/pasting the scale), we receive only one event and the sub-property is not set :
-				check(SubPropertyName == PropertyName); // Any other combination of property / sub-property is invalid				
+				check(SubPropertyName == MemberPropertyName); // Any other combination of property / sub-property is invalid				
 				if (!FMath::IsNearlyEqual(ModifiedScale.X, ModifiedScale.Y))
 				{
 					UE_LOG(LogLandscape, Warning, TEXT("Non-uniform XY scale for landscape (%f, %f) : scale will be forced to (%f, %f)"), ModifiedScale.X, ModifiedScale.Y, ModifiedScale.X, ModifiedScale.X);
@@ -5761,21 +5731,19 @@ void ALandscapeProxy::PostEditChangeProperty(FPropertyChangedEvent& PropertyChan
 			}
 		}
 	}
-
-	if (GIsEditor && PropertyName == GET_MEMBER_NAME_CHECKED(ALandscapeProxy, StreamingDistanceMultiplier))
+	else if (MemberPropertyName == GET_MEMBER_NAME_CHECKED(ALandscapeProxy, StreamingDistanceMultiplier))
 	{
 		// Recalculate in a few seconds.
 		GetWorld()->TriggerStreamingDataRebuild();
 	}
-	else if (GIsEditor && PropertyName == GET_MEMBER_NAME_CHECKED(ALandscapeProxy, DefaultPhysMaterial))
+	else if (MemberPropertyName == GET_MEMBER_NAME_CHECKED(ALandscapeProxy, DefaultPhysMaterial))
 	{
 		bChangedPhysMaterial = true;
 	}
-	else if (GIsEditor &&
-		(PropertyName == GET_MEMBER_NAME_CHECKED(ALandscapeProxy, CollisionMipLevel) ||
-		 PropertyName == GET_MEMBER_NAME_CHECKED(ALandscapeProxy, SimpleCollisionMipLevel) ||
-		 PropertyName == GET_MEMBER_NAME_CHECKED(ALandscapeProxy, bBakeMaterialPositionOffsetIntoCollision) ||
-		 PropertyName == GET_MEMBER_NAME_CHECKED(ALandscapeProxy, bGenerateOverlapEvents)))
+	else if ((MemberPropertyName == GET_MEMBER_NAME_CHECKED(ALandscapeProxy, CollisionMipLevel))
+		|| (MemberPropertyName == GET_MEMBER_NAME_CHECKED(ALandscapeProxy, SimpleCollisionMipLevel))
+		|| (MemberPropertyName == GET_MEMBER_NAME_CHECKED(ALandscapeProxy, bBakeMaterialPositionOffsetIntoCollision))
+		|| (MemberPropertyName == GET_MEMBER_NAME_CHECKED(ALandscapeProxy, bGenerateOverlapEvents)))
 	{
 		if (bBakeMaterialPositionOffsetIntoCollision)
 		{
@@ -5786,41 +5754,48 @@ void ALandscapeProxy::PostEditChangeProperty(FPropertyChangedEvent& PropertyChan
 			RecreateCollisionComponents();
 		}
 	}
-	else if (PropertyName == GET_MEMBER_NAME_CHECKED(ALandscapeProxy, ComponentScreenSizeToUseSubSections))
-	{
-		ChangeComponentScreenSizeToUseSubSections(ComponentScreenSizeToUseSubSections);
-	}
-	else if (PropertyName == GET_MEMBER_NAME_CHECKED(ALandscapeProxy, ScalableLODDistributionSetting)
-		|| PropertyName == GET_MEMBER_NAME_CHECKED(ALandscapeProxy, ScalableLOD0DistributionSetting)
-		|| PropertyName == GET_MEMBER_NAME_CHECKED(ALandscapeProxy, ScalableLOD0ScreenSize)
-		|| PropertyName == GET_MEMBER_NAME_CHECKED(ALandscapeProxy, LODDistributionSetting)
-		|| PropertyName == GET_MEMBER_NAME_CHECKED(ALandscapeProxy, LOD0DistributionSetting)
-		|| PropertyName == GET_MEMBER_NAME_CHECKED(ALandscapeProxy, LOD0ScreenSize))
+	else if ((MemberPropertyName == GET_MEMBER_NAME_CHECKED(ALandscapeProxy, ScalableLODDistributionSetting))
+		|| (MemberPropertyName == GET_MEMBER_NAME_CHECKED(ALandscapeProxy, ScalableLOD0DistributionSetting))
+		|| (MemberPropertyName == GET_MEMBER_NAME_CHECKED(ALandscapeProxy, ScalableLOD0ScreenSize))
+		|| (MemberPropertyName == GET_MEMBER_NAME_CHECKED(ALandscapeProxy, LODDistributionSetting))
+		|| (MemberPropertyName == GET_MEMBER_NAME_CHECKED(ALandscapeProxy, LOD0DistributionSetting))
+		|| (MemberPropertyName == GET_MEMBER_NAME_CHECKED(ALandscapeProxy, LOD0ScreenSize))
+		|| (MemberPropertyName == GET_MEMBER_NAME_CHECKED(ALandscapeProxy, bUseScalableLODSettings))
+		|| (MemberPropertyName == GET_MEMBER_NAME_CHECKED(ALandscapeProxy, LODBlendRange))
+		|| (MemberPropertyName == GET_MEMBER_NAME_CHECKED(ALandscapeProxy, ComponentScreenSizeToUseSubSections))
+		|| (MemberPropertyName == GET_MEMBER_NAME_CHECKED(ALandscapeProxy, RuntimeVirtualTextures))
+		|| (MemberPropertyName == GET_MEMBER_NAME_CHECKED(ALandscapeProxy, VirtualTextureRenderPassType))
+		|| (MemberPropertyName == GET_MEMBER_NAME_CHECKED(ALandscapeProxy, bVirtualTextureRenderWithQuad))
+		|| (MemberPropertyName == GET_MEMBER_NAME_CHECKED(ALandscapeProxy, bVirtualTextureRenderWithQuadHQ))
+		|| (MemberPropertyName == GET_MEMBER_NAME_CHECKED(ALandscapeProxy, VirtualTextureNumLods))
+		|| (MemberPropertyName == GET_MEMBER_NAME_CHECKED(ALandscapeProxy, VirtualTextureLodBias))
+		|| (MemberPropertyName == GET_MEMBER_NAME_CHECKED(ALandscapeProxy, bUseDynamicMaterialInstance)))
 	{		
 		MarkComponentsRenderStateDirty();
 	}
-	else if (PropertyName == GET_MEMBER_NAME_CHECKED(ALandscapeProxy, bUseMaterialPositionOffsetInStaticLighting))
+	else if (MemberPropertyName == GET_MEMBER_NAME_CHECKED(ALandscapeProxy, bUseMaterialPositionOffsetInStaticLighting))
 	{
 		InvalidateLightingCache();
 	}
-	else if ((PropertyName == GET_MEMBER_NAME_CHECKED(ALandscapeProxy, CastShadow))
-		|| (PropertyName == GET_MEMBER_NAME_CHECKED(ALandscapeProxy, bCastDynamicShadow))
-		|| (PropertyName == GET_MEMBER_NAME_CHECKED(ALandscapeProxy, ShadowCacheInvalidationBehavior))
-		|| (PropertyName == GET_MEMBER_NAME_CHECKED(ALandscapeProxy, bCastStaticShadow))
-		|| (PropertyName == GET_MEMBER_NAME_CHECKED(ALandscapeProxy, bCastContactShadow))
-		|| (PropertyName == GET_MEMBER_NAME_CHECKED(ALandscapeProxy, bCastFarShadow))
-		|| (PropertyName == GET_MEMBER_NAME_CHECKED(ALandscapeProxy, bCastHiddenShadow))
-		|| (PropertyName == GET_MEMBER_NAME_CHECKED(ALandscapeProxy, bCastShadowAsTwoSided))
-		|| (PropertyName == GET_MEMBER_NAME_CHECKED(ALandscapeProxy, bAffectDistanceFieldLighting))
-		|| (PropertyName == GET_MEMBER_NAME_CHECKED(ALandscapeProxy, bRenderCustomDepth))
-		|| (PropertyName == GET_MEMBER_NAME_CHECKED(ALandscapeProxy, CustomDepthStencilWriteMask))
-		|| (PropertyName == GET_MEMBER_NAME_CHECKED(ALandscapeProxy, CustomDepthStencilValue))
-		|| (PropertyName == GET_MEMBER_NAME_CHECKED(ALandscapeProxy, LightingChannels))
-		|| (PropertyName == GET_MEMBER_NAME_CHECKED(ALandscapeProxy, LDMaxDrawDistance))
-		|| (PropertyName == GET_MEMBER_NAME_CHECKED(ALandscapeProxy, bUsedForNavigation))
-		|| (PropertyName == GET_MEMBER_NAME_CHECKED(ALandscapeProxy, bFillCollisionUnderLandscapeForNavmesh)))
-	{
-		// Replicate shared properties to all components.
+	// Replicate properties shared with components to all of them :
+	else if ((MemberPropertyName == GET_MEMBER_NAME_CHECKED(ALandscapeProxy, CastShadow))
+		|| (MemberPropertyName == GET_MEMBER_NAME_CHECKED(ALandscapeProxy, bCastDynamicShadow))
+		|| (MemberPropertyName == GET_MEMBER_NAME_CHECKED(ALandscapeProxy, ShadowCacheInvalidationBehavior))
+		|| (MemberPropertyName == GET_MEMBER_NAME_CHECKED(ALandscapeProxy, bCastStaticShadow))
+		|| (MemberPropertyName == GET_MEMBER_NAME_CHECKED(ALandscapeProxy, bCastContactShadow))
+		|| (MemberPropertyName == GET_MEMBER_NAME_CHECKED(ALandscapeProxy, bCastFarShadow))
+		|| (MemberPropertyName == GET_MEMBER_NAME_CHECKED(ALandscapeProxy, bCastHiddenShadow))
+		|| (MemberPropertyName == GET_MEMBER_NAME_CHECKED(ALandscapeProxy, bCastShadowAsTwoSided))
+		|| (MemberPropertyName == GET_MEMBER_NAME_CHECKED(ALandscapeProxy, bAffectDistanceFieldLighting))
+		|| (MemberPropertyName == GET_MEMBER_NAME_CHECKED(ALandscapeProxy, bRenderCustomDepth))
+		|| (MemberPropertyName == GET_MEMBER_NAME_CHECKED(ALandscapeProxy, CustomDepthStencilWriteMask))
+		|| (MemberPropertyName == GET_MEMBER_NAME_CHECKED(ALandscapeProxy, CustomDepthStencilValue))
+		|| (MemberPropertyName == GET_MEMBER_NAME_CHECKED(ALandscapeProxy, LightingChannels))
+		|| (MemberPropertyName == GET_MEMBER_NAME_CHECKED(ALandscapeProxy, LDMaxDrawDistance))
+		|| (MemberPropertyName == GET_MEMBER_NAME_CHECKED(ALandscapeProxy, bUsedForNavigation))
+		|| (MemberPropertyName == GET_MEMBER_NAME_CHECKED(ALandscapeProxy, bFillCollisionUnderLandscapeForNavmesh)))
+	{
+		// TODO [jonathan.bard] : Move to its own function so that we can streamline the setters of these properties when we start exposing them
 		for (int32 ComponentIndex = 0; ComponentIndex < LandscapeComponents.Num(); ComponentIndex++)
 		{
 			ULandscapeComponent* Comp = LandscapeComponents[ComponentIndex];
@@ -5832,70 +5807,16 @@ void ALandscapeProxy::PostEditChangeProperty(FPropertyChangedEvent& PropertyChan
 
 		UpdateNaniteSharedPropertiesFromActor();
 	}
-
-	if (GIsEditor && PropertyName == GET_MEMBER_NAME_CHECKED(ALandscapeProxy, bEnableNanite))
-	{
-		// This property is entirely shared with the parent material so don't let it be set to a different value than the parent : 
-		if (ALandscape* Parent = GetLandscapeActor(); (Parent != this) && (bEnableNanite != Parent->IsNaniteEnabled()))
-		{
-			bEnableNanite = Parent->IsNaniteEnabled();
-		}
-		InvalidateGeneratedComponentData(/* bInvalidateLightingCache = */false);
-		MarkComponentsRenderStateDirty();
-	}
-	if (GIsEditor && PropertyName == GET_MEMBER_NAME_CHECKED(ALandscapeProxy, bNaniteSkirtEnabled))
-	{
-		// This property is entirely shared with the parent material so don't let it be set to a different value than the parent : 
-		if (ALandscape* Parent = GetLandscapeActor(); (Parent != this) && (bNaniteSkirtEnabled != Parent->IsNaniteSkirtEnabled()))
-		{
-			bNaniteSkirtEnabled = Parent->IsNaniteSkirtEnabled();
-		}
-		InvalidateGeneratedComponentData(/* bInvalidateLightingCache = */false);
-		MarkComponentsRenderStateDirty();
-	}
-	if (GIsEditor && PropertyName == GET_MEMBER_NAME_CHECKED(ALandscapeProxy, NaniteSkirtDepth))
-	{
-		// This property is entirely shared with the parent material so don't let it be set to a different value than the parent : 
-		if (ALandscape* Parent = GetLandscapeActor(); (Parent != this) && (NaniteSkirtDepth != Parent->GetNaniteSkirtDepth()))
-		{
-			NaniteSkirtDepth = Parent->GetNaniteSkirtDepth();
-		}
-		InvalidateGeneratedComponentData(/* bInvalidateLightingCache = */false);
-		MarkComponentsRenderStateDirty();
-	}
-	if (GIsEditor && PropertyName == GET_MEMBER_NAME_CHECKED(ALandscapeProxy, NaniteLODIndex))
-	{
-		// This property is entirely shared with the parent material so don't let it be set to a different value than the parent : 
-		if (ALandscape* Parent = GetLandscapeActor(); (Parent != this) && (NaniteLODIndex != Parent->GetNaniteLODIndex()))
-		{
-			NaniteLODIndex = Parent->GetNaniteLODIndex();
-		}
-		InvalidateGeneratedComponentData(/* bInvalidateLightingCache = */false);
-		MarkComponentsRenderStateDirty();
-	}
-	if (GIsEditor && PropertyName == GET_MEMBER_NAME_CHECKED(ALandscapeProxy, NanitePositionPrecision))
+	// Nanite-related changes require invalidating the generated component data :
+	else if ((MemberPropertyName == GET_MEMBER_NAME_CHECKED(ALandscapeProxy, bEnableNanite))
+		|| (MemberPropertyName == GET_MEMBER_NAME_CHECKED(ALandscapeProxy, bNaniteSkirtEnabled))
+		|| (MemberPropertyName == GET_MEMBER_NAME_CHECKED(ALandscapeProxy, NaniteSkirtDepth))
+		|| (MemberPropertyName == GET_MEMBER_NAME_CHECKED(ALandscapeProxy, NaniteLODIndex))
+		|| (MemberPropertyName == GET_MEMBER_NAME_CHECKED(ALandscapeProxy, NanitePositionPrecision)))
 	{
-		// This property is entirely shared with the parent material so don't let it be set to a different value than the parent : 
-		if (ALandscape* Parent = GetLandscapeActor(); (Parent != this) && (NanitePositionPrecision != Parent->GetNanitePositionPrecision()))
-		{
-			NanitePositionPrecision = Parent->GetNanitePositionPrecision();
-		}
 		InvalidateGeneratedComponentData(/* bInvalidateLightingCache = */false);
 		MarkComponentsRenderStateDirty();
 	}
-	else if (PropertyName == GET_MEMBER_NAME_CHECKED(ALandscapeProxy, bUseDynamicMaterialInstance))
-	{
-		MarkComponentsRenderStateDirty();
-	}
-	else if (PropertyName == GET_MEMBER_NAME_CHECKED(ALandscapeProxy, RuntimeVirtualTextures)
-		|| PropertyName == GET_MEMBER_NAME_CHECKED(ALandscapeProxy, VirtualTextureRenderPassType)
-		|| PropertyName == GET_MEMBER_NAME_CHECKED(ALandscapeProxy, bVirtualTextureRenderWithQuad)
-		|| PropertyName == GET_MEMBER_NAME_CHECKED(ALandscapeProxy, bVirtualTextureRenderWithQuadHQ)
-		|| PropertyName == GET_MEMBER_NAME_CHECKED(ALandscapeProxy, VirtualTextureNumLods)
-		|| PropertyName == GET_MEMBER_NAME_CHECKED(ALandscapeProxy, VirtualTextureLodBias))
-	{
-		MarkComponentsRenderStateDirty();
-	}
 
 	// Remove null layer infos
 	EditorLayerSettings.RemoveAll([](const FLandscapeEditorLayerSettings& Entry) { return Entry.LayerInfoObj == nullptr; });
@@ -5931,9 +5852,9 @@ bool ALandscapeStreamingProxy::CanEditChange(const FProperty* InProperty) const
 
 void ALandscapeStreamingProxy::PostEditChangeProperty(FPropertyChangedEvent& PropertyChangedEvent)
 {
-	const FName PropertyName = PropertyChangedEvent.MemberProperty ? PropertyChangedEvent.MemberProperty->GetFName() : NAME_None;
+	const FName MemberPropertyName = PropertyChangedEvent.MemberProperty ? PropertyChangedEvent.MemberProperty->GetFName() : NAME_None;
 
-	if (PropertyName == FName(TEXT("LandscapeActorRef")))
+	if (MemberPropertyName == GET_MEMBER_NAME_CHECKED(ALandscapeStreamingProxy, LandscapeActorRef))
 	{
 		// Landscape Actor reference was changed .. need to update LandscapeGUIDs to match
 		if (LandscapeActorRef && IsValidLandscapeActor(LandscapeActorRef.Get()))
@@ -5954,11 +5875,13 @@ void ALandscapeStreamingProxy::PostEditChangeProperty(FPropertyChangedEvent& Pro
 			LandscapeActorRef = nullptr;
 		}
 	}
-	else if (PropertyName == FName(TEXT("LandscapeMaterial")) || PropertyName == FName(TEXT("LandscapeHoleMaterial")) || PropertyName == FName(TEXT("PerLODOverrideMaterials")))
+	else if ((MemberPropertyName == GET_MEMBER_NAME_CHECKED(ALandscapeProxy, LandscapeMaterial))
+		|| (MemberPropertyName == GET_MEMBER_NAME_CHECKED(ALandscapeProxy, LandscapeHoleMaterial))
+		|| (MemberPropertyName == GET_MEMBER_NAME_CHECKED(ALandscapeStreamingProxy, PerLODOverrideMaterials)))
 	{
 		bool RecreateMaterialInstances = true;
 
-		if (PropertyName == FName(TEXT("PerLODOverrideMaterials")) && PropertyChangedEvent.ChangeType == EPropertyChangeType::ArrayAdd)
+		if ((MemberPropertyName == GET_MEMBER_NAME_CHECKED(ALandscapeStreamingProxy, PerLODOverrideMaterials)) && (PropertyChangedEvent.ChangeType == EPropertyChangeType::ArrayAdd))
 		{
 			RecreateMaterialInstances = false;
 		}
@@ -6272,7 +6195,6 @@ void ALandscape::PostEditChangeProperty(FPropertyChangedEvent& PropertyChangedEv
 		return Clamped;
 	};
 
-	const FName PropertyName = PropertyChangedEvent.Property ? PropertyChangedEvent.Property->GetFName() : NAME_None;
 	const FName MemberPropertyName = PropertyChangedEvent.MemberProperty ? PropertyChangedEvent.MemberProperty->GetFName() : NAME_None;
 
 	bool bMaterialChanged = false;
@@ -6284,19 +6206,20 @@ void ALandscape::PostEditChangeProperty(FPropertyChangedEvent& PropertyChangedEv
 
 	ULandscapeInfo* Info = GetLandscapeInfo();
 
-	if ((PropertyName == GET_MEMBER_NAME_CHECKED(ALandscapeProxy, LandscapeMaterial) || PropertyName == GET_MEMBER_NAME_CHECKED(ALandscapeProxy, LandscapeHoleMaterial) || MemberPropertyName == GET_MEMBER_NAME_CHECKED(ALandscapeProxy, PerLODOverrideMaterials))
-		&& PropertyChangedEvent.ChangeType != EPropertyChangeType::ArrayAdd)
+	if ((MemberPropertyName == GET_MEMBER_NAME_CHECKED(ALandscapeProxy, LandscapeMaterial)) 
+		|| (MemberPropertyName == GET_MEMBER_NAME_CHECKED(ALandscapeProxy, LandscapeHoleMaterial)) 
+		|| ((MemberPropertyName == GET_MEMBER_NAME_CHECKED(ALandscapeProxy, PerLODOverrideMaterials)) && (PropertyChangedEvent.ChangeType != EPropertyChangeType::ArrayAdd)))
 	{
-		bool HasMaterialChanged = false;
+		bool bHasMaterialChanged = false;
 
 		if (PropertyChangedEvent.ChangeType != EPropertyChangeType::Interactive)
 		{
 			if (PreEditLandscapeMaterial != LandscapeMaterial || PreEditLandscapeHoleMaterial != LandscapeHoleMaterial || PreEditPerLODOverrideMaterials.Num() != PerLODOverrideMaterials.Num() || bIsPerformingInteractiveActionOnLandscapeMaterialOverride)
 			{
-				HasMaterialChanged = true;
+				bHasMaterialChanged = true;
 			}
 
-			if (!HasMaterialChanged)
+			if (!bHasMaterialChanged)
 			{
 				for (int32 i = 0; i < PerLODOverrideMaterials.Num(); ++i)
 				{
@@ -6305,7 +6228,7 @@ void ALandscape::PostEditChangeProperty(FPropertyChangedEvent& PropertyChangedEv
 
 					if (!(PreEditMaterialOverride == NewMaterialOverride))
 					{
-						HasMaterialChanged = true;
+						bHasMaterialChanged = true;
 						break;
 					}
 				}
@@ -6316,10 +6239,10 @@ void ALandscape::PostEditChangeProperty(FPropertyChangedEvent& PropertyChangedEv
 		else
 		{
 			// We are probably using a slider or something similar in PerLODOverrideMaterials
-			bIsPerformingInteractiveActionOnLandscapeMaterialOverride = MemberPropertyName == GET_MEMBER_NAME_CHECKED(ALandscapeProxy, PerLODOverrideMaterials);
+			bIsPerformingInteractiveActionOnLandscapeMaterialOverride = (MemberPropertyName == GET_MEMBER_NAME_CHECKED(ALandscapeProxy, PerLODOverrideMaterials));
 		}
 
-		if (Info != nullptr && HasMaterialChanged)
+		if (Info != nullptr && bHasMaterialChanged)
 		{
 			FMaterialUpdateContext MaterialUpdateContext;
 			Info->UpdateLayerInfoMap(/*this*/);
@@ -6353,14 +6276,10 @@ void ALandscape::PostEditChangeProperty(FPropertyChangedEvent& PropertyChangedEv
 		// Some edit layers could be affected by BP brushes, which might need to be updated when the landscape is transformed :
 		RequestLayersContentUpdate(ELandscapeLayerUpdateMode::Update_All);
 	}
-	else if (GIsEditor && PropertyName == GET_MEMBER_NAME_CHECKED(ALandscapeProxy, MaxLODLevel))
+	else if (MemberPropertyName == GET_MEMBER_NAME_CHECKED(ALandscapeProxy, MaxLODLevel))
 	{
 		MaxLODLevel = FMath::Clamp<int32>(MaxLODLevel, -1, FMath::CeilLogTwo(SubsectionSizeQuads + 1) - 1);
 	}
-	else if (PropertyName == GET_MEMBER_NAME_CHECKED(ALandscapeProxy, ComponentScreenSizeToUseSubSections))
-	{
-		ComponentScreenSizeToUseSubSections = FMath::Clamp<float>(ComponentScreenSizeToUseSubSections, 0.01f, 1.0f);
-	}
 	else if (MemberPropertyName == GET_MEMBER_NAME_CHECKED(ALandscapeProxy, LODDistributionSetting))
 	{
 		LODDistributionSetting = FMath::Clamp(LODDistributionSetting, 1.0f, 10.0f);
@@ -6385,34 +6304,29 @@ void ALandscape::PostEditChangeProperty(FPropertyChangedEvent& PropertyChangedEv
 	{
 		ScalableLOD0ScreenSize = Clamp(ScalableLOD0ScreenSize, 0.1f, 10.0f);
 	}
-	else if (PropertyName == GET_MEMBER_NAME_CHECKED(ALandscapeProxy, CollisionMipLevel))
+	else if (MemberPropertyName == GET_MEMBER_NAME_CHECKED(ALandscapeProxy, CollisionMipLevel))
 	{
 		CollisionMipLevel = FMath::Clamp<int32>(CollisionMipLevel, 0, FMath::CeilLogTwo(SubsectionSizeQuads + 1) - 1);
 	}
-	else if (PropertyName == GET_MEMBER_NAME_CHECKED(ALandscapeProxy, SimpleCollisionMipLevel))
+	else if (MemberPropertyName == GET_MEMBER_NAME_CHECKED(ALandscapeProxy, SimpleCollisionMipLevel))
 	{
 		SimpleCollisionMipLevel = FMath::Clamp<int32>(SimpleCollisionMipLevel, 0, FMath::CeilLogTwo(SubsectionSizeQuads + 1) - 1);
 	}
-	else if (PropertyName == GET_MEMBER_NAME_CHECKED(ALandscapeProxy, bBakeMaterialPositionOffsetIntoCollision))
-	{
-		// TODO [jonathan.bard] : why propagate to proxies?
-		bPropagateToProxies = true;
-	}
-	else if (GIsEditor && PropertyName == GET_MEMBER_NAME_CHECKED(ALandscapeProxy, StaticLightingResolution))
+	else if (MemberPropertyName == GET_MEMBER_NAME_CHECKED(ALandscapeProxy, StaticLightingResolution))
 	{
 		StaticLightingResolution = ::AdjustStaticLightingResolution(StaticLightingResolution, NumSubsections, SubsectionSizeQuads, ComponentSizeQuads);
 		bChangedLighting = true;
 	}
-	else if (GIsEditor && PropertyName == GET_MEMBER_NAME_CHECKED(ALandscapeProxy, StaticLightingLOD))
+	else if (MemberPropertyName == GET_MEMBER_NAME_CHECKED(ALandscapeProxy, StaticLightingLOD))
 	{
 		StaticLightingLOD = FMath::Clamp<int32>(StaticLightingLOD, 0, FMath::CeilLogTwo(SubsectionSizeQuads + 1) - 1);
 		bChangedLighting = true;
 	}
-	else if (GIsEditor && PropertyName == GET_MEMBER_NAME_CHECKED(ALandscapeProxy, ExportLOD))
+	else if (MemberPropertyName == GET_MEMBER_NAME_CHECKED(ALandscapeProxy, ExportLOD))
 	{
 		ExportLOD = FMath::Clamp<int32>(ExportLOD, 0, FMath::CeilLogTwo(SubsectionSizeQuads + 1) - 1);
 	}
-	else if (GIsEditor && (PropertyName == GET_MEMBER_NAME_CHECKED(ALandscape, bEnableNanite)))
+	else if (MemberPropertyName == GET_MEMBER_NAME_CHECKED(ALandscape, bEnableNanite))
 	{
 		bNaniteToggled = true;
 
@@ -6422,13 +6336,13 @@ void ALandscape::PostEditChangeProperty(FPropertyChangedEvent& PropertyChangedEv
 		UpdateRenderingMethod();
 		MarkComponentsRenderStateDirty();
 	}
-	else if (PropertyName == GET_MEMBER_NAME_CHECKED(ALandscape, LODGroupKey))
+	else if (MemberPropertyName == GET_MEMBER_NAME_CHECKED(ALandscape, LODGroupKey))
 	{
 		bMarkAllLandscapeRenderStateDirty = true;
 	}
 	
 	// If the property that has changed is overridable or inherited, synchronize the change on all landscape proxies :
-	if (GIsEditor && IsSharedProperty(PropertyName))
+	if (IsSharedProperty(MemberPropertyName))
 	{
 		bPropagateToProxies = true;
 	}
@@ -6461,8 +6375,11 @@ void ALandscape::PostEditChangeProperty(FPropertyChangedEvent& PropertyChangedEv
 		// Update normals if DrawScale3D is changed
 		if (MemberPropertyName == FName(TEXT("RelativeScale3D")))
 		{
-			FLandscapeEditDataInterface LandscapeEdit(Info);
-			LandscapeEdit.RecalculateNormals();
+			if (!Info->CanHaveLayersContent())
+			{
+				FLandscapeEditDataInterface LandscapeEdit(Info);
+				LandscapeEdit.RecalculateNormals();
+			}
 		}
 
 		if (bNeedsRecalcBoundingBox || bMaterialChanged || bChangedLighting)
@@ -6555,7 +6472,7 @@ void ALandscapeProxy::ChangedPhysMaterial()
 void ULandscapeComponent::PreEditChange(FProperty* PropertyThatWillChange)
 {
 	Super::PreEditChange(PropertyThatWillChange);
-	if (GIsEditor && PropertyThatWillChange && (PropertyThatWillChange->GetFName() == GET_MEMBER_NAME_CHECKED(ULandscapeComponent, ForcedLOD) || PropertyThatWillChange->GetFName() == GET_MEMBER_NAME_CHECKED(ULandscapeComponent, LODBias)))
+	if (PropertyThatWillChange && (PropertyThatWillChange->GetFName() == GET_MEMBER_NAME_CHECKED(ULandscapeComponent, ForcedLOD) || PropertyThatWillChange->GetFName() == GET_MEMBER_NAME_CHECKED(ULandscapeComponent, LODBias)))
 	{
 		// PreEdit unregister component and re-register after PostEdit so we will lose XYtoComponentMap for this component
 		ULandscapeInfo* Info = GetLandscapeInfo();
@@ -6574,14 +6491,16 @@ void ULandscapeComponent::PreEditChange(FProperty* PropertyThatWillChange)
 
 void ULandscapeComponent::PostEditChangeProperty(FPropertyChangedEvent& PropertyChangedEvent)
 {
-	const FName PropertyName = PropertyChangedEvent.Property ? PropertyChangedEvent.Property->GetFName() : NAME_None;
 	const FName MemberPropertyName = PropertyChangedEvent.MemberProperty ? PropertyChangedEvent.MemberProperty->GetFName() : NAME_None;
+	const FName SubPropertyName = PropertyChangedEvent.Property ? PropertyChangedEvent.Property->GetFName() : NAME_None;
 
-	if (PropertyName == FName(TEXT("OverrideMaterial")) || MemberPropertyName == FName(TEXT("PerLODOverrideMaterials")) || MemberPropertyName == FName(TEXT("MaterialPerLOD_Key")))
+	if ((MemberPropertyName == GET_MEMBER_NAME_CHECKED(ULandscapeComponent, OverrideMaterial))
+		|| (MemberPropertyName == GET_MEMBER_NAME_CHECKED(ULandscapeComponent, PerLODOverrideMaterials))
+		|| (SubPropertyName == FName(TEXT("MaterialPerLOD_Key"))))
 	{
 		bool RecreateMaterialInstances = true;
 
-		if (PropertyName == FName(TEXT("PerLODOverrideMaterials")) && PropertyChangedEvent.ChangeType == EPropertyChangeType::ArrayAdd)
+		if ((MemberPropertyName == GET_MEMBER_NAME_CHECKED(ULandscapeComponent, PerLODOverrideMaterials)) && (PropertyChangedEvent.ChangeType == EPropertyChangeType::ArrayAdd))
 		{
 			RecreateMaterialInstances = false;
 		}
@@ -6598,12 +6517,13 @@ void ULandscapeComponent::PostEditChangeProperty(FPropertyChangedEvent& Property
 			}
 		}
 	}
-	else if (GIsEditor && (PropertyName == FName(TEXT("ForcedLOD")) || PropertyName == FName(TEXT("LODBias"))))
+	else if ((MemberPropertyName == GET_MEMBER_NAME_CHECKED(ULandscapeComponent, ForcedLOD))
+		|| (MemberPropertyName == GET_MEMBER_NAME_CHECKED(ULandscapeComponent, LODBias)))
 	{
-		bool bForcedLODChanged = PropertyName == FName(TEXT("ForcedLOD"));
+		bool bForcedLODChanged = (MemberPropertyName == GET_MEMBER_NAME_CHECKED(ULandscapeComponent, ForcedLOD));
 		SetLOD(bForcedLODChanged, bForcedLODChanged ? ForcedLOD : LODBias);
 	}
-	else if (GIsEditor && PropertyName == FName(TEXT("StaticLightingResolution")))
+	else if (MemberPropertyName == GET_MEMBER_NAME_CHECKED(ULandscapeComponent, StaticLightingResolution))
 	{
 		if (StaticLightingResolution > 0.0f)
 		{
@@ -6615,15 +6535,14 @@ void ULandscapeComponent::PostEditChangeProperty(FPropertyChangedEvent& Property
 		}
 		InvalidateLightingCache();
 	}
-	else if (GIsEditor && PropertyName == FName(TEXT("LightingLODBias")))
+	else if (MemberPropertyName == GET_MEMBER_NAME_CHECKED(ULandscapeComponent, LightingLODBias))
 	{
 		int32 MaxLOD = FMath::CeilLogTwo(SubsectionSizeQuads + 1) - 1;
 		LightingLODBias = FMath::Clamp<int32>(LightingLODBias, -1, MaxLOD);
 		InvalidateLightingCache();
 	}
-	else if (GIsEditor &&
-		(PropertyName == GET_MEMBER_NAME_CHECKED(ULandscapeComponent, CollisionMipLevel) ||
-		 PropertyName == GET_MEMBER_NAME_CHECKED(ULandscapeComponent, SimpleCollisionMipLevel)))
+	else if ((MemberPropertyName == GET_MEMBER_NAME_CHECKED(ULandscapeComponent, CollisionMipLevel))
+		 || (MemberPropertyName == GET_MEMBER_NAME_CHECKED(ULandscapeComponent, SimpleCollisionMipLevel)))
 	{
 		CollisionMipLevel = FMath::Clamp<int32>(CollisionMipLevel, 0, FMath::CeilLogTwo(SubsectionSizeQuads + 1) - 1);
 		SimpleCollisionMipLevel = FMath::Clamp<int32>(SimpleCollisionMipLevel, 0, FMath::CeilLogTwo(SubsectionSizeQuads + 1) - 1);
