Fix for ensure hit on level close / shutdown
Make sure we push async tasks forward when shutting down

diff --git a/Engine/Source/Runtime/Landscape/Private/LandscapeAsyncTextureReadback.h b/Engine/Source/Runtime/Landscape/Private/LandscapeAsyncTextureReadback.h
--- a/Engine/Source/Runtime/Landscape/Private/LandscapeAsyncTextureReadback.h
+++ b/Engine/Source/Runtime/Landscape/Private/LandscapeAsyncTextureReadback.h
@@ -59,6 +59,13 @@ public:
 		return MoveTemp(ReadbackResults);
 	}
 
+	FString ToString()
+	{
+		FString Result;
+		Result.Appendf(TEXT("FLandscapeAsyncTextureReadback { Submit: %d Complete: %d AsyncReadback: %p }"), bAsyncReadbackSubmitOnRenderThread, bAsyncReadbackCompleteOnRenderThread, AsyncReadback.Get());
+		return Result;
+	}
+
 	// once complete, call this to queue deletion of the readback object on the render thread
 	// (this must be deleted on the render thread to avoid other render-thread queued commands from accessing a deallocated pointer)
 	void QueueDeletionFromGameThread();
diff --git a/Engine/Source/Runtime/Landscape/Private/LandscapeGrassMapsBuilder.cpp b/Engine/Source/Runtime/Landscape/Private/LandscapeGrassMapsBuilder.cpp
--- a/Engine/Source/Runtime/Landscape/Private/LandscapeGrassMapsBuilder.cpp
+++ b/Engine/Source/Runtime/Landscape/Private/LandscapeGrassMapsBuilder.cpp
@@ -306,17 +306,55 @@ FLandscapeGrassMapsBuilder::~FLandscapeGrassMapsBuilder()
 		{
 			UE::Landscape::SubmitGPUCommands(/* bBlockUntilComplete =  */ true);
 		}
+
+		// if any async fetch tasks are in flight, force completion
+		if (AsyncFetchCount > 0)
+		{
+			for (auto It = ComponentStates.CreateIterator(); It; ++It)
+			{
+				FComponentState* State = It.Value();
+				if (State->Stage == EComponentStage::AsyncFetch)
+				{
+					check(State->AsyncFetchTask.Get());
+					State->AsyncFetchTask->EnsureCompletion(/* bDoWorkOnThisThreadIfNotStarted= */ true, /* bIsLatencySensitive= */ true);
+				}
+			}
+		}
+
 		Iterations++;
 	}
 
-	if (!ensure(ComponentStates.Num() == 0))
+	if (ComponentStates.Num() != 0)
 	{
-		// somehow we failed to free the components the right way (either a GPU readback is stuck, or state logic is broken)
-		// force free the state anyways and hope for the best
+		// somehow we failed to free the components the right way (either a GPU readback or async fetch is stuck, or state transition logic is broken)
 		for (auto It = ComponentStates.CreateIterator(); It; ++It)
 		{
 			FComponentState* State = It.Value();
 			UE_LOG(LogGrass, Warning, TEXT("Failed to clear grass data state after %d iterations (stage:%d ticks:%d), forcing deletion of the state."), Iterations, State->Stage, State->TickCount);
+
+			if (State->Stage == EComponentStage::Rendering)
+			{
+				if ((State->ActiveRender != nullptr) && (State->ActiveRender->AsyncReadbackPtr != nullptr))
+				{
+					UE_LOG(LogGrass, Warning, TEXT("  %s"), *State->ActiveRender->AsyncReadbackPtr->ToString());
+				}
+			}
+			else if (State->Stage == EComponentStage::AsyncFetch)
+			{
+				if (State->AsyncFetchTask.Get() != nullptr)
+				{
+					UE_LOG(LogGrass, Warning, TEXT("  AsyncFetchTask: %p"), State->AsyncFetchTask.Get());
+				}
+			}
+		}
+
+		// report the error so we capture the callstack and the log warnings above
+		ensure(ComponentStates.Num() == 0);
+
+		// force free the states anyways and hope for the best.  If crashes ensue the logs above should indicate why.
+		for (auto It = ComponentStates.CreateIterator(); It; ++It)
+		{
+			FComponentState* State = It.Value();
 			State->~FComponentState();
 			StatePoolAllocator.Free(State);
 			It.RemoveCurrent();
@@ -867,6 +905,8 @@ bool FLandscapeGrassMapsBuilder::CancelAndEvict(FComponentState& State)
 					// we can't cancel yet.. must wait for the readback to complete
 					return false;
 				}
+
+				// readback is complete, we can delete the active render
 				State.ActiveRender.Reset();
 			}
 			check(RenderingCount > 0);
