Fix for grass not appearing on Level Instanced landscapes.
Also enable fastpath transitions on registration to speed up grass generation when good grass maps already exist.

Grass does not populate because it is waiting for landscape layers to be up to date, which will never happen while the landscape level instanced.
Added code to skip the up to date check when level instanced.

diff --git a/Engine/Source/Runtime/Landscape/Private/LandscapeGrass.cpp b/Engine/Source/Runtime/Landscape/Private/LandscapeGrass.cpp
--- a/Engine/Source/Runtime/Landscape/Private/LandscapeGrass.cpp
+++ b/Engine/Source/Runtime/Landscape/Private/LandscapeGrass.cpp
@@ -1696,8 +1696,12 @@ void ALandscapeProxy::TickGrass(const TArray<FVector>& Cameras, int32& InOutNumC
 #if WITH_EDITORONLY_DATA
 	if (ALandscape* Landscape = GetLandscapeActor())
 	{
-		// Don't allow grass to tick, except in ES3.1 preview mode, since landscape update is not possible there : IsUpToDate might return false and we'll never see grass in the preview mode as a result :
-		bool bAllowGrassTick = Landscape->IsUpToDate() || (GetWorld()->GetFeatureLevel() < ERHIFeatureLevel::SM5);
+		ULandscapeInfo* LandscapeInfo = GetLandscapeInfo();
+		
+		bool bLandscapeUpdateAllowed = (LandscapeInfo != nullptr) && (GetWorld()->GetFeatureLevel() >= ERHIFeatureLevel::SM5) && LandscapeInfo->SupportsLandscapeEditing();
+
+		// Don't allow grass to tick if landscape is not up to date -- unless landscape update is not possible (preview or level instanced modes)
+		bool bAllowGrassTick = Landscape->IsUpToDate() || !bLandscapeUpdateAllowed;
 		if (!bAllowGrassTick || !Landscape->bGrassUpdateEnabled)
 		{
 			return;
diff --git a/Engine/Source/Runtime/Landscape/Private/LandscapeGrassMapsBuilder.cpp b/Engine/Source/Runtime/Landscape/Private/LandscapeGrassMapsBuilder.cpp
--- a/Engine/Source/Runtime/Landscape/Private/LandscapeGrassMapsBuilder.cpp
+++ b/Engine/Source/Runtime/Landscape/Private/LandscapeGrassMapsBuilder.cpp
@@ -310,15 +310,7 @@ FLandscapeGrassMapsBuilder::~FLandscapeGrassMapsBuilder()
 		// if any async fetch tasks are in flight, force completion
 		if (AsyncFetchCount > 0)
 		{
-			for (auto It = ComponentStates.CreateIterator(); It; ++It)
-			{
-				FComponentState* State = It.Value();
-				if (State->Stage == EComponentStage::AsyncFetch)
-				{
-					check(State->AsyncFetchTask.Get());
-					State->AsyncFetchTask->EnsureCompletion(/* bDoWorkOnThisThreadIfNotStarted= */ true, /* bIsLatencySensitive= */ true);
-				}
-			}
+			CompleteAllAsyncTasksNow();
 		}
 
 		Iterations++;
@@ -679,6 +671,9 @@ void FLandscapeGrassMapsBuilder::RegisterComponent(ULandscapeComponent* Componen
 		PendingCount++;
 
 		GRASS_DEBUG_LOG(TEXT("Register %s (%d total)"), *Component->GetName(), ComponentStates.Num());
+
+		// immediately after a new registration, check if fast paths apply
+		TryFastpathsFromPending(*NewState, /* bRecalculateHashes = */ false);
 	}
 }
 
@@ -813,15 +808,20 @@ bool FLandscapeGrassMapsBuilder::BuildGrassMapsNowForComponents(
 			{
 				// Start tracking to kick off the build process
 				constexpr bool bForceCompileShaders = true;
-				if ((AvailableStreamingSlots > 0) && StartGrassMapGeneration(*State, bForceCompileShaders))
+				if (AvailableStreamingSlots > 0)
 				{
-					// modification isn't complete yet, but convenient to dirty the package when starting the process here
-					if (bMarkDirty)
+					StartGrassMapGeneration(*State, bForceCompileShaders);
+					if (State->Stage != EComponentStage::NotReady &&
+						State->Stage != EComponentStage::Pending)
 					{
-						Component->MarkPackageDirty();
+						// modification isn't complete yet, but convenient to dirty the package when starting the process here
+						if (bMarkDirty)
+						{
+							Component->MarkPackageDirty();
+						}
+						AvailableStreamingSlots--;
+						bChanged = true;
 					}
-					AvailableStreamingSlots--;
-					bChanged = true;
 				}
 			}
 			if (State->Stage == EComponentStage::GrassMapsPopulated)
@@ -853,13 +853,18 @@ bool FLandscapeGrassMapsBuilder::BuildGrassMapsNowForComponents(
 			// not sure why this is, something must be getting starved in the non-blocking path.
 			UE::Landscape::SubmitGPUCommands(/* bBlockUntilComplete =  */ true);
 		}
-
+		
 		// If any streaming is in flight, do a blocking texture streaming update.
 		// TODO [chris.tchou] : ideally this would be a non-blocking streaming update tick, so we can react to other updates finishing
 		if (StreamingCount > 0)
 		{
 			TextureStreamingManager.WaitForTextureStreaming();
 		}
+
+		if (AsyncFetchCount > 0)
+		{
+			CompleteAllAsyncTasksNow();
+		}
 	}
 
 	UE_LOG(LogGrass, Verbose, TEXT("BuildGrassMapsNowForComponents() updated %d/%d components in %f seconds"), UpToDateCount, LandscapeComponents.Num(), FPlatformTime::Seconds() - StartTime);
@@ -867,6 +872,19 @@ bool FLandscapeGrassMapsBuilder::BuildGrassMapsNowForComponents(
 	return (UpToDateCount == LandscapeComponents.Num());
 }
 
+void FLandscapeGrassMapsBuilder::CompleteAllAsyncTasksNow()
+{
+	for (auto It = ComponentStates.CreateIterator(); It; ++It)
+	{
+		FComponentState* State = It.Value();
+		if (State->Stage == EComponentStage::AsyncFetch)
+		{
+			check(State->AsyncFetchTask.Get());
+			State->AsyncFetchTask->EnsureCompletion(/* bDoWorkOnThisThreadIfNotStarted= */ true, /* bIsLatencySensitive= */ true);
+		}
+	}
+}
+
 bool FLandscapeGrassMapsBuilder::CancelAndEvict(FComponentState& State)
 {
 	TRACE_CPUPROFILER_EVENT_SCOPE(FLandscapeGrassMapsBuilder::CancelAndEvict);
@@ -983,11 +1001,9 @@ bool FLandscapeGrassMapsBuilder::CancelAndEvict(FComponentState& State)
 	return true;
 }
 
-bool FLandscapeGrassMapsBuilder::StartGrassMapGeneration(FComponentState& State, bool bForceCompileShaders)
+bool FLandscapeGrassMapsBuilder::TryFastpathsFromPending(FComponentState& State, bool bRecalculateHashes)
 {
-	TRACE_CPUPROFILER_EVENT_SCOPE(FLandscapeGrassMapsBuilder::StartGrassMapGeneration);
 	check(State.Stage == EComponentStage::Pending);
-
 	ULandscapeComponent* Component = State.Component;
 
 	if (World->IsGameWorld()) // including PIE
@@ -997,18 +1013,21 @@ bool FLandscapeGrassMapsBuilder::StartGrassMapGeneration(FComponentState& State,
 		{
 			GRASS_DEBUG_LOG(TEXT("GrassMap Proxy DisableRuntimeGeneration for %s"), *Component->GetName());
 			PendingToNotReady(State);
-			return false;
+			return true;
 		}
 	}
 	else
 	{
 #if WITH_EDITOR
 		// recalculate hashes
-		if (UMaterialInterface* Material = Component->GetLandscapeMaterial())
+		if (bRecalculateHashes)
 		{
-			State.GrassMapGenerationHash = UE::Landscape::ComputeGrassMapGenerationHash(Component, Material);
-			const TArray<TObjectPtr<ULandscapeGrassType>>& GrassTypes = Material->GetCachedExpressionData().GrassTypes;
-			State.GrassInstanceGenerationHash = UE::Landscape::ComputeGrassInstanceGenerationHash(State.GrassMapGenerationHash, GrassTypes);
+			if (UMaterialInterface* Material = Component->GetLandscapeMaterial())
+			{
+				State.GrassMapGenerationHash = UE::Landscape::ComputeGrassMapGenerationHash(Component, Material);
+				const TArray<TObjectPtr<ULandscapeGrassType>>& GrassTypes = Material->GetCachedExpressionData().GrassTypes;
+				State.GrassInstanceGenerationHash = UE::Landscape::ComputeGrassInstanceGenerationHash(State.GrassMapGenerationHash, GrassTypes);
+			}
 		}
 
 		// in editor, if the existing grass data is valid and has a matching hash, then we can skip straight to Populated
@@ -1027,7 +1046,22 @@ bool FLandscapeGrassMapsBuilder::StartGrassMapGeneration(FComponentState& State,
 		PendingToPopulatedFastPathNoGrass(State);
 		return true;
 	}
-	
+
+	return false;
+}
+
+bool FLandscapeGrassMapsBuilder::StartGrassMapGeneration(FComponentState& State, bool bForceCompileShaders)
+{
+	TRACE_CPUPROFILER_EVENT_SCOPE(FLandscapeGrassMapsBuilder::StartGrassMapGeneration);
+	check(State.Stage == EComponentStage::Pending);
+
+	ULandscapeComponent* Component = State.Component;
+
+	if (TryFastpathsFromPending(State, /* bRecalculateHashes = */ true))
+	{
+		return false;
+	}
+
 	// if we can't currently render, it's not ready
 	if (!UE::Landscape::CanRenderGrassMap(Component))
 	{
diff --git a/Engine/Source/Runtime/Landscape/Private/LandscapeGrassMapsBuilder.h b/Engine/Source/Runtime/Landscape/Private/LandscapeGrassMapsBuilder.h
--- a/Engine/Source/Runtime/Landscape/Private/LandscapeGrassMapsBuilder.h
+++ b/Engine/Source/Runtime/Landscape/Private/LandscapeGrassMapsBuilder.h
@@ -237,9 +237,12 @@ private:
 	// tries to cancel any in flight operations and transition back to the Pending state
 	bool CancelAndEvict(FComponentState& State);
 
-	// try to kick off the grass map generation pipeline -- returns true if texture streaming was started, false if the shaders are not ready, or if it shortcut to populated because data exists already
+	// try to kick off the grass map generation pipeline -- returns true if it started the amortized update path, false otherwise.
 	bool StartGrassMapGeneration(FComponentState& State, bool bForceCompileShaders);
 
+	// try to apply fast path transitions to a pending component -- returns true if a fastpath was taken, and the component is no longer pending.
+	bool TryFastpathsFromPending(FComponentState& State, bool bRecalculateHashes);
+
 	// once textures are streamed, this kicks off the grass data render, and async GPU readback
 	void KickOffRenderAndReadback(FComponentState& State);
 
@@ -259,6 +262,8 @@ private:
 	void PendingToStreaming(FComponentState& State);
 	void RemoveFromPendingComponentHeap(FComponentState* State);
 
+	void CompleteAllAsyncTasksNow();
+
 #if WITH_EDITOR
 	// cached count of how many grass maps are outdated, and the last time we calculated that value
 	mutable int32 OutdatedGrassMapCount = 0;
