Fix water info rendering issues with custom render pass method
- Fix multiple water zones (don't apply view relevance overrides so custom render pass renders in main pass, custom render pass has no view state or view id so track them by scene view pointer)
- Fix packaged build (render landscape with dynamic rendering path)

diff --git a/Engine/Source/Runtime/Landscape/Private/LandscapeCulling.cpp b/Engine/Source/Runtime/Landscape/Private/LandscapeCulling.cpp
--- a/Engine/Source/Runtime/Landscape/Private/LandscapeCulling.cpp
+++ b/Engine/Source/Runtime/Landscape/Private/LandscapeCulling.cpp
@@ -451,7 +451,7 @@ struct FCullingEntry
 struct FCullingSystem
 {
 	TArray<FCullingEntry> Landscapes; 
-	TArray<uint32, TInlineAllocator<2>> ViewStateKeys;
+	TArray<const FSceneView*, TInlineAllocator<2>> Views;
 };
 
 FCullingSystem GCullingSystem;
@@ -473,7 +473,7 @@ static void ResetIntermediateData()
 		CullingEntry.CullingArguments.Reset();
 	}
 
-	GCullingSystem.ViewStateKeys.Reset();
+	GCullingSystem.Views.Reset();
 }
 
 void PreRenderViewFamily(FSceneViewFamily& InViewFamily)
@@ -605,9 +605,7 @@ static void ComputeSectionIntermediateData(FRDGBuilder& GraphBuilder, TArrayView
 	for (int32 ViewIdx = 0; ViewIdx < Views.Num(); ++ViewIdx)
 	{
 		const FSceneView& View = *Views[ViewIdx];
-		uint32 ViewKey = View.GetViewKey();
-
-		GCullingSystem.ViewStateKeys.Add(ViewKey);
+		GCullingSystem.Views.Add(&View);
 		
 		// Collect all LOD0 sections for each ViewState+Landscape and upload to GPU
 		for (int32 LandscapeIdx = 0; LandscapeIdx < GCullingSystem.Landscapes.Num(); ++LandscapeIdx)
@@ -618,7 +616,7 @@ static void ComputeSectionIntermediateData(FRDGBuilder& GraphBuilder, TArrayView
 			ViewStateIntermediates.SectionsBufferRDG = nullptr;
 
 			const FLandscapeRenderSystem& RenderSystem = *LandscapeRenderSystems.FindChecked(CullingEntry.LandscapeKey);
-			const TResourceArray<float>& SectionLODValues = RenderSystem.CachedSectionLODValues.FindRef(ViewKey);
+			const TResourceArray<float>& SectionLODValues = RenderSystem.CachedSectionLODValues.FindRef(&View);
 
 			for (int32 SectionIdx = 0; SectionIdx < SectionLODValues.Num(); ++SectionIdx)
 			{
@@ -658,9 +656,10 @@ static void ComputeSectionIntermediateData(FRDGBuilder& GraphBuilder, TArrayView
 static void DispatchCulling(FRDGBuilder& GraphBuilder, TArrayView<const FSceneView*> CullingViews, TArrayView<FViewMatrices> CullingViewMatrices, bool bNearClip)
 {
 	// A separate dispatch for each Landscape and ViewState
-	for (int32 ViewStateIdx = 0; ViewStateIdx < GCullingSystem.ViewStateKeys.Num(); ++ViewStateIdx)
+	for (int32 ViewStateIdx = 0; ViewStateIdx < GCullingSystem.Views.Num(); ++ViewStateIdx)
 	{
-		uint32 ViewStateKey = GCullingSystem.ViewStateKeys[ViewStateIdx];
+		const FSceneView* View = GCullingSystem.Views[ViewStateIdx];
+		uint32 ViewStateKey = View->GetViewKey();
 
 		TArray<FBuildLandscapeTileDataCS::FLandscapeView, TInlineAllocator<8>> CullingViewsData;
 		CullingViewsData.Reserve(CullingViews.Num());
@@ -671,7 +670,7 @@ static void DispatchCulling(FRDGBuilder& GraphBuilder, TArrayView<const FSceneVi
 		for (int32 CullingViewIdx = 0; CullingViewIdx < CullingViews.Num(); ++CullingViewIdx)
 		{
 			const FSceneView& CullingView = *CullingViews[CullingViewIdx];
-			if (CullingView.GetViewKey() != ViewStateKey)
+			if (&CullingView != View)
 			{
 				continue;
 			}
@@ -792,7 +791,7 @@ void InitMainViews(FRDGBuilder& GraphBuilder, TArrayView<const FSceneView*> View
 
 	QUICK_SCOPE_CYCLE_COUNTER(STAT_LandscapeCulling_InitMainViews);
 
-	check(GCullingSystem.ViewStateKeys.Num() == 0);
+	check(GCullingSystem.Views.Num() == 0);
 
 	ComputeSectionIntermediateData(GraphBuilder, Views);
 
diff --git a/Engine/Source/Runtime/Landscape/Private/LandscapeRender.cpp b/Engine/Source/Runtime/Landscape/Private/LandscapeRender.cpp
--- a/Engine/Source/Runtime/Landscape/Private/LandscapeRender.cpp
+++ b/Engine/Source/Runtime/Landscape/Private/LandscapeRender.cpp
@@ -893,7 +893,7 @@ const TResourceArray<float>& FLandscapeRenderSystem::ComputeSectionsLODForView(c
 {
 	TRACE_CPUPROFILER_EVENT_SCOPE(FLandscapeRenderSystem::ComputeSectionsLODForView);
 
-	TResourceArray<float>& SectionLODValues = CachedSectionLODValues.Add(InView.GetViewKey());
+	TResourceArray<float>& SectionLODValues = CachedSectionLODValues.Add(&InView);
 	SectionLODValues.AddZeroed(SectionInfos.Num());
 
 	for (int32 SectionIndex = 0; SectionIndex < SectionInfos.Num(); SectionIndex++)
@@ -1875,7 +1875,8 @@ FPrimitiveViewRelevance FLandscapeComponentSceneProxy::GetViewRelevance(const FS
 		(CVarLandscapeShowDirty.GetValueOnAnyThread() && GLandscapeDirtyMaterial) ||
 		(GetViewLodOverride(*View, LandscapeKey) >= 0)
 #else
-		IsSelected()
+		IsSelected() ||
+		(View->CustomRenderPass && GetViewLodOverride(*View, LandscapeKey) >= 0)
 #endif
 		)
 	{
@@ -1895,7 +1896,7 @@ FPrimitiveViewRelevance FLandscapeComponentSceneProxy::GetViewRelevance(const FS
 	}
 #endif // !UE_BUILD_SHIPPING
 
-	if (bNaniteActive && View->Family->EngineShowFlags.NaniteMeshes)
+	if (bNaniteActive && View->Family->EngineShowFlags.NaniteMeshes && !View->CustomRenderPass)
 	{
 		Result.bShadowRelevance = false;
 		Result.bVelocityRelevance = false;
@@ -2581,7 +2582,8 @@ void FLandscapeComponentSceneProxy::GetDynamicMeshElements(const TArray<const FS
 						(IsSelected() && !GLandscapeEditModeActive) ||
 						(GetViewLodOverride(*View, LandscapeKey) >= 0)
 #else
-						IsSelected()
+						IsSelected() ||
+						(View->CustomRenderPass && GetViewLodOverride(*View, LandscapeKey) >= 0)
 #endif
 						)
 					{
diff --git a/Engine/Source/Runtime/Landscape/Public/LandscapeRender.h b/Engine/Source/Runtime/Landscape/Public/LandscapeRender.h
--- a/Engine/Source/Runtime/Landscape/Public/LandscapeRender.h
+++ b/Engine/Source/Runtime/Landscape/Public/LandscapeRender.h
@@ -486,7 +486,7 @@ struct FLandscapeRenderSystem
 
 	FUniformBufferRHIRef SectionLODUniformBuffer;
 
-	TMap<FViewKey, TResourceArray<float>> CachedSectionLODValues;
+	TMap<const FSceneView*, TResourceArray<float>> CachedSectionLODValues;
 
 	/** Forced LOD level which overrides the ForcedLOD level of all the sections under this LandscapeRenderSystem. */
 	int8 ForcedLODOverride;
@@ -546,7 +546,7 @@ struct FLandscapeRenderSystem
 
 	float GetSectionLODValue(const FSceneView& SceneView, FIntPoint InRenderCoord) const
 	{
-		return CachedSectionLODValues[SceneView.GetViewKey()][GetSectionLinearIndex(InRenderCoord)];
+		return CachedSectionLODValues[&SceneView][GetSectionLinearIndex(InRenderCoord)];
 	}
 
 	float GetSectionLODBias(FIntPoint InRenderCoord) const
