Fixed landscape crash (on behalf of josie.yang)

diff --git a/Engine/Source/Runtime/Landscape/Public/LandscapeRender.h b/Engine/Source/Runtime/Landscape/Public/LandscapeRender.h
--- a/Engine/Source/Runtime/Landscape/Public/LandscapeRender.h
+++ b/Engine/Source/Runtime/Landscape/Public/LandscapeRender.h
@@ -546,7 +546,23 @@ struct FLandscapeRenderSystem
 
 	float GetSectionLODValue(const FSceneView& SceneView, FIntPoint InRenderCoord) const
 	{
-		return CachedSectionLODValues[&SceneView][GetSectionLinearIndex(InRenderCoord)];
+		if (CachedSectionLODValues.Contains(&SceneView))
+		{
+			return CachedSectionLODValues[&SceneView][GetSectionLinearIndex(InRenderCoord)];
+		}
+		else
+		{
+			// Some views such as shadow are spawned off main view and share the same view id
+			for (const TPair<const FSceneView*, TResourceArray<float>>& Pair : CachedSectionLODValues)
+			{
+				if (Pair.Key->GetViewKey() == SceneView.GetViewKey())
+				{
+					return Pair.Value[GetSectionLinearIndex(InRenderCoord)];
+				}
+			}
+		}
+		checkf(false, TEXT("No section LOD value cached for this view. Make sure FLandscapeRenderSystem::ComputeSectionsLODForView (FLandscapeSceneViewExtension::PreRenderView_RenderThread) was called"));
+		return 0.0f;
 	}
 
 	float GetSectionLODBias(FIntPoint InRenderCoord) const
