Don't flush FMessageLog until after the message has been populated with data
#rnx

diff --git a/Engine/Source/Runtime/Landscape/Private/Landscape.cpp b/Engine/Source/Runtime/Landscape/Private/Landscape.cpp
--- a/Engine/Source/Runtime/Landscape/Private/Landscape.cpp
+++ b/Engine/Source/Runtime/Landscape/Private/Landscape.cpp
@@ -4311,18 +4311,19 @@ void ALandscapeProxy::PostLoad()
 		FFormatNamedArguments Arguments;
 		Arguments.Add(TEXT("LandscapeProxyName"), FText::FromString(GetActorNameOrLabel())); 
 
-		auto CreateMapCheckMessage = []()
+		auto CreateMapCheckMessage = [](FMessageLog& MessageLog)
 		{
 			if (CVarLandscapeSupressMapCheckWarnings_Nanite->GetBool())
 			{
-				return FMessageLog("MapCheck").Info();	
+				return MessageLog.Info();	
 			}
-			return FMessageLog("MapCheck").Warning();
+			return MessageLog.Warning();
 		};
 		
 		TWeakObjectPtr<ALandscapeProxy> WeakLandscapeProxy(this);
 
-		CreateMapCheckMessage()
+		FMessageLog MessageLog("MapCheck");
+		CreateMapCheckMessage(MessageLog)
 			->AddToken(FTextToken::Create(FText::Format(LOCTEXT("MapCheck_Message_LandscapeRebuildNanite", "{LandscapeProxyName} : Landscape Nanite is enabled but saved mesh data is out of date. "), Arguments)))
 			->AddToken(FActionToken::Create(LOCTEXT("MapCheck_SaveFixedUpData", "Save Modified Landscapes"), LOCTEXT("MapCheck_SaveFixedUpData_Desc", "Saves the modified landscape proxy actors"),
 				FOnActionTokenExecuted::CreateLambda([WeakLandscapeProxy]()
