Added the ability to debug draw grass exclusion volumes via the grass.DrawExclusionVolumes cheat.

Also added the number of grass exclusion volumes to `stat foliage` since the exclusion volumes can really hit performance negatively right now.

#tests PIE, Win64 Client
#rnx

diff --git a/Engine/Source/Runtime/Landscape/Classes/LandscapeProxy.h b/Engine/Source/Runtime/Landscape/Classes/LandscapeProxy.h
--- a/Engine/Source/Runtime/Landscape/Classes/LandscapeProxy.h
+++ b/Engine/Source/Runtime/Landscape/Classes/LandscapeProxy.h
@@ -1030,6 +1030,8 @@ public:
 
 	static void RemoveInvalidExclusionBoxes();
 
+	static void DebugDrawExclusionBoxes(const UWorld* World);
+
 	/* Invalidate the precomputed grass and baked texture data for the specified components */
 	LANDSCAPE_API static void InvalidateGeneratedComponentData(const TSet<ULandscapeComponent*>& Components, bool bInvalidateLightingCache = false);
 	LANDSCAPE_API static void InvalidateGeneratedComponentData(const TArray<ULandscapeComponent*>& Components, bool bInvalidateLightingCache = false);
diff --git a/Engine/Source/Runtime/Landscape/Private/LandscapeGrass.cpp b/Engine/Source/Runtime/Landscape/Private/LandscapeGrass.cpp
--- a/Engine/Source/Runtime/Landscape/Private/LandscapeGrass.cpp
+++ b/Engine/Source/Runtime/Landscape/Private/LandscapeGrass.cpp
@@ -78,6 +78,10 @@
 #include "MaterialHLSLGenerator.h"
 #endif
 
+#if !UE_BUILD_SHIPPING
+#include "DrawDebugHelpers.h"
+#endif // !UE_BUILD_SHIPPING
+
 #define LOCTEXT_NAMESPACE "Landscape"
 
 DEFINE_LOG_CATEGORY(LogGrass);
@@ -797,6 +801,14 @@ static FAutoConsoleCommand ConsoleCommandDumpLandscapeGrassData(
 		}
 	}));
 
+#if !UE_BUILD_SHIPPING
+static bool bGGrassDrawExclusionVolumes = false;
+static FAutoConsoleVariableRef CVarGrassDrawExclusionVolumes(
+	TEXT("grass.DrawExclusionVolumes"),
+	bGGrassDrawExclusionVolumes,
+	TEXT("Whether we should draw the exclusion volumes or not"));
+#endif // !UE_BUILD_SHIPPING
+
 struct FPerQualityLevelInt;
 struct FPerQualityLevelFloat;
 
@@ -806,6 +818,7 @@ DECLARE_CYCLE_STAT(TEXT("Grass End Comp"), STAT_FoliageGrassEndComp, STATGROUP_F
 DECLARE_CYCLE_STAT(TEXT("Grass Destroy Comps"), STAT_FoliageGrassDestoryComp, STATGROUP_Foliage);
 DECLARE_CYCLE_STAT(TEXT("Grass Update"), STAT_GrassUpdate, STATGROUP_Foliage);
 DECLARE_CYCLE_STAT(TEXT("Grass Type Update Summary"), STAT_UpdateGrassTypeSummary, STATGROUP_Foliage);
+DECLARE_DWORD_ACCUMULATOR_STAT(TEXT("Grass Exclusion Volumes"), STAT_GrassExclusionVolumes, STATGROUP_Foliage);
 
 int32 ALandscapeProxy::GrassUpdateInterval = 1;
 
@@ -2404,16 +2417,22 @@ static TMap<FWeakObjectPtr, FBox, FDefaultSetAllocator, TWeakObjectPtrMapKeyFunc
 
 void ALandscapeProxy::AddExclusionBox(FWeakObjectPtr Owner, const FBox& Box)
 {
+	INC_DWORD_STAT(STAT_GrassExclusionVolumes);
+
 	GGrassExclusionBoxes.Add(Owner, Box);
 	GGrassExclusionChangeTag++;
 }
 void ALandscapeProxy::RemoveExclusionBox(FWeakObjectPtr Owner)
 {
+	DEC_DWORD_STAT(STAT_GrassExclusionVolumes);
+
 	GGrassExclusionBoxes.Remove(Owner);
 	GGrassExclusionChangeTag++;
 }
 void ALandscapeProxy::RemoveAllExclusionBoxes()
 {
+	SET_DWORD_STAT(STAT_GrassExclusionVolumes, 0);
+
 	if (GGrassExclusionBoxes.Num())
 	{
 		GGrassExclusionBoxes.Empty();
@@ -2427,6 +2446,8 @@ void ALandscapeProxy::RemoveInvalidExclusionBoxes()
 	{
 		if (!Iter->Key.IsValid())
 		{
+			DEC_DWORD_STAT(STAT_GrassExclusionVolumes);
+
 			Iter.RemoveCurrent();
 			GGrassExclusionChangeTag++;
 		}
@@ -2435,6 +2456,23 @@ void ALandscapeProxy::RemoveInvalidExclusionBoxes()
 	GGrassExclusionBoxes.Compact();
 }
 
+void ALandscapeProxy::DebugDrawExclusionBoxes(const UWorld* InWorld)
+{
+#if !UE_BUILD_SHIPPING
+	if (bGGrassDrawExclusionVolumes)
+	{
+		for (auto Iter = GGrassExclusionBoxes.CreateIterator(); Iter; ++Iter)
+		{
+			if (Iter->Key.IsValid())
+			{
+				const FBox& Box = Iter->Value;
+				DrawDebugBox(InWorld, Box.GetCenter(), Box.GetExtent(), FColor::Red);
+			}
+		}
+	}
+#endif // !UE_BUILD_SHIPPING
+}
+
 #if WITH_EDITOR
 
 void ALandscapeProxy::BuildGrassMaps(FScopedSlowTask* InSlowTask)
diff --git a/Engine/Source/Runtime/Landscape/Private/LandscapeSubsystem.cpp b/Engine/Source/Runtime/Landscape/Private/LandscapeSubsystem.cpp
--- a/Engine/Source/Runtime/Landscape/Private/LandscapeSubsystem.cpp
+++ b/Engine/Source/Runtime/Landscape/Private/LandscapeSubsystem.cpp
@@ -462,6 +462,8 @@ void ULandscapeSubsystem::Tick(float DeltaTime)
 	}
 
 	ActiveProxies.Reset();
+	
+	ALandscapeProxy::DebugDrawExclusionBoxes(World);
 
 #if WITH_EDITOR
 	if (GIsEditor && !World->IsPlayInEditor())
