Landscape - Added HLOD settings
* Will allow better control on the HLOD mesh complexity and it's texture size
* Legacy default was to automatically compute the required size
* New default is to use the last LOD and a 256x256 texture
* Added a project wide maximum texture size for the landscape HLODs textures

diff --git a/Engine/Source/Runtime/Core/Public/UObject/FortniteMainBranchObjectVersion.h b/Engine/Source/Runtime/Core/Public/UObject/FortniteMainBranchObjectVersion.h
--- a/Engine/Source/Runtime/Core/Public/UObject/FortniteMainBranchObjectVersion.h
+++ b/Engine/Source/Runtime/Core/Public/UObject/FortniteMainBranchObjectVersion.h
@@ -408,6 +408,9 @@ struct FFortniteMainBranchObjectVersion
 		// Changes to Orthographic Camera default settings
 		OrthographicCameraDefaultSettings,
 
+		// Added settings to Landscape HLODs
+		LandscapeAddedHLODSettings,
+
 		// -----<new versions can be added above this line>-------------------------------------------------
 		VersionPlusOne,
 		LatestVersion = VersionPlusOne - 1
diff --git a/Engine/Source/Runtime/Landscape/Classes/LandscapeProxy.h b/Engine/Source/Runtime/Landscape/Classes/LandscapeProxy.h
--- a/Engine/Source/Runtime/Landscape/Classes/LandscapeProxy.h
+++ b/Engine/Source/Runtime/Landscape/Classes/LandscapeProxy.h
@@ -183,6 +183,29 @@ enum class ELandscapeLayerDisplayMode : uint8
 	UserSpecific,
 };
 
+UENUM()
+enum class ELandscapeHLODTextureSizePolicy : uint8
+{
+	/** Automatic texture size, based on the expected HLOD draw distance and the landscape size. */
+	AutomaticSize,
+
+	/** User specified texture size. */
+	SpecificSize
+};
+
+UENUM()
+enum class ELandscapeHLODMeshSourceLODPolicy : uint8
+{
+	/** Automatic LOD selection, based on the expected HLOD draw distance and the landscape LOD Distribution settings. */
+	AutomaticLOD,
+
+	/** User specified landscape LOD. */
+	SpecificLOD,
+
+	/** Use the lowest detailed LOD of this landscape. */
+	LowestDetailLOD
+};
+
 UENUM()
 namespace ELandscapeLODFalloff
 {
@@ -811,6 +834,24 @@ public:
 	UPROPERTY(EditAnywhere, Category = HLOD, meta = (LandscapeOverridable))
 	bool bUseLandscapeForCullingInvisibleHLODVertices;
 
+#if WITH_EDITORONLY_DATA
+	/** Specify how to choose the texture size of the resulting HLOD mesh */
+	UPROPERTY(EditAnywhere, Category = HLOD, meta = (DisplayName = "HLOD Texture Size Policy", LandscapeOverridable))
+	ELandscapeHLODTextureSizePolicy HLODTextureSizePolicy;
+
+	/** Specify the texture size to use for the HLOD mesh if HLODTextureSizePolicy is set to SpecificSize */
+	UPROPERTY(EditAnywhere, Category = HLOD, meta = (DisplayName = "HLOD Texture Size", LandscapeOverridable, EditCondition = "HLODTextureSizePolicy == ELandscapeHLODTextureSizePolicy::SpecificSize", EditConditionHides, ClampMin = "16", ClampMax = "8192"))
+	int32 HLODTextureSize;
+
+	/** Specify how to choose the LOD used as input for the HLOD mesh */
+	UPROPERTY(EditAnywhere, Category = HLOD, meta = (DisplayName = "HLOD Mesh Source LOD Policy", LandscapeOverridable))
+	ELandscapeHLODMeshSourceLODPolicy HLODMeshSourceLODPolicy;
+
+	/** Specify which LOD to use for the HLOD mesh if HLODMeshSourceLODPolicy is set to SpecificLOD */
+	UPROPERTY(EditAnywhere, Category = HLOD, meta = (DisplayName = "HLOD Mesh Source LOD", LandscapeOverridable, EditCondition = "HLODMeshSourceLODPolicy == ELandscapeHLODMeshSourceLODPolicy::SpecificLOD", EditConditionHides, ClampMin = "0"))
+	int32 HLODMeshSourceLOD;
+#endif
+
 	/** Flag that tell if we have some layers content **/
 	UPROPERTY()
 	bool bHasLayersContent;
diff --git a/Engine/Source/Runtime/Landscape/Private/Landscape.cpp b/Engine/Source/Runtime/Landscape/Private/Landscape.cpp
--- a/Engine/Source/Runtime/Landscape/Private/Landscape.cpp
+++ b/Engine/Source/Runtime/Landscape/Private/Landscape.cpp
@@ -1552,6 +1552,10 @@ ALandscapeProxy::ALandscapeProxy(const FObjectInitializer& ObjectInitializer)
 #if WITH_EDITORONLY_DATA
 	MaxPaintedLayersPerComponent = 0;
 	bHasLayersContent = false;
+	HLODTextureSizePolicy = ELandscapeHLODTextureSizePolicy::SpecificSize;
+	HLODTextureSize = 256;
+	HLODMeshSourceLODPolicy = ELandscapeHLODMeshSourceLODPolicy::LowestDetailLOD;
+	HLODMeshSourceLOD = 0;
 #endif
 
 #if WITH_EDITOR
@@ -4400,6 +4404,13 @@ void ALandscapeProxy::PostLoad()
 		// Remove RF_Transactional from Nanite components : they're re-created upon transacting now : 
 		ClearNaniteTransactional();
 	}
+
+	// Keep previous behavior of landscape HLODs if created before the settings were added
+	if (GetLinkerCustomVersion(FFortniteMainBranchObjectVersion::GUID) < FFortniteMainBranchObjectVersion::LandscapeAddedHLODSettings)
+	{
+		HLODTextureSizePolicy = ELandscapeHLODTextureSizePolicy::AutomaticSize;
+		HLODMeshSourceLODPolicy = ELandscapeHLODMeshSourceLODPolicy::AutomaticLOD;
+	}
 #endif // WITH_EDITOR
 }
 
diff --git a/Engine/Source/Runtime/Landscape/Private/LandscapeHLODBuilder.cpp b/Engine/Source/Runtime/Landscape/Private/LandscapeHLODBuilder.cpp
--- a/Engine/Source/Runtime/Landscape/Private/LandscapeHLODBuilder.cpp
+++ b/Engine/Source/Runtime/Landscape/Private/LandscapeHLODBuilder.cpp
@@ -8,9 +8,10 @@
 #if WITH_EDITOR
 
 #include "Landscape.h"
-#include "LandscapeProxy.h"
 #include "LandscapeComponent.h"
 #include "LandscapeMeshProxyComponent.h"
+#include "LandscapeProxy.h"
+#include "LandscapeSettings.h"
 
 #include "MeshDescription.h"
 #include "PhysicsEngine/BodySetup.h"
@@ -42,6 +43,10 @@ uint32 ULandscapeHLODBuilder::ComputeHLODHash(const UActorComponent* InSourceCom
 {
 	FArchiveCrc32 Ar;
 
+	// Base lanscape HLOD key, changing this will force a rebuild of all landscape HLODs
+	FString HLODBaseKey = "38DC3700FC0742929BA00ACCF5B1B626";
+	Ar << HLODBaseKey;
+
 	if (const ULandscapeComponent* LSComponent = Cast<ULandscapeComponent>(InSourceComponent))
 	{
 		UE_LOG(LogHLODBuilder, VeryVerbose, TEXT("ULandscapeHLODBuilder::ComputeHLODHash %s"), *InSourceComponent->GetName());
@@ -113,6 +118,35 @@ uint32 ULandscapeHLODBuilder::ComputeHLODHash(const UActorComponent* InSourceCom
 		// Nanite enabled?
 		bool bNaniteEnabled = LSProxy->IsNaniteEnabled();
 		Ar << bNaniteEnabled;
+		UE_LOG(LogHLODBuilder, VeryVerbose, TEXT("     - NaniteEnabled = %d"), bNaniteEnabled);
+
+		// HLODTextureSize
+		ELandscapeHLODTextureSizePolicy HLODTextureSizePolicy = LSProxy->HLODTextureSizePolicy;
+		Ar << HLODTextureSizePolicy;
+		UE_LOG(LogHLODBuilder, VeryVerbose, TEXT("     - HLODTextureSizePolicy = %x"), HLODTextureSizePolicy);
+		if (HLODTextureSizePolicy == ELandscapeHLODTextureSizePolicy::SpecificSize)
+		{
+			int32 HLODTextureSize = LSProxy->HLODTextureSize;
+			Ar << HLODTextureSize;
+			UE_LOG(LogHLODBuilder, VeryVerbose, TEXT("     - HLODTextureSize = %d"), HLODTextureSize);
+		}
+
+		// HLODMeshSourceLOD
+		ELandscapeHLODMeshSourceLODPolicy HLODMeshSourceLODPolicy = LSProxy->HLODMeshSourceLODPolicy;
+		Ar << HLODMeshSourceLODPolicy;
+		UE_LOG(LogHLODBuilder, VeryVerbose, TEXT("     - HLODMeshSourceLODPolicy = %x"), HLODMeshSourceLODPolicy);
+		if (HLODMeshSourceLODPolicy == ELandscapeHLODMeshSourceLODPolicy::SpecificLOD)
+		{
+			int32 HLODMeshSourceLOD = LSProxy->HLODMeshSourceLOD;
+			Ar << HLODMeshSourceLOD;
+			UE_LOG(LogHLODBuilder, VeryVerbose, TEXT("     - HLODMeshSourceLOD = %d"), HLODMeshSourceLOD);
+		}
+
+		// Project max texture size for landscape HLOD textures
+		const ULandscapeSettings* LandscapeSettings = GetDefault<ULandscapeSettings>();
+		int32 ProjectHLODMaxTextureSize = LandscapeSettings->GetHLODMaxTextureSize();
+		Ar << ProjectHLODMaxTextureSize;
+		UE_LOG(LogHLODBuilder, VeryVerbose, TEXT("     - ProjectHLODMaxTextureSize = %d"), ProjectHLODMaxTextureSize);
 	}
 
 	return Ar.GetCrc();
@@ -157,31 +191,83 @@ static int32 GetMeshTextureSizeFromTargetTexelDensity(const FMeshDescription& In
 static int32 ComputeRequiredLandscapeLOD(const ALandscapeProxy* InLandscapeProxy, const float InViewDistance)
 {
 	check(InLandscapeProxy && !InLandscapeProxy->LandscapeComponents.IsEmpty());
-	
-	// These constants are showing up a lot in the screen size computation for Level HLODs. This should be configurable per project.
-	const float HalfFOV = PI * 0.25f;
-	const float ScreenWidth = 1920.0f;
-	const float ScreenHeight = 1080.0f;
-	const FPerspectiveMatrix ProjMatrix(HalfFOV, ScreenWidth, ScreenHeight, 1.0f);
-
-	TArray<float> LODScreenSizes = InLandscapeProxy->GetLODScreenSizeArray();
 
-	const ULandscapeComponent* LSComponent = InLandscapeProxy->LandscapeComponents[0];
-	const float ComponentRadiusScaled = static_cast<float>(LSComponent->GetLocalBounds().SphereRadius * LSComponent->GetComponentTransform().GetScale3D().GetAbsMax());
-	const float ExpectedScreenSize = ComputeBoundsScreenSize(FVector::ZeroVector, ComponentRadiusScaled, FVector(0.0f, 0.0f, InViewDistance), ProjMatrix);
+	const TArray<float> LODScreenSizes = InLandscapeProxy->GetLODScreenSizeArray();
+	int32 RequiredLOD = 0;	
 
-	int32 RequiredLOD;
-	for (RequiredLOD = 0; RequiredLOD < LODScreenSizes.Num(); ++RequiredLOD)
+	switch (InLandscapeProxy->HLODMeshSourceLODPolicy)
 	{
-		if (ExpectedScreenSize > LODScreenSizes[RequiredLOD])
+		case ELandscapeHLODMeshSourceLODPolicy::AutomaticLOD:
 		{
-			break;
-		}
+			// These constants are showing up a lot in the screen size computation for Level HLODs. This should be configurable per project.
+			const float HalfFOV = PI * 0.25f;
+			const float ScreenWidth = 1920.0f;
+			const float ScreenHeight = 1080.0f;
+			const FPerspectiveMatrix ProjMatrix(HalfFOV, ScreenWidth, ScreenHeight, 1.0f);
+
+			
+
+			const ULandscapeComponent* LSComponent = InLandscapeProxy->LandscapeComponents[0];
+			const float ComponentRadiusScaled = static_cast<float>(LSComponent->GetLocalBounds().SphereRadius * LSComponent->GetComponentTransform().GetScale3D().GetAbsMax());
+			const float ExpectedScreenSize = ComputeBoundsScreenSize(FVector::ZeroVector, ComponentRadiusScaled, FVector(0.0f, 0.0f, InViewDistance), ProjMatrix);
+
+			
+			for (RequiredLOD = 0; RequiredLOD < LODScreenSizes.Num(); ++RequiredLOD)
+			{
+				if (ExpectedScreenSize > LODScreenSizes[RequiredLOD])
+				{
+					break;
+				}
+			}
+		} break;
+
+		case ELandscapeHLODMeshSourceLODPolicy::SpecificLOD:
+		{
+			RequiredLOD = FMath::Clamp(InLandscapeProxy->HLODMeshSourceLOD, 0, LODScreenSizes.Num() - 1);
+		} break;
+
+		case ELandscapeHLODMeshSourceLODPolicy::LowestDetailLOD:
+		{
+			RequiredLOD = LODScreenSizes.Num() - 1;
+		} break;
 	}
 
 	return RequiredLOD;
 }
 
+static int32 ComputeRequiredTextureSize(const ALandscapeProxy* InLandscapeProxy, const float InViewDistance, const FMeshDescription* InMeshDescription)
+{
+	int32 RequiredTextureSize = 0;
+
+	switch (InLandscapeProxy->HLODTextureSizePolicy)
+	{
+		case ELandscapeHLODTextureSizePolicy::AutomaticSize:
+		{
+			const float TargetTexelDensityPerMeter = FMaterialUtilities::ComputeRequiredTexelDensityFromDrawDistance(InViewDistance, static_cast<float>(InMeshDescription->GetBounds().SphereRadius));
+			RequiredTextureSize = GetMeshTextureSizeFromTargetTexelDensity(*InMeshDescription, TargetTexelDensityPerMeter);
+		} break;
+
+		case ELandscapeHLODTextureSizePolicy::SpecificSize:
+		{
+			RequiredTextureSize = InLandscapeProxy->HLODTextureSize;
+		} break;
+	}
+
+	// Clamp to a sane minimum value
+	const int32 MinLandscapeHLODTextureSize = 16;
+	RequiredTextureSize = FMath::Max(RequiredTextureSize, MinLandscapeHLODTextureSize);
+
+	// Clamp to the project's max texture size for landscape HLODs
+	const ULandscapeSettings* LandscapeSettings = GetDefault<ULandscapeSettings>();
+	RequiredTextureSize = FMath::Min(RequiredTextureSize, LandscapeSettings->GetHLODMaxTextureSize());
+
+	// Clamp to the maximum possible texture size for safety
+	RequiredTextureSize = FMath::Min(RequiredTextureSize, (int32)GetMax2DTextureDimension());
+
+	return RequiredTextureSize;
+}
+
+
 static UMaterialInterface* BakeLandscapeMaterial(const FHLODBuildContext& InHLODBuildContext, const FMeshDescription& InMeshDescription, const ALandscapeProxy* InLandscapeProxy, const int32 InLandscapeLOD, int32 InTextureSize)
 {
 	// Build landscape material
@@ -294,11 +380,8 @@ TArray<UActorComponent*> ULandscapeHLODBuilder::Build(const FHLODBuildContext& I
 
 		// Material
 		{
-			const float TargetTexelDensityPerMeter = FMaterialUtilities::ComputeRequiredTexelDensityFromDrawDistance(
-				static_cast<float>(InHLODBuildContext.MinVisibleDistance), static_cast<float>(MeshDescription->GetBounds().SphereRadius));
-			int32 RequiredTextureSize = GetMeshTextureSizeFromTargetTexelDensity(*MeshDescription, TargetTexelDensityPerMeter);
-
-			UMaterialInterface* LandscapeMaterial = BakeLandscapeMaterial(InHLODBuildContext, *MeshDescription, LandscapeProxy, LandscapeLOD, RequiredTextureSize);
+			int32 TextureSize = ComputeRequiredTextureSize(LandscapeProxy, static_cast<float>(InHLODBuildContext.MinVisibleDistance), MeshDescription);
+			UMaterialInterface* LandscapeMaterial = BakeLandscapeMaterial(InHLODBuildContext, *MeshDescription, LandscapeProxy, LandscapeLOD, TextureSize);
 
 			//Assign the proxy material to the static mesh
 			StaticMesh->GetStaticMaterials().Add(FStaticMaterial(LandscapeMaterial));
diff --git a/Engine/Source/Runtime/Landscape/Public/LandscapeSettings.h b/Engine/Source/Runtime/Landscape/Public/LandscapeSettings.h
--- a/Engine/Source/Runtime/Landscape/Public/LandscapeSettings.h
+++ b/Engine/Source/Runtime/Landscape/Public/LandscapeSettings.h
@@ -41,6 +41,8 @@ public:
 	float GetBrushSizeUIMax() const { return BrushSizeUIMax; }
 	float GetBrushSizeClampMax() const { return BrushSizeClampMax; }
 
+	int32 GetHLODMaxTextureSize() const { return HLODMaxTextureSize; }
+
 	/** Returns the default landscape material that should be used when creating a new landscape. */
 	TSoftObjectPtr<UMaterialInterface> GetDefaultLandscapeMaterial() const { return DefaultLandscapeMaterial; }
 
@@ -87,6 +89,9 @@ protected:
 	UPROPERTY(EditAnywhere, config, Category = "Configuration", meta = (ToolTip = "Maximum size that can be set manually for the landscape sculpt/paint brushes"))
 	float BrushSizeClampMax = 65536;
 
+	UPROPERTY(EditAnywhere, config, Category = "HLOD", meta = (DisplayName = "HLOD Max Texture Size", ClampMin = "64", ClampMax = "8192", ToolTip = "Maximum size of the textures generated for landscape HLODs"))
+	int32 HLODMaxTextureSize = 1024;
+
 	UPROPERTY(transient)
 	bool bRestrictiveMode = false;
 };
