Fix for ensures hit on Landscape creation undo.

After an undo, we double check the state of all of our streamed in textures.

diff --git a/Engine/Source/Runtime/Landscape/Private/LandscapeSubsystem.cpp b/Engine/Source/Runtime/Landscape/Private/LandscapeSubsystem.cpp
--- a/Engine/Source/Runtime/Landscape/Private/LandscapeSubsystem.cpp
+++ b/Engine/Source/Runtime/Landscape/Private/LandscapeSubsystem.cpp
@@ -318,6 +318,9 @@ void ULandscapeSubsystem::Tick(float DeltaTime)
 	}
 #endif // WITH_EDITOR
 	
+	// Double check requested textures are in the right state
+	TextureStreamingManager->CheckRequestedTextures();
+
 	static TArray<FVector> OldCameras;
 	TArray<FVector>* Cameras = nullptr;
 	if (GUseStreamingManagerForCameras == 0)
diff --git a/Engine/Source/Runtime/Landscape/Private/LandscapeTextureStreamingManager.cpp b/Engine/Source/Runtime/Landscape/Private/LandscapeTextureStreamingManager.cpp
--- a/Engine/Source/Runtime/Landscape/Private/LandscapeTextureStreamingManager.cpp
+++ b/Engine/Source/Runtime/Landscape/Private/LandscapeTextureStreamingManager.cpp
@@ -140,6 +140,32 @@ void FLandscapeTextureStreamingManager::CleanupInvalidEntries()
 	}
 }
 
+void FLandscapeTextureStreamingManager::CheckRequestedTextures()
+{
+#if WITH_EDITOR
+	if (UndoDetector.bUndoRedoPerformed)
+	{
+		// the force mip levels resident flag sometimes gets cleared on an undo after landscape creation, but we can fix it
+		// (otherwise we may wait forever for them to become resident)
+		for (auto It = TextureStates.CreateIterator(); It; ++It)
+		{
+			if (UTexture* Texture = It.Key().Get())
+			{
+				FTextureState& State = It.Value();
+				if (State.bForever || State.RequestCount > 0)
+				{
+					if (!Texture->bForceMiplevelsToBeResident)
+					{
+						Texture->bForceMiplevelsToBeResident = true;
+					}
+				}
+			}
+		}
+		UndoDetector.bUndoRedoPerformed = false;
+	}
+#endif // WITH_EDITOR
+}
+
 bool FLandscapeTextureStreamingManager::IsTextureFullyStreamedIn(UTexture* InTexture)
 {
 	return InTexture &&
@@ -163,4 +189,4 @@ FLandscapeTextureStreamingManager::~FLandscapeTextureStreamingManager()
 			}
 		}
 	}
-}
\ No newline at end of file
+}
diff --git a/Engine/Source/Runtime/Landscape/Private/LandscapeTextureStreamingManager.h b/Engine/Source/Runtime/Landscape/Private/LandscapeTextureStreamingManager.h
--- a/Engine/Source/Runtime/Landscape/Private/LandscapeTextureStreamingManager.h
+++ b/Engine/Source/Runtime/Landscape/Private/LandscapeTextureStreamingManager.h
@@ -4,6 +4,29 @@
 #include "CoreMinimal.h"
 #include "UObject/WeakObjectPtr.h"
 
+#if WITH_EDITOR
+
+#include "EditorUndoClient.h"
+
+class FLandscapeTextureStreamingManagerUndoDetector : FSelfRegisteringEditorUndoClient
+{
+public:
+	bool bUndoRedoPerformed = false;
+
+private:
+	virtual void PostUndo(bool bSuccess) override
+	{
+		bUndoRedoPerformed = true;
+	};
+
+	virtual void PostRedo(bool bSuccess) override
+	{
+		bUndoRedoPerformed = true;
+	};
+};
+
+#endif // WITH_EDITOR
+
 class UTexture;
 
 class FLandscapeTextureStreamingManager
@@ -37,6 +60,9 @@ public:
 	// Call this clean up any old entries in tracked TextureStates for textures that have been unloaded without first being Unrequested.
 	void CleanupInvalidEntries();
 
+	// Check that all requested textures are still requested.
+	void CheckRequestedTextures();
+
 	static bool IsTextureFullyStreamedIn(UTexture* Texture);
 
 	~FLandscapeTextureStreamingManager();
@@ -49,4 +75,9 @@ private:
 	};
 
 	TMap<TWeakObjectPtr<UTexture>, FTextureState, FDefaultSetAllocator, TWeakObjectPtrMapKeyFuncs<TWeakObjectPtr<UTexture>, FTextureState>> TextureStates;
+
+#if WITH_EDITOR
+	FLandscapeTextureStreamingManagerUndoDetector UndoDetector;
+#endif // WITH_EDITOR
+
 };
