Removing grass map build toasts

diff --git a/Engine/Source/Runtime/Landscape/Classes/Landscape.h b/Engine/Source/Runtime/Landscape/Classes/Landscape.h
--- a/Engine/Source/Runtime/Landscape/Classes/Landscape.h
+++ b/Engine/Source/Runtime/Landscape/Classes/Landscape.h
@@ -606,7 +606,6 @@ private:
 	TSharedPtr<FLandscapeNotification> WaitingForTexturesNotification;
 	TSharedPtr<FLandscapeNotification> WaitingForBrushesNotification;
 	TSharedPtr<FLandscapeNotification> InvalidShadingModelNotification;
-	TSharedPtr<FLandscapeNotification> GrassRenderingNotification;
 
 	// Represent all the resolved paint layer, from all layers blended together (size of the landscape x material layer count)
 	class FLandscapeTexture2DArrayResource* CombinedLayersWeightmapAllMaterialLayersResource;
diff --git a/Engine/Source/Runtime/Landscape/Private/Landscape.cpp b/Engine/Source/Runtime/Landscape/Private/Landscape.cpp
--- a/Engine/Source/Runtime/Landscape/Private/Landscape.cpp
+++ b/Engine/Source/Runtime/Landscape/Private/Landscape.cpp
@@ -2779,42 +2779,6 @@ void ULandscapeInfo::Serialize(FArchive& Ar)
 void ALandscape::PostInitProperties()
 {
 	Super::PostInitProperties();
-	
-#if WITH_EDITOR
-	if (!IsTemplate())
-	{
-		if (const UWorld* World = GetWorld())
-		{
-			if (ULandscapeSubsystem* LandscapeSubSystem = World->GetSubsystem<ULandscapeSubsystem>())
-			{
-				if (FLandscapeNotificationManager* LandscapeNotificationManager = LandscapeSubSystem->GetNotificationManager())
-				{
-					// Create conditional GrassRenderingNotification
-					GrassRenderingNotification = MakeShared<FLandscapeNotification>(this, FLandscapeNotification::EType::GrassRendering,
-						// display condition
-						[this, LandscapeSubSystem]()
-						{
-							return GGrassEnable && !GGrassMapUseRuntimeGeneration &&
-								(LandscapeSubSystem->GetGrassMapBuilder()->GetTotalGrassMapsWaitingToRender() > 0) &&
-								!(FSlateApplication::Get().HasAnyMouseCaptor() || GUnrealEd->IsUserInteracting());
-						},
-						// retrieve text
-						[this, LandscapeSubSystem](FText& InText)
-						{
-							// Accumulating all outstanding grass maps that need to be rendered for ALL landscapes because only one such notification can be displayed at a time
-							int ComponentsNeedingGrassMapRenderAllLandscapes = LandscapeSubSystem->GetGrassMapBuilder()->GetTotalGrassMapsWaitingToRender();
-							
-							FFormatNamedArguments Args;
-							Args.Add(TEXT("OutstandingGrassMaps"), FText::AsNumber(ComponentsNeedingGrassMapRenderAllLandscapes));
-							InText = FText::Format(NSLOCTEXT("GrassMapRender", "GrassMapRenderFormat", "Building Grass Maps ({OutstandingGrassMaps})"), Args);
-						});
-			
-					LandscapeNotificationManager->RegisterNotification(GrassRenderingNotification);
-				}
-			}
-		}
-	}
-#endif // WITH_EDITOR
 }
 
 
