Add per-platform controls over landscape compression.
Always show the landscape compression option under advanced settings in the landscape details.

diff --git a/Engine/Source/Editor/DetailCustomizations/Private/LandscapeProxyUIDetails.cpp b/Engine/Source/Editor/DetailCustomizations/Private/LandscapeProxyUIDetails.cpp
--- a/Engine/Source/Editor/DetailCustomizations/Private/LandscapeProxyUIDetails.cpp
+++ b/Engine/Source/Editor/DetailCustomizations/Private/LandscapeProxyUIDetails.cpp
@@ -54,12 +54,6 @@ class URuntimeVirtualTexture;
 
 #define LOCTEXT_NAMESPACE "FLandscapeProxyUIDetails"
 
-static TAutoConsoleVariable<int32> CVarShowCompressHeightmapsOption(
-	TEXT("landscape.ShowCompressHeightmapsOption"),
-	0,
-	TEXT("Enable editing of the compressed height map option on landscape proxies (experimental)."),
-	ECVF_Default);
-
 FLandscapeProxyUIDetails::FLandscapeProxyUIDetails()
 {
 	// Position Precision options are copied from StaticMeshEditorTools.cpp 
@@ -136,11 +130,6 @@ void FLandscapeProxyUIDetails::CustomizeDetails( IDetailLayoutBuilder& DetailBui
 		DetailBuilder.HideProperty(DetailBuilder.GetProperty(GET_MEMBER_NAME_CHECKED(ALandscape, bAreNewLandscapeActorsSpatiallyLoaded), ALandscape::StaticClass()));
 	}
 
-	if (CVarShowCompressHeightmapsOption->GetInt() == 0)
-	{
-		DetailBuilder.HideProperty(DetailBuilder.GetProperty(GET_MEMBER_NAME_CHECKED(ALandscapeProxy, bUseCompressedHeightmapStorage), ALandscapeProxy::StaticClass()));
-	}
-
 	if (LandscapeActors.Num() == 1)
 	{
 		TWeakObjectPtr<ALandscape> LandscapeActor = LandscapeActors[0];
diff --git a/Engine/Source/Runtime/Landscape/Private/Landscape.cpp b/Engine/Source/Runtime/Landscape/Private/Landscape.cpp
--- a/Engine/Source/Runtime/Landscape/Private/Landscape.cpp
+++ b/Engine/Source/Runtime/Landscape/Private/Landscape.cpp
@@ -210,6 +210,17 @@ static FAutoConsoleVariable CVarAllowGrassStripping(
 	TEXT("landscape.AllowGrassStripping"),
 	true,
 	TEXT("Enables the conditional stripping of grass data during cook.  Disabling this means the bStripGrassWhenCooked* will be ignored."));
+
+int32 GLandscapeHeightmapCompressionMode = 0;
+static FAutoConsoleVariableRef CVarLandscapeHeightmapCompressionMode(
+	TEXT("landscape.HeightmapCompressionMode"),
+	GLandscapeHeightmapCompressionMode,
+	TEXT("Defines whether compression is applied to landscapes.\n")
+	TEXT(" 0: use the per-landscape setting bUseCompressedHeightmapStorage (default)\n")
+	TEXT(" 1: force enable heightmap compression on all landscapes\n")
+	TEXT(" -1: force disable heightmap compression on all landscapes\n"),
+	ECVF_Preview | ECVF_ReadOnly);
+
 #endif // WITH_EDITOR
 
 int32 GRenderNaniteLandscape = 1;
@@ -3645,17 +3656,28 @@ void ALandscapeProxy::PreSave(FObjectPreSaveContext ObjectSaveContext)
 		}
 	}
 
-	if (bUseCompressedHeightmapStorage && ObjectSaveContext.IsCooking())
+	if (ObjectSaveContext.IsCooking())
 	{
-		FString PlatformName = ObjectSaveContext.GetTargetPlatform()->PlatformName();
+		FName IniPlatformName = *ObjectSaveContext.GetTargetPlatform()->IniPlatformName();
+
+		int32 HeightmapCompressionMode = 0;
+		if (IConsoleVariable* PlatformCVar = CVarLandscapeHeightmapCompressionMode->GetPlatformValueVariable(IniPlatformName).Get())
+		{
+			PlatformCVar->GetValue(HeightmapCompressionMode);
+		}
+		else
+		{
+			CVarLandscapeHeightmapCompressionMode->GetValue(HeightmapCompressionMode);
+		}
 
-		// TODO [chris.tchou] : add user and/or platform data controllable toggles here
+		bool bShouldCompressHeightmap = (HeightmapCompressionMode == 0) ? bUseCompressedHeightmapStorage : (HeightmapCompressionMode > 0);
+		if (bShouldCompressHeightmap)
 		{
 			for (ULandscapeComponent* LandscapeComponent : LandscapeComponents)
 			{
-				UTexture2D* Tex = LandscapeComponent->GetHeightmap();
+				UTexture2D* HeightmapTex = LandscapeComponent->GetHeightmap();
 				FVector LandscapeGridScale = GetRootComponent()->GetRelativeScale3D();
-				ULandscapeTextureStorageProviderFactory::ApplyTo(Tex, LandscapeGridScale);
+				ULandscapeTextureStorageProviderFactory::ApplyTo(HeightmapTex, LandscapeGridScale);
 			}
 		}
 	}
