Fixed race condition where the water info rendering could attempt to access the landscape render systems before they are created. Instead, we avoid accessing them altogether since they are a RT-only resource and compute them on the gamethread directly from the source data.

diff --git a/Engine/Plugins/Experimental/Water/Source/Runtime/Private/WaterInfoRendering.cpp b/Engine/Plugins/Experimental/Water/Source/Runtime/Private/WaterInfoRendering.cpp
--- a/Engine/Plugins/Experimental/Water/Source/Runtime/Private/WaterInfoRendering.cpp
+++ b/Engine/Plugins/Experimental/Water/Source/Runtime/Private/WaterInfoRendering.cpp
@@ -4,6 +4,7 @@
 
 #include "CommonRenderResources.h"
 #include "DataDrivenShaderPlatformInfo.h"
+#include "EngineUtils.h"
 #include "LegacyScreenPercentageDriver.h"
 #include "Modules/ModuleManager.h"
 #include "RenderCaptureInterface.h"
@@ -1004,64 +1005,44 @@ const FName& GetWaterInfoDilationPassName() { return FWaterInfoRenderingDilation
 
 // ----------------------------------------------------------------------------------
 
-static TMap<uint32, int32> GatherLandscapeLODOverrides(const FIntPoint& RenderTargetSize, const FVector& WaterZoneExtents)
+static TMap<uint32, int32> GatherLandscapeLODOverrides(const UWorld* World, const FIntPoint& RenderTargetSize, const FVector& WaterZoneExtents)
 {
-	TMap<uint32, int32> LandscapeLODOverrides;
-
-	TMap<uint32, FLandscapeRenderSystem*> const* LandscapeRenderSystems = nullptr;
-	if (ILandscapeModule* LandscapeModule = FModuleManager::GetModulePtr<ILandscapeModule>("Landscape"))
-	{
-		if (TSharedPtr<FLandscapeSceneViewExtension, ESPMode::ThreadSafe> LandscapeViewExtension = LandscapeModule->GetLandscapeSceneViewExtension())
-		{
-			LandscapeRenderSystems = &LandscapeViewExtension->GetLandscapeRenderSystems();
-		}
-	}
-
-	if (!LandscapeRenderSystems)
-	{
-		return LandscapeLODOverrides;
-	}
-
 	// In order to prevent overdrawing the landscape components, we compute the lowest-detailed LOD level which satisfies the pixel coverage of the Water Info texture
 	// and force it on all landscape components. This override is set different per Landscape actor in case there are multiple under the same water zone.
 	//
 	// Ex: If the WaterInfoTexture only has 1 pixel per 100 units, and the highest landscape LOD has 1 vertex per 20 units, we don't need to use the maximum landscape LOD
 	// and can force a lower level of detail (in this case LOD2) while still satisfying the resolution of the water info texture.
 
-	const double MinWaterInfoTextureExtent = (double)FMath::Min(RenderTargetSize.X, RenderTargetSize.Y);
+	const double MinWaterInfoTextureExtent = FMath::Min(RenderTargetSize.X, RenderTargetSize.Y);
 	const double MaxWaterZoneExtent = FMath::Max(WaterZoneExtents.X, WaterZoneExtents.Y);
 	const double WaterInfoUnitsPerPixel =  MaxWaterZoneExtent / MinWaterInfoTextureExtent;
 
-	for (const TPair<uint32, FLandscapeRenderSystem*>& Pair : *LandscapeRenderSystems)
+	TMap<uint32, int32> LandscapeLODOverrides;
+	for (const ALandscapeProxy* LandscapeProxy : TActorRange<ALandscapeProxy>(World))
 	{
-		uint32 LandscapeRenderSystemKey = Pair.Key;
-		FLandscapeRenderSystem* LandscapeRenderSystem = Pair.Value;
+		if (LandscapeProxy == nullptr)
+		{
+			continue;
+		}
+		
+		const uint32 LandscapeKey = LandscapeProxy->ComputeLandscapeKey();
 		int32 OptimalLODLevel = INDEX_NONE;
 
 		// All components within the same landscape (and thus its render system) should have the same number of quads and the same extent.
 		// therefore we can simply find the first component and compute its optimal LOD level.
-		for (FLandscapeSectionInfo* LandscapeSectionInfo : LandscapeRenderSystem->SectionInfos)
-		{
-			if (LandscapeSectionInfo != nullptr)
-			{
-				// Double the required landscape resolution to achieve 2 quads per pixel.
-				const double LandscapeComponentUnitsPerQuad = 2.0 * LandscapeSectionInfo->ComputeSectionResolution();
-				if (LandscapeComponentUnitsPerQuad <= 0.f)
-				{
-					// No section resolution probably means the section is a mesh proxy, which might not have regular units per vertex.
-					// Avoid computing optimal LOD in this case.
-					continue;
-				}
+		const double FullExtent = LandscapeProxy->SubsectionSizeQuads * FMath::Max(LandscapeProxy->GetTransform().GetScale3D().X, LandscapeProxy->GetTransform().GetScale3D().Y);
+		const double NumQuads = LandscapeProxy->ComponentSizeQuads;
+		const double LandscapeResolution = FullExtent / NumQuads;
 
-				// Derived from:
-				// (ComponentWorldExtent / WaterInfoWorldspaceExtent) * WaterInfoTextureResolution = (NumComponentQuads / 2 ^ (LodLevel))
-				OptimalLODLevel = FMath::Max(WaterInfoRenderLandscapeMinimumMipLevel, FMath::FloorToInt(FMath::Log2(WaterInfoUnitsPerPixel / LandscapeComponentUnitsPerQuad)));
+		// Double the required landscape resolution to achieve 2 quads per pixel.
+		const double LandscapeComponentUnitsPerQuad = 2.0 * LandscapeResolution;
+		check(LandscapeComponentUnitsPerQuad > 0.f);
 
-				break;
-			}
-		}
+		// Derived from:
+		// (ComponentWorldExtent / WaterInfoWorldspaceExtent) * WaterInfoTextureResolution = (NumComponentQuads / 2 ^ (LodLevel))
+		OptimalLODLevel = FMath::Max(WaterInfoRenderLandscapeMinimumMipLevel, FMath::FloorToInt(FMath::Log2(WaterInfoUnitsPerPixel / LandscapeComponentUnitsPerQuad)));
 
-		LandscapeLODOverrides.Add(LandscapeRenderSystemKey, OptimalLODLevel);
+		LandscapeLODOverrides.Add(LandscapeKey, OptimalLODLevel);
 	}
 
 	return LandscapeLODOverrides;
@@ -1124,7 +1105,7 @@ void UpdateWaterInfoRendering_CustomRenderPass(
 	}
 	PassInput.ShowOnlyPrimitives = MoveTemp(ComponentsToRenderInDepthPass);
 
-	FWaterInfoRenderingDepthPass* DepthPass = new FWaterInfoRenderingDepthPass(RenderTargetSize, GatherLandscapeLODOverrides(RenderTargetSize, ZoneExtent));
+	FWaterInfoRenderingDepthPass* DepthPass = new FWaterInfoRenderingDepthPass(RenderTargetSize, GatherLandscapeLODOverrides(Scene->GetWorld(), RenderTargetSize, ZoneExtent));
 	if (bPerformRenderCapture)
 	{
 		// Initiate a render capture when this pass runs :
diff --git a/Engine/Source/Runtime/Landscape/Classes/LandscapeProxy.h b/Engine/Source/Runtime/Landscape/Classes/LandscapeProxy.h
--- a/Engine/Source/Runtime/Landscape/Classes/LandscapeProxy.h
+++ b/Engine/Source/Runtime/Landscape/Classes/LandscapeProxy.h
@@ -997,6 +997,10 @@ public:
 		OriginalLandscapeGuid = Guid;
 	}
 
+	/** Computes the rendering key for this landscape proxy */
+	LANDSCAPE_API uint32 ComputeLandscapeKey() const;
+	LANDSCAPE_API static uint32 ComputeLandscapeKey(const UWorld* World, uint32 InLODGroupKey, FGuid InLandscapeGuid);
+
 	/* Return the original landscape guid, before it was modified by instancing.
 	 * When not instanced, this value is equal to LandscapeGuid.
 	 */
diff --git a/Engine/Source/Runtime/Landscape/Private/LandscapeProxy.cpp b/Engine/Source/Runtime/Landscape/Private/LandscapeProxy.cpp
--- a/Engine/Source/Runtime/Landscape/Private/LandscapeProxy.cpp
+++ b/Engine/Source/Runtime/Landscape/Private/LandscapeProxy.cpp
@@ -30,3 +30,14 @@ void ALandscapeProxy::SetPerLODOverrideMaterials(const TArray<FLandscapePerLODMa
 {
 	PerLODOverrideMaterials = InValue;
 }
+
+uint32 ALandscapeProxy::ComputeLandscapeKey() const
+{
+	return ComputeLandscapeKey(GetWorld(), LODGroupKey, LandscapeGuid);
+}
+
+uint32 ALandscapeProxy::ComputeLandscapeKey(const UWorld* InWorld, uint32 InLODGroupKey, FGuid InLandscapeGuid)
+{
+	// use LODGroupKey instead of LandscapeGUID when LODGroupKey is non-zero
+	return HashCombine(GetTypeHash(InWorld), (InLODGroupKey != 0) ? InLODGroupKey : GetTypeHash(InLandscapeGuid));
+}
diff --git a/Engine/Source/Runtime/Landscape/Private/LandscapeRender.cpp b/Engine/Source/Runtime/Landscape/Private/LandscapeRender.cpp
--- a/Engine/Source/Runtime/Landscape/Private/LandscapeRender.cpp
+++ b/Engine/Source/Runtime/Landscape/Private/LandscapeRender.cpp
@@ -1097,6 +1097,7 @@ void FLandscapeSceneViewExtension::EndFrame_RenderThread()
 
 const TMap<uint32, FLandscapeRenderSystem*>& FLandscapeSceneViewExtension::GetLandscapeRenderSystems() const
 {
+	checkf(IsInParallelRenderingThread(), TEXT("Accessing the Landscape render systems is only valid from the rendering thread!"));
 	return LandscapeRenderSystems;
 }
 
@@ -1136,7 +1137,7 @@ bool FLandscapeVisibilityHelper::OnRemoveFromWorld()
 
 FLandscapeComponentSceneProxy::FLandscapeComponentSceneProxy(ULandscapeComponent* InComponent)
 	: FPrimitiveSceneProxy(InComponent, NAME_LandscapeResourceNameForDebugging)
-	, FLandscapeSectionInfo(InComponent->GetWorld(), InComponent->GetLandscapeProxy()->GetLandscapeGuid(), InComponent->GetSectionBase() / InComponent->ComponentSizeQuads, InComponent->GetLandscapeProxy()->LODGroupKey)
+	, FLandscapeSectionInfo(InComponent->GetWorld(), InComponent->GetLandscapeProxy()->GetLandscapeGuid(), InComponent->GetSectionBase() / InComponent->ComponentSizeQuads, InComponent->GetLandscapeProxy()->LODGroupKey, InComponent->GetLandscapeProxy()->ComputeLandscapeKey())
 	, MaxLOD(static_cast<int8>(FMath::CeilLogTwo(InComponent->SubsectionSizeQuads + 1) - 1))
 	, NumWeightmapLayerAllocations(static_cast<int8>(InComponent->GetWeightmapLayerAllocations().Num()))
 	, StaticLightingLOD(static_cast<uint8>(InComponent->GetLandscapeProxy()->StaticLightingLOD))
@@ -4283,8 +4284,8 @@ void FLandscapeComponentSceneProxy::GetSectionCenterAndVectors(FVector& OutSecti
 //
 // FLandscapeSectionInfo
 //
-FLandscapeSectionInfo::FLandscapeSectionInfo(const UWorld* InWorld, const FGuid& InLandscapeGuid, const FIntPoint& InComponentBase, uint32 LODGroupKey)
-	: LandscapeKey(HashCombine(GetTypeHash(InWorld), (LODGroupKey != 0) ? LODGroupKey : GetTypeHash(InLandscapeGuid)))    // use LODGroupKey instead of LandscapeGUID when LODGroupKey is non-zero
+FLandscapeSectionInfo::FLandscapeSectionInfo(const UWorld* InWorld, const FGuid& InLandscapeGuid, const FIntPoint& InComponentBase, uint32 LODGroupKey, uint32 InLandscapeKey)
+	: LandscapeKey(InLandscapeKey)
 	, LODGroupKey(LODGroupKey)
 	, RenderCoord(INT32_MIN, INT32_MIN)
 	, ComponentBase(InComponentBase)
@@ -4299,8 +4300,8 @@ FLandscapeSectionInfo::FLandscapeSectionInfo(const UWorld* InWorld, const FGuid&
 class FLandscapeProxySectionInfo : public FLandscapeSectionInfo
 {
 public:
-	FLandscapeProxySectionInfo(const UWorld* InWorld, const FGuid& InLandscapeGuid, const FIntPoint& InComponentBase, const FVector& InComponentCenterLocalSpace, const FVector& InComponentXVectorLocalSpace, const FVector& InComponentYVectorLocalSpace, const FTransform& LocalToWorld, int32 SectionComponentResolution, int8 InProxyLOD, uint32 LODGroupKey)
-		: FLandscapeSectionInfo(InWorld, InLandscapeGuid, InComponentBase, LODGroupKey)
+	FLandscapeProxySectionInfo(const UWorld* InWorld, const FGuid& InLandscapeGuid, const FIntPoint& InComponentBase, const FVector& InComponentCenterLocalSpace, const FVector& InComponentXVectorLocalSpace, const FVector& InComponentYVectorLocalSpace, const FTransform& LocalToWorld, int32 SectionComponentResolution, int8 InProxyLOD, uint32 LODGroupKey, uint32 InLandscapeKey)
+		: FLandscapeSectionInfo(InWorld, InLandscapeGuid, InComponentBase, LODGroupKey, InLandscapeKey)
 		, ProxyLOD(InProxyLOD)
 		, ComponentResolution(SectionComponentResolution)
 		, CenterWorldSpace(LocalToWorld.TransformPosition(InComponentCenterLocalSpace))
@@ -4357,7 +4358,7 @@ private:
 //
 // FLandscapeMeshProxySceneProxy
 //
-FLandscapeMeshProxySceneProxy::FLandscapeMeshProxySceneProxy(UStaticMeshComponent* InComponent, const FGuid& InLandscapeGuid, const TArray<FIntPoint>& InProxySectionsBases, const TArray<FVector>& InProxySectionsCentersLocalSpace, const FVector& InComponentXVector, const FVector& InComponentYVector, const FTransform& LocalToWorld, int32 ComponentResolution, int8 InProxyLOD, uint32 InLODGroupKey)
+FLandscapeMeshProxySceneProxy::FLandscapeMeshProxySceneProxy(UStaticMeshComponent* InComponent, const FGuid& InLandscapeGuid, const TArray<FIntPoint>& InProxySectionsBases, const TArray<FVector>& InProxySectionsCentersLocalSpace, const FVector& InComponentXVector, const FVector& InComponentYVector, const FTransform& LocalToWorld, int32 ComponentResolution, int8 InProxyLOD, uint32 InLODGroupKey, uint32 InLandscapeKey)
 	: FStaticMeshSceneProxy(InComponent, false)
 {
 	check(InLODGroupKey == 0 || ComponentResolution != 0);	// if using LODGroupKey, we require HLODs to be rebuilt to capture the component resolution and section transform data (position / orientation)
@@ -4377,7 +4378,7 @@ FLandscapeMeshProxySceneProxy::FLandscapeMeshProxySceneProxy(UStaticMeshComponen
 	{
 		FIntPoint SectionBase = InProxySectionsBases[Index];
 		FVector ProxySectionCenterLocalSpace = InProxySectionsCentersLocalSpace.IsValidIndex(Index) ? InProxySectionsCentersLocalSpace[Index] : FVector::Zero();
-		ProxySectionsInfos.Emplace(MakeUnique<FLandscapeProxySectionInfo>(InComponent->GetWorld(), InLandscapeGuid, SectionBase, ProxySectionCenterLocalSpace, InComponentXVector, InComponentYVector, LocalToWorld, ComponentResolution, InProxyLOD, InLODGroupKey));
+		ProxySectionsInfos.Emplace(MakeUnique<FLandscapeProxySectionInfo>(InComponent->GetWorld(), InLandscapeGuid, SectionBase, ProxySectionCenterLocalSpace, InComponentXVector, InComponentYVector, LocalToWorld, ComponentResolution, InProxyLOD, InLODGroupKey, InLandscapeKey));
 	}
 }
 
@@ -4470,7 +4471,7 @@ FPrimitiveSceneProxy* ULandscapeMeshProxyComponent::CreateSceneProxy()
 		return nullptr;
 	}
 
-	return new FLandscapeMeshProxySceneProxy(this, LandscapeGuid, ProxyComponentBases, ProxyComponentCentersObjectSpace, ComponentXVectorObjectSpace, ComponentYVectorObjectSpace, GetComponentTransform(), ComponentResolution, ProxyLOD, LODGroupKey);
+	return new FLandscapeMeshProxySceneProxy(this, LandscapeGuid, ProxyComponentBases, ProxyComponentCentersObjectSpace, ComponentXVectorObjectSpace, ComponentYVectorObjectSpace, GetComponentTransform(), ComponentResolution, ProxyLOD, LODGroupKey, ALandscapeProxy::ComputeLandscapeKey(GetWorld(), LODGroupKey, LandscapeGuid));
 }
 
 class FLandscapeNaniteSceneProxy : public ::Nanite::FSceneProxy
diff --git a/Engine/Source/Runtime/Landscape/Public/LandscapeRender.h b/Engine/Source/Runtime/Landscape/Public/LandscapeRender.h
--- a/Engine/Source/Runtime/Landscape/Public/LandscapeRender.h
+++ b/Engine/Source/Runtime/Landscape/Public/LandscapeRender.h
@@ -402,7 +402,7 @@ public:
 class FLandscapeSectionInfo : public TIntrusiveLinkedList<FLandscapeSectionInfo>
 {
 public:
-	FLandscapeSectionInfo(const UWorld* InWorld, const FGuid& InLandscapeGuid, const FIntPoint& InComponentBase, uint32 LODGroupKey);
+	FLandscapeSectionInfo(const UWorld* InWorld, const FGuid& InLandscapeGuid, const FIntPoint& InComponentBase, uint32 LODGroupKey, uint32 InLandscapeKey);
 	virtual ~FLandscapeSectionInfo() = default;
 
 	virtual float ComputeLODForView(const FSceneView& InView) const = 0;
@@ -702,7 +702,7 @@ class FLandscapeMeshProxySceneProxy final : public FStaticMeshSceneProxy
 public:
 	SIZE_T GetTypeHash() const override;
 
-	FLandscapeMeshProxySceneProxy(UStaticMeshComponent* InComponent, const FGuid& InLandscapeGuid, const TArray<FIntPoint>& InProxySectionsBases, const TArray<FVector>& InProxySectionsCentersLocalSpace, const FVector& InComponentXVector, const FVector& InComponentYVector, const FTransform& LocalToWorld, int32 ComponentResolution, int8 InProxyLOD, uint32 InLODGroupKey);
+	FLandscapeMeshProxySceneProxy(UStaticMeshComponent* InComponent, const FGuid& InLandscapeGuid, const TArray<FIntPoint>& InProxySectionsBases, const TArray<FVector>& InProxySectionsCentersLocalSpace, const FVector& InComponentXVector, const FVector& InComponentYVector, const FTransform& LocalToWorld, int32 ComponentResolution, int8 InProxyLOD, uint32 InLODGroupKey, uint32 LandscapeKey);
 	virtual void CreateRenderThreadResources(FRHICommandListBase& RHICmdList) override;
 	virtual void DestroyRenderThreadResources() override;
 	virtual bool OnLevelAddedToWorld_RenderThread() override;
