Fix crashes that happen on LandscapeGrassMapBuilder shutdown when it fails to shut down grass map GPU Readbacks.

Added ability to force immediate shutdown on all tasks and GPU readbacks.  This means when we shut down the system, we no longer need to busy loop wait for all tasks to complete themselves.

Also modified readback synchronization code to launch fewer render thread commands, fewer synchronizations of render thread and GPU with the game thread.

Also upped the number of simultaneous renders allowed in editor, as that was the bottleneck for throughput.  With a limit of 3, it can now usually maintain a throughput of 1 per frame, with a 3 frame latency.  (Previously it would complete 1 every 3 frames because of the latency).

diff --git a/Engine/Source/Runtime/Landscape/Private/LandscapeAsyncTextureReadback.cpp b/Engine/Source/Runtime/Landscape/Private/LandscapeAsyncTextureReadback.cpp
--- a/Engine/Source/Runtime/Landscape/Private/LandscapeAsyncTextureReadback.cpp
+++ b/Engine/Source/Runtime/Landscape/Private/LandscapeAsyncTextureReadback.cpp
@@ -7,7 +7,7 @@
 
 void FLandscapeAsyncTextureReadback::StartReadback_RenderThread(FRDGBuilder& GraphBuilder, FRDGTextureRef RDGTexture)
 {
-	check(!bAsyncReadbackSubmitOnRenderThread && !AsyncReadback);
+	check(!bStartedOnRenderThread && !AsyncReadback);
 	check(RDGTexture->Desc.Format == PF_B8G8R8A8);
 	AsyncReadback = MakeUnique<FRHIGPUTextureReadback>(TEXT("LandscapeGrassReadback"));
 	AddEnqueueCopyPass(GraphBuilder, AsyncReadback.Get(), RDGTexture);
@@ -16,59 +16,62 @@ void FLandscapeAsyncTextureReadback::StartReadback_RenderThread(FRDGBuilder& Gra
 	TextureHeight = Size.Y;
 	check(Size.Z == 1);
 
-	bAsyncReadbackSubmitOnRenderThread = true;
+	bStartedOnRenderThread = true;
 }
 
 void FLandscapeAsyncTextureReadback::FinishReadback_RenderThread()
 {
-	check(bAsyncReadbackSubmitOnRenderThread && AsyncReadback.IsValid());
+	check(bStartedOnRenderThread && AsyncReadback.IsValid());
 	int32 RowPitchInPixels = 0;
 	int32 BufferHeight = 0;
-	void* SrcData = AsyncReadback->Lock(RowPitchInPixels, &BufferHeight);
+	void* SrcData = AsyncReadback->Lock(RowPitchInPixels, &BufferHeight);	// this will block if the readback is not yet ready
 	check(SrcData);
 	check(RowPitchInPixels >= TextureWidth);
 	check(BufferHeight >= TextureHeight);
 
-	// copy into ReadbackResults
-	ReadbackResults.SetNumUninitialized(TextureWidth * TextureHeight);
+	if (!bCancel)	// skip the copy work if we're cancelling
+	{
+		// copy into ReadbackResults
+		ReadbackResults.SetNumUninitialized(TextureWidth * TextureHeight);
 
-	// OpenGL does not really support BGRA images and uses channnel swizzling to emulate them
-	// so when we read them back we get internal RGBA representation
-	const bool bSwapRBChannels = IsOpenGLPlatform(GMaxRHIShaderPlatform);
+		// OpenGL does not really support BGRA images and uses channnel swizzling to emulate them
+		// so when we read them back we get internal RGBA representation
+		const bool bSwapRBChannels = IsOpenGLPlatform(GMaxRHIShaderPlatform);
 
-	if (!bSwapRBChannels && TextureWidth == RowPitchInPixels)
-	{
-		memcpy(ReadbackResults.GetData(), SrcData, TextureWidth * TextureHeight * sizeof(FColor));
-	}
-	else
-	{
-		// copy row by row
-		FColor* Dst = ReadbackResults.GetData();
-		FColor* Src = (FColor*)SrcData;
-		if (bSwapRBChannels)
+		if (!bSwapRBChannels && TextureWidth == RowPitchInPixels)
 		{
-			for (int y = 0; y < TextureHeight; y++)
+			memcpy(ReadbackResults.GetData(), SrcData, TextureWidth * TextureHeight * sizeof(FColor));
+		}
+		else
+		{
+			// copy row by row
+			FColor* Dst = ReadbackResults.GetData();
+			FColor* Src = (FColor*)SrcData;
+			if (bSwapRBChannels)
 			{
-				for (int x = 0; x < TextureWidth; x++)
+				for (int y = 0; y < TextureHeight; y++)
 				{
-					// swap B and R channels when copying
-					Dst->B = Src->R;
-					Dst->G = Src->G;
-					Dst->R = Src->B;
-					Dst->A = Src->A;
-					Dst++;
-					Src++;
+					for (int x = 0; x < TextureWidth; x++)
+					{
+						// swap B and R channels when copying
+						Dst->B = Src->R;
+						Dst->G = Src->G;
+						Dst->R = Src->B;
+						Dst->A = Src->A;
+						Dst++;
+						Src++;
+					}
+					Src += RowPitchInPixels - TextureWidth;
 				}
-				Src += RowPitchInPixels - TextureWidth;
 			}
-		}
-		else
-		{
-			for (int y = 0; y < TextureHeight; y++)
+			else
 			{
-				memcpy(Dst, Src, TextureWidth * sizeof(FColor));
-				Dst += TextureWidth;
-				Src += RowPitchInPixels;
+				for (int y = 0; y < TextureHeight; y++)
+				{
+					memcpy(Dst, Src, TextureWidth * sizeof(FColor));
+					Dst += TextureWidth;
+					Src += RowPitchInPixels;
+				}
 			}
 		}
 	}
@@ -77,49 +80,75 @@ void FLandscapeAsyncTextureReadback::FinishReadback_RenderThread()
 	AsyncReadback.Reset();
 
 	FPlatformMisc::MemoryBarrier();
-	bAsyncReadbackCompleteOnRenderThread = true;
+	bFinishedOnRenderThread = true;
 }
 
-bool FLandscapeAsyncTextureReadback::CheckAndUpdate()
+bool FLandscapeAsyncTextureReadback::CheckAndUpdate(bool& bOutFinishCommandQueued)
 {
-	if (bAsyncReadbackCompleteOnRenderThread)
+	// if we already queued the finish commands to render thread, then we're just waiting on it signaling readback complete	
+	if (bFinishQueuedFromGameThread)
 	{
-		// already done
-		return true;
+		return bFinishedOnRenderThread;
+	}
+	
+	// if we haven't started, or if the readback is not yet ready, then we have nothing to do but wait
+	if (!bStartedOnRenderThread || !AsyncReadback->IsReady())
+	{
+		return false;
 	}
 
-	// not done yet -- queue update check on render thread
+	// the readback was started and it is ready, but we have not yet queued the finish command.
+	// queue it to make the data available to the game thread
 	FLandscapeAsyncTextureReadback* Readback = this;
-	ENQUEUE_RENDER_COMMAND(FLandscapeAsyncTextureReadback_CheckAndUpdate)(
+	ENQUEUE_RENDER_COMMAND(FLandscapeAsyncTextureReadback_FinishReadback)(
 		[Readback](FRHICommandListImmediate& RHICmdList)
 		{
-			if (Readback->bAsyncReadbackCompleteOnRenderThread)
-			{
-				// Actually finished already.. game thread was just asking too early.
-				// No need to do anything, game thread will advance next time it checks.
-			}
-			else if (Readback->bAsyncReadbackSubmitOnRenderThread)
-			{
-				// Readback was submit but not found to be complete yet -- let's check if it completed.
-				check(Readback->AsyncReadback.IsValid());
-				if (Readback->AsyncReadback->IsReady())
-				{
-					Readback->FinishReadback_RenderThread();
-				}
-			}
-			else
-			{
-				// not submit yet.. nothing to do
-			}
+			// sanity check the state 
+			check(Readback->bStartedOnRenderThread && !Readback->bFinishedOnRenderThread);
+			check(Readback->AsyncReadback.IsValid() && Readback->AsyncReadback->IsReady());
+			Readback->FinishReadback_RenderThread();
 		});
 
+	bFinishQueuedFromGameThread = true;
+	bOutFinishCommandQueued = true;
+
 	return false;
 }
 
+void FLandscapeAsyncTextureReadback::CancelAndSelfDestruct()
+{
+	// set the cancel flag, which will reduce work done by the finish command
+	bCancel = true;
+
+	// run a render thread finish command if it hasn't been queued  yet
+	bool bNeedsFinish = !bFinishQueuedFromGameThread;
+	bFinishQueuedFromGameThread = true;
+
+	FLandscapeAsyncTextureReadback* Readback = this;
+	ENQUEUE_RENDER_COMMAND(FCancelAndSelfDestructCommand)(
+		[Readback, bNeedsFinish](FRHICommandListImmediate& RHICmdList)
+		{
+			check(Readback->bStartedOnRenderThread);
+			check(Readback->bCancel);
+
+			// if not yet finished, run the finish command
+			if (bNeedsFinish)
+			{
+				Readback->FinishReadback_RenderThread();
+			}
+
+			check(Readback->bFinishedOnRenderThread);
+
+			// self destruct
+			delete Readback;
+		});
+}
+
+
 void FLandscapeAsyncTextureReadback::QueueDeletionFromGameThread()
 {
 	check(IsInGameThread());
-	check(bAsyncReadbackCompleteOnRenderThread);
+	check(bFinishedOnRenderThread);
 
 	FLandscapeAsyncTextureReadback* Readback = this;
 	ENQUEUE_RENDER_COMMAND(FLandscapeAsyncTextureReadback_CheckAndUpdate)(
diff --git a/Engine/Source/Runtime/Landscape/Private/LandscapeAsyncTextureReadback.h b/Engine/Source/Runtime/Landscape/Private/LandscapeAsyncTextureReadback.h
--- a/Engine/Source/Runtime/Landscape/Private/LandscapeAsyncTextureReadback.h
+++ b/Engine/Source/Runtime/Landscape/Private/LandscapeAsyncTextureReadback.h
@@ -5,23 +5,28 @@
 #include "RHIGPUReadback.h"
 #include "RenderGraphFwd.h"
 
+#include <atomic>
+
 // Class that performs an async GPU texture readback for landscape purposes.
 // Only supports BGRA8 texture format, and returns the results as an array of FColors.
 class FLandscapeAsyncTextureReadback
 {
 private:
+	// game thread state
+	bool bFinishQueuedFromGameThread = false;					// game thread has sent the finish command to the render thread
+	bool bCancel = false;										// game thread wants to cancel readback, don't need to produce valid results
+
 	// render thread state
-	bool bAsyncReadbackSubmitOnRenderThread = false;
-	bool bAsyncReadbackCompleteOnRenderThread = false;
-	TUniquePtr<FRHIGPUTextureReadback> AsyncReadback;
+	std::atomic_bool bStartedOnRenderThread = false;			// render thread start command has sent the readback command to the GPU
+	std::atomic_bool bFinishedOnRenderThread = false;			// render thread finish command has made the data available to game thread
 
-	// results
+	TUniquePtr<FRHIGPUTextureReadback> AsyncReadback;			// render thread managed async readback structure
+
+	// results - readable by game thread when bFinishedOnRenderThread
 	int32 TextureWidth = 0;
 	int32 TextureHeight = 0;
 	TArray<FColor> ReadbackResults;
 
-	void FinishReadback_RenderThread();
-
 public:
 	// construct on the game thread
 	FLandscapeAsyncTextureReadback() {}
@@ -35,23 +40,32 @@ public:
 	// use this to start an async readback operation from the render thread (on a render graph texture)
 	void StartReadback_RenderThread(FRDGBuilder& GraphBuilder, FRDGTextureRef RDGTexture);
 
+	// use this to finish an async readback operation from the render thread
+	// Calling this when !IsComplete() will cause a stall until the GPU has completed the readback
+	void FinishReadback_RenderThread();
+
 	// TODO [chris.tchou] : add other StartReadback functions as needed
 
-	// Check readback status and update the process if needed. Return true when the AsyncReadbackResults are available.
-	// You must call this occasionally (from the game thread) or the readback may never complete.
-	bool CheckAndUpdate();
+	// Non-blocking call from the game thread to check readback status and start the finish command if needed.
+	// Returns true when the AsyncReadbackResults are available to the game thread, at which point you can call TakeResults() to access them.
+	// bOutFinishCommandQueued is set to true if the finish render command was queued in this call, othwerwise unchanged.
+	// You must call this function occasionally or the readback will never complete, as the finish command is required.
+	bool CheckAndUpdate(bool& bOutFinishCommandQueued);
+
+	// Call from game thread to terminate any readback in flight and queue deletion of this (FLandscapeAsyncTextureReadback) on the render thread
+	void CancelAndSelfDestruct();
 
-	// Returns true when async readback results are available.  Call GetResults() to retrieve them.
+	// Returns true when async readback results are available.  Call TakeResults() to retrieve them.
 	bool IsComplete()
 	{
-		return bAsyncReadbackCompleteOnRenderThread;
+		return bFinishedOnRenderThread;
 	}
 
 	// Retrieve the async readback results.  Requires readback to be complete.
 	// This function returns its internal memory buffer, relinquishing control over it, so this function can only be called once.
 	TArray<FColor> TakeResults(FIntPoint* OutSize)
 	{
-		check(bAsyncReadbackCompleteOnRenderThread);
+		check(bFinishedOnRenderThread);
 		if (OutSize)
 		{
 			*OutSize = FIntPoint(TextureWidth, TextureHeight);
@@ -62,7 +76,7 @@ public:
 	FString ToString()
 	{
 		FString Result;
-		Result.Appendf(TEXT("FLandscapeAsyncTextureReadback { Submit: %d Complete: %d AsyncReadback: %p }"), bAsyncReadbackSubmitOnRenderThread, bAsyncReadbackCompleteOnRenderThread, AsyncReadback.Get());
+		Result.Appendf(TEXT("FLandscapeAsyncTextureReadback { RTStart: %d RTComplete: %d Queued: %d Cancel: %d AsyncReadback: %p }"), bStartedOnRenderThread.load(), bFinishedOnRenderThread.load(), bFinishQueuedFromGameThread, bCancel, AsyncReadback.Get());
 		return Result;
 	}
 
diff --git a/Engine/Source/Runtime/Landscape/Private/LandscapeGrassMapsBuilder.cpp b/Engine/Source/Runtime/Landscape/Private/LandscapeGrassMapsBuilder.cpp
--- a/Engine/Source/Runtime/Landscape/Private/LandscapeGrassMapsBuilder.cpp
+++ b/Engine/Source/Runtime/Landscape/Private/LandscapeGrassMapsBuilder.cpp
@@ -60,7 +60,13 @@ static FAutoConsoleVariableRef CVarGrassMapMaxComponentsStreaming(
 	GGrassMapMaxComponentsStreaming,
 	TEXT("How many landscape components can be streaming their textures at once for grass map renders, when using amortized runtime generation."));
 
+// Rendering readback takes ~3 frames on average to complete, while streaming usually takes 1 frame.
+// By setting rendering limit higher in editor we can achieve the same average throughput for both streaming and rendering at 1 per frame.
+#if WITH_EDITOR
+static int32 GGrassMapMaxComponentsRendering = 3;
+#else
 static int32 GGrassMapMaxComponentsRendering = 1;
+#endif // WITH_EDITOR
 static FAutoConsoleVariableRef CVarGrassMapMaxComponentsRendering(
 	TEXT("grass.GrassMap.MaxComponentsRendering"),
 	GGrassMapMaxComponentsRendering,
@@ -242,24 +248,28 @@ namespace UE::Landscape
 		return MinSqrDistance;
 	}
 
-	static void SubmitGPUCommands(bool bBlockUntilComplete)
+	static void SubmitGPUCommands(bool bBlockUntilRTComplete, bool bBlockRTUntilGPUComplete)
 	{
 		FEvent* ResultsReadyEvent = nullptr;
-		if (bBlockUntilComplete)
+		if (bBlockUntilRTComplete)
 		{
 			ResultsReadyEvent = FPlatformProcess::GetSynchEventFromPool(true);
 		}
 
 		ENQUEUE_RENDER_COMMAND(FFlushResourcesCommand)(
-			[ResultsReadyEvent](FRHICommandList& RHICmdList)
+			[ResultsReadyEvent, bBlockRTUntilGPUComplete](FRHICommandList& RHICmdList)
 			{
 				FRHICommandListExecutor::GetImmediateCommandList().ImmediateFlush(EImmediateFlushType::FlushRHIThreadFlushResources);
 				RHIFlushResources();
 				FRHICommandListExecutor::GetImmediateCommandList().SubmitCommandsAndFlushGPU();
 				if (ResultsReadyEvent)
 				{
-					// block render thread waiting for GPU to complete
-					FRHICommandListExecutor::GetImmediateCommandList().BlockUntilGPUIdle();
+					if (bBlockRTUntilGPUComplete)
+					{
+						// Block render thread waiting for GPU to complete.  Note this can be very expensive on some platforms.
+						FRHICommandListExecutor::GetImmediateCommandList().BlockUntilGPUIdle();
+					}
+					
 					ResultsReadyEvent->Trigger();
 				}
 			});
@@ -288,30 +298,22 @@ FLandscapeGrassMapsBuilder::~FLandscapeGrassMapsBuilder()
 			// This happens when deleting a level, the components are not unregistered before the world is destroyed.
 			UnregisterComponent(Component);
 		}
-		State->TickCount = 7777;	// set to a large number so we don't wait to delete entries
 	}
 
 	// update component state until they all delete themselves
 	// (this should happen on the first update, unless a GPU readback is active.
 	// And it shouldn't take more than 3 update if there is a GPU readback.
+
+	double LastFlush = FPlatformTime::Seconds();
+
 	int32 Iterations = 0;
-	while (Iterations < 6 && ComponentStates.Num() > 0)
+	while (Iterations < 3 && ComponentStates.Num() > 0)
 	{
 		TArray<FVector> EmptyCamerasArray;
 		int32 UpdateAllComponentCount = ComponentStates.Num();
-		UpdateTrackedComponents(EmptyCamerasArray, 0, UpdateAllComponentCount, /* bCancelAndEvictAll = */ true);
+		UpdateTrackedComponents(EmptyCamerasArray, 0, UpdateAllComponentCount, /* bCancelAndEvictAllImmediately = */ true);
 
-		// if we happen to get caught with a GPU readback in flight, submit GPU commands to make sure it moves forward
-		if (RenderingCount > 0)
-		{
-			UE::Landscape::SubmitGPUCommands(/* bBlockUntilComplete =  */ true);
-		}
-
-		// if any async fetch tasks are in flight, force completion
-		if (AsyncFetchCount > 0)
-		{
-			CompleteAllAsyncTasksNow();
-		}
+		ensure(NotReadyCount == 0 && StreamingCount == 0 && RenderingCount == 0 && AsyncFetchCount == 0 && PopulatedCount == 0);
 
 		Iterations++;
 	}
@@ -384,12 +386,13 @@ void FAsyncFetchTask::DoWork()
 	Results = ActiveRender->FetchResults(bFreeAsyncReadback);
 }
 
-bool FLandscapeGrassMapsBuilder::UpdateTrackedComponents(const TArray<FVector>& Cameras, int32 LocalMaxRendering, int32 MaxExpensiveUpdateChecksToPerform, bool bCancelAndEvictAll)
+bool FLandscapeGrassMapsBuilder::UpdateTrackedComponents(const TArray<FVector>& Cameras, int32 LocalMaxRendering, int32 MaxExpensiveUpdateChecksToPerform, bool bCancelAndEvictAllImmediately)
 {
 	TRACE_CPUPROFILER_EVENT_SCOPE(FLandscapeGrassMapsBuilder::UpdateTrackedComponents);
 	SCOPE_CYCLE_COUNTER(STAT_UpdateComponentGrassMaps);
 
 	bool bChanged = false;
+	bRenderCommandsQueuedByLastUpdate = false;
 
 	// Array to store components that are updated after the initial update pass (should never be more than max streaming components)
 	static TArray<FComponentState*> StreamingStatesToProcess;
@@ -415,7 +418,7 @@ bool FLandscapeGrassMapsBuilder::UpdateTrackedComponents(const TArray<FVector>&
 		ULandscapeComponent* Component = State->Component;
 		State->TickCount++;
 
-		if (bCancelAndEvictAll || (Component == nullptr))
+		if (bCancelAndEvictAllImmediately || (Component == nullptr))
 		{
 			// fallthrough to CancelAndEvict below
 		}
@@ -443,7 +446,7 @@ bool FLandscapeGrassMapsBuilder::UpdateTrackedComponents(const TArray<FVector>&
 
 			case EComponentStage::Rendering:
 				check(State->ActiveRender != nullptr);
-				if (State->ActiveRender->CheckAndUpdateAsyncReadback())
+				if (State->ActiveRender->CheckAndUpdateAsyncReadback(bRenderCommandsQueuedByLastUpdate))
 				{
 					if (GGrassMapUseAsyncFetch != 0)
 					{
@@ -523,7 +526,7 @@ bool FLandscapeGrassMapsBuilder::UpdateTrackedComponents(const TArray<FVector>&
 		}
 
 		// we only fall through to this statement if the code above didn't invoke `continue` (or if the component was unregistered)
-		if (CancelAndEvict(*State))
+		if (CancelAndEvict(*State, bCancelAndEvictAllImmediately))
 		{
 			if (State->Component == nullptr)
 			{
@@ -560,6 +563,7 @@ bool FLandscapeGrassMapsBuilder::UpdateTrackedComponents(const TArray<FVector>&
 			if (State->AreTexturesStreamedIn())
 			{
 				KickOffRenderAndReadback(*State);
+				bRenderCommandsQueuedByLastUpdate = true;
 				bChanged = true;
 			}
 		}
@@ -739,8 +743,8 @@ void FLandscapeGrassMapsBuilder::AmortizedUpdateGrassMaps(
 		AmortizedMaxRendering *= GGrassMapPrioritizedMultiplier;
 	}
 
-	const bool bCancelAndEvictAll = !GGrassEnable;
-	UpdateTrackedComponents(Cameras, AmortizedMaxRendering, GGrassMapMaxDiscardChecksPerFrame, bCancelAndEvictAll);
+	const bool bCancelAndEvictAllImmediately = !GGrassEnable;
+	UpdateTrackedComponents(Cameras, AmortizedMaxRendering, GGrassMapMaxDiscardChecksPerFrame, bCancelAndEvictAllImmediately);
 
 	// no point in looking to start new grass map generation if nothing is pending, if grass is disabled or there are no cameras
 	if (bAllowStartGrassMapGeneration && PendingCount > 0 && GGrassEnable && Cameras.Num() > 0)
@@ -782,12 +786,13 @@ bool FLandscapeGrassMapsBuilder::BuildGrassMapsNowForComponents(
 	int32 UpToDateCount = 0;
 	
 	const double StartTime = FPlatformTime::Seconds();
+	double LastFlush = StartTime;
 	double LastChangeTime = StartTime;
 	while (UpToDateCount != LandscapeComponents.Num())
 	{
 		// ensure we are making progress within a reasonable amount of time TODO [chris.tchou] there should be a better way to detect non-progress here
 		const double CurTime = FPlatformTime::Seconds();
-		if (CurTime > LastChangeTime + 15.0)
+		if (CurTime > LastChangeTime + 30.0)
 		{
 			UE_LOG(LogGrass, Error, TEXT("ERROR: BuildGrassMapsNowForComponents() took too long, grass maps are not up to date"));
 			break;
@@ -795,7 +800,7 @@ bool FLandscapeGrassMapsBuilder::BuildGrassMapsNowForComponents(
 
 		// update all components that are tracked (without evicting)
 		const int32 UpdateAllComponentCount = ComponentStates.Num();
-		bool bChanged = UpdateTrackedComponents(EmptyCamerasArray, MaxStreamingRendering, UpdateAllComponentCount, /* bCancelAndEvictAll= */ false);
+		bool bChanged = UpdateTrackedComponents(EmptyCamerasArray, MaxStreamingRendering, UpdateAllComponentCount, /* bCancelAndEvictAllImmediately= */ false);
 
 		UpToDateCount = 0;
 		int32 AvailableStreamingSlots = MaxStreamingRendering - StreamingCount; // here we don't limit by overall population count
@@ -841,17 +846,13 @@ bool FLandscapeGrassMapsBuilder::BuildGrassMapsNowForComponents(
 			LastUpToDateCount = UpToDateCount;
 		}
 
-		if (bChanged)
-		{
-			LastChangeTime = FPlatformTime::Seconds();
-		}
-
 		// If any rendering is in flight, queue up the gpu commands on the render thread, so the GPU can start working on them.
-		if (RenderingCount > 0)
+		if (bRenderCommandsQueuedByLastUpdate || (RenderingCount > 0 && (CurTime - LastFlush > (1.0 / 60.0))))
 		{
 			// TODO [chris.tchou] it currently seems to be faster to block here; otherwise it takes a long time to complete the readback
 			// not sure why this is, something must be getting starved in the non-blocking path.
-			UE::Landscape::SubmitGPUCommands(/* bBlockUntilComplete =  */ true);
+			UE::Landscape::SubmitGPUCommands(/* bBlockUntilRTComplete =  */ true, /* bBlockRTUntilGPUComplete =  */ false);
+			LastFlush = CurTime;
 		}
 		
 		// If any streaming is in flight, do a blocking texture streaming update.
@@ -865,6 +866,11 @@ bool FLandscapeGrassMapsBuilder::BuildGrassMapsNowForComponents(
 		{
 			CompleteAllAsyncTasksNow();
 		}
+
+		if (bChanged)
+		{
+			LastChangeTime = FPlatformTime::Seconds();
+		}
 	}
 
 	UE_LOG(LogGrass, Verbose, TEXT("BuildGrassMapsNowForComponents() updated %d/%d components in %f seconds"), UpToDateCount, LandscapeComponents.Num(), FPlatformTime::Seconds() - StartTime);
@@ -885,7 +891,7 @@ void FLandscapeGrassMapsBuilder::CompleteAllAsyncTasksNow()
 	}
 }
 
-bool FLandscapeGrassMapsBuilder::CancelAndEvict(FComponentState& State)
+bool FLandscapeGrassMapsBuilder::CancelAndEvict(FComponentState& State, bool bCancelImmediately)
 {
 	TRACE_CPUPROFILER_EVENT_SCOPE(FLandscapeGrassMapsBuilder::CancelAndEvict);
 
@@ -918,14 +924,23 @@ bool FLandscapeGrassMapsBuilder::CancelAndEvict(FComponentState& State)
 			if (State.ActiveRender != nullptr)
 			{
 				// calling update ensures it is pushed forward if there is still a readback in progress
-				if (!State.ActiveRender->CheckAndUpdateAsyncReadback())
+				bool bNewRenderCommands = false;
+				if (bCancelImmediately)
+				{
+					// release active render from the state, and cancel and self destruct it
+					FLandscapeGrassWeightExporter* Render = State.ActiveRender.Release();
+					Render->CancelAndSelfDestruct();
+				}
+				else if (!State.ActiveRender->CheckAndUpdateAsyncReadback(bNewRenderCommands))
 				{
 					// we can't cancel yet.. must wait for the readback to complete
 					return false;
 				}
-
-				// readback is complete, we can delete the active render
-				State.ActiveRender.Reset();
+				else
+				{
+					// readback is complete, we can delete the active render
+					State.ActiveRender.Reset();
+				}
 			}
 			check(RenderingCount > 0);
 			RenderingCount--;
@@ -938,7 +953,11 @@ bool FLandscapeGrassMapsBuilder::CancelAndEvict(FComponentState& State)
 			FAsyncTask<FAsyncFetchTask>* Task = State.AsyncFetchTask.Get();
 			if (Task != nullptr)
 			{
-				if (!Task->IsDone())
+				if (bCancelImmediately)
+				{
+					Task->EnsureCompletion(/* bDoWorkOnThisThreadIfNotStarted= */ true, /* bIsLatencySensitive= */ true);
+				}
+				else if (!Task->IsDone())
 				{
 					// can't cancel, async task is still in flight
 					return false;
@@ -958,7 +977,7 @@ bool FLandscapeGrassMapsBuilder::CancelAndEvict(FComponentState& State)
 		{
 			ULandscapeComponent* Component = State.Component;
 		
-			if (Component == nullptr)
+			if (!bCancelImmediately && (Component == nullptr))
 			{
 				// component was unregistered. Wait a few ticks to see if it comes back before fully evicting.
 				if (State.TickCount < 2)
@@ -1296,12 +1315,6 @@ bool FLandscapeGrassMapsBuilder::FComponentState::AreTexturesStreamedIn() const
 	return true;
 }
 
-bool FLandscapeGrassMapsBuilder::FComponentState::IsRenderReadbackComplete() const
-{
-	check(ActiveRender != nullptr);
-	return ActiveRender->CheckAndUpdateAsyncReadback();
-}
-
 bool FLandscapeGrassMapsBuilder::FComponentState::IsBeyondEvictionRange(const TArray<FVector>& Cameras) const
 {
 	check(Stage == EComponentStage::GrassMapsPopulated);
diff --git a/Engine/Source/Runtime/Landscape/Private/LandscapeGrassMapsBuilder.h b/Engine/Source/Runtime/Landscape/Private/LandscapeGrassMapsBuilder.h
--- a/Engine/Source/Runtime/Landscape/Private/LandscapeGrassMapsBuilder.h
+++ b/Engine/Source/Runtime/Landscape/Private/LandscapeGrassMapsBuilder.h
@@ -126,7 +126,6 @@ private:
 		FComponentState(ULandscapeComponent* Component);
 
 		bool AreTexturesStreamedIn() const;
-		bool IsRenderReadbackComplete() const;
 		bool IsBeyondEvictionRange(const TArray<FVector>& Cameras) const;
 	};
 
@@ -219,6 +218,9 @@ private:
 	// number of components that need to render but are waiting (as of the last call to StartTrackingComponents())
 	int32 TotalComponentsWaitingCount = 0;
 
+	// true if any render thread commands were queued by the last call to UpdateTrackedComponents()
+	bool bRenderCommandsQueuedByLastUpdate = false;
+
 	TAllocatorFixedSizeFreeList<sizeof(FComponentState), 32> StatePoolAllocator;
 
 	// store the grass map state of each registered (or recently unregistered) component
@@ -235,7 +237,9 @@ private:
 	FLandscapeTextureStreamingManager& TextureStreamingManager;
 
 	// tries to cancel any in flight operations and transition back to the Pending state
-	bool CancelAndEvict(FComponentState& State);
+	// returns true when the state has been successfully transitioned back to Pending
+	// if bCancelImmediately is true, it will ensure the component reaches Pending state before returning (possibly blocking on async or gpu tasks)
+	bool CancelAndEvict(FComponentState& State, bool bCancelImmediately);
 
 	// try to kick off the grass map generation pipeline -- returns true if it started the amortized update path, false otherwise.
 	bool StartGrassMapGeneration(FComponentState& State, bool bForceCompileShaders);
diff --git a/Engine/Source/Runtime/Landscape/Private/LandscapeGrassWeightExporter.cpp b/Engine/Source/Runtime/Landscape/Private/LandscapeGrassWeightExporter.cpp
--- a/Engine/Source/Runtime/Landscape/Private/LandscapeGrassWeightExporter.cpp
+++ b/Engine/Source/Runtime/Landscape/Private/LandscapeGrassWeightExporter.cpp
@@ -578,6 +578,23 @@ void FLandscapeGrassWeightExporter::FreeAsyncReadback()
 	AsyncReadbackPtr = nullptr;
 }
 
+void FLandscapeGrassWeightExporter::CancelAndSelfDestruct()
+{
+	check(AsyncReadbackPtr != nullptr);
+
+	// Cancel the readback, and queue destruction on the render thread
+	AsyncReadbackPtr->CancelAndSelfDestruct();
+	AsyncReadbackPtr = nullptr;
+
+	// Queue destruction of FLandscapeGrassWeightExporter, also on the render thread
+	FLandscapeGrassWeightExporter* Exporter = this;
+	ENQUEUE_RENDER_COMMAND(FCancelAndDestructCommand)(
+		[Exporter](FRHICommandListImmediate& RHICmdList)
+		{
+			check(Exporter->AsyncReadbackPtr == nullptr);
+			delete Exporter;
+		});
+}
 
 TMap<ULandscapeComponent*, TUniquePtr<FLandscapeComponentGrassData>, TInlineSetAllocator<1>> FLandscapeGrassWeightExporter::FetchResults(bool bFreeAsyncReadback)
 {
diff --git a/Engine/Source/Runtime/Landscape/Private/LandscapeGrassWeightExporter.h b/Engine/Source/Runtime/Landscape/Private/LandscapeGrassWeightExporter.h
--- a/Engine/Source/Runtime/Landscape/Private/LandscapeGrassWeightExporter.h
+++ b/Engine/Source/Runtime/Landscape/Private/LandscapeGrassWeightExporter.h
@@ -91,10 +91,10 @@ public:
 
 	// If using the async readback path, check its status and update if needed. Return true when the AsyncReadbackResults are available.
 	// You must call this periodically, or the async readback may not complete.
-	bool CheckAndUpdateAsyncReadback()
+	bool CheckAndUpdateAsyncReadback(bool& bOutRenderCommandsQueued)
 	{
 		check(AsyncReadbackPtr != nullptr);
-		return AsyncReadbackPtr->CheckAndUpdate();
+		return AsyncReadbackPtr->CheckAndUpdate(bOutRenderCommandsQueued);
 	}
 
 	// return true if the async readback is complete.  (Does not update the readback state)
@@ -117,6 +117,8 @@ public:
 	// Fetches the results and applies them to the landscape components
 	// If using async readback, requires AsyncReadback to be complete before calling this.
 	void ApplyResults();
+
+	void CancelAndSelfDestruct();
 };
 
 namespace UE::Landscape::Grass
