Fix for DDC Stats related to Lansdcape Collision builds so they don't show up if DDC is skipped

diff --git a/Engine/Source/Runtime/Landscape/Private/LandscapeCollision.cpp b/Engine/Source/Runtime/Landscape/Private/LandscapeCollision.cpp
--- a/Engine/Source/Runtime/Landscape/Private/LandscapeCollision.cpp
+++ b/Engine/Source/Runtime/Landscape/Private/LandscapeCollision.cpp
@@ -109,8 +109,11 @@ namespace LandscapeCollisionCookStats
 	static FCookStats::FDDCResourceUsageStats MeshUsageStats;
 	static FCookStatsManager::FAutoRegisterCallback RegisterCookStats([](FCookStatsManager::AddStatFuncRef AddStat)
 	{
-		HeightfieldUsageStats.LogStats(AddStat, TEXT("LandscapeCollision.Usage"), TEXT("Heightfield"));
-		MeshUsageStats.LogStats(AddStat, TEXT("LandscapeCollision.Usage"), TEXT("Mesh"));
+		if (!GLandscapeCollisionSkipDDC)
+		{
+			HeightfieldUsageStats.LogStats(AddStat, TEXT("LandscapeCollision.Usage"), TEXT("Heightfield"));
+			MeshUsageStats.LogStats(AddStat, TEXT("LandscapeCollision.Usage"), TEXT("Mesh"));
+		}
 	});
 }
 #endif
