#UE Adding a little more info to a warning

diff --git a/Engine/Source/Runtime/Landscape/Private/Landscape.cpp b/Engine/Source/Runtime/Landscape/Private/Landscape.cpp
--- a/Engine/Source/Runtime/Landscape/Private/Landscape.cpp
+++ b/Engine/Source/Runtime/Landscape/Private/Landscape.cpp
@@ -5611,7 +5611,7 @@ void ULandscapeInfo::RegisterActor(ALandscapeProxy* Proxy, bool bMapCheck, bool
 	if (Proxy->GetRootComponent() != nullptr && !DrawScale.Equals(Proxy->GetRootComponent()->GetRelativeScale3D()))
 	{
 		UE_LOG(LogLandscape, Warning, TEXT("Landscape proxy (%s) scale (%s) does not match to main actor scale (%s)."),
-			*Proxy->GetName(), *Proxy->GetRootComponent()->GetRelativeScale3D().ToCompactString(), *DrawScale.ToCompactString());
+			*Proxy->GetPathName(), *Proxy->GetRootComponent()->GetRelativeScale3D().ToCompactString(), *DrawScale.ToCompactString());
 	}
 
 	// register
