Fixes undo action when setting a material on landscape for components with no weightmap data.
Implements PostTransacted to be able to use the same codepath as PostEditChangedProperty after a transaction.

diff --git a/Engine/Source/Runtime/Landscape/Classes/LandscapeProxy.h b/Engine/Source/Runtime/Landscape/Classes/LandscapeProxy.h
--- a/Engine/Source/Runtime/Landscape/Classes/LandscapeProxy.h
+++ b/Engine/Source/Runtime/Landscape/Classes/LandscapeProxy.h
@@ -1207,6 +1207,7 @@ public:
 	LANDSCAPE_API virtual void CreateSplineComponent(const FVector& Scale3D) override;
 
 	virtual bool CanEditChange(const FProperty* InProperty) const override;
+	virtual void PostTransacted(const FTransactionObjectEvent& InTransactionEvent) override;
 	virtual void PostEditChangeProperty(FPropertyChangedEvent& PropertyChangedEvent) override;
 	virtual void PostEditImport() override;
 	virtual void PostEditUndo() override;
diff --git a/Engine/Source/Runtime/Landscape/Private/LandscapeEdit.cpp b/Engine/Source/Runtime/Landscape/Private/LandscapeEdit.cpp
--- a/Engine/Source/Runtime/Landscape/Private/LandscapeEdit.cpp
+++ b/Engine/Source/Runtime/Landscape/Private/LandscapeEdit.cpp
@@ -87,6 +87,7 @@ LandscapeEdit.cpp: Landscape editing
 #include "Materials/MaterialExpressionLandscapeGrassOutput.h"
 #include "ShaderPlatformCachedIniValue.h"
 #include "DataDrivenShaderPlatformInfo.h"
+#include "Misc/TransactionObjectEvent.h"
 #endif
 #include "Algo/Count.h"
 #include "Algo/Transform.h"
@@ -5647,6 +5648,27 @@ bool ALandscapeProxy::CanEditChange(const FProperty* InProperty) const
 	return true;
 }
 
+void ALandscapeProxy::PostTransacted(const FTransactionObjectEvent& InTransactionEvent)
+{
+	if (InTransactionEvent.GetEventType() == ETransactionObjectEventType::UndoRedo)
+	{
+		const TArray<FName>& ChangedProperties = InTransactionEvent.GetChangedProperties();
+		
+		for (const FName& ChangedProperty : ChangedProperties)
+		{
+			FProperty* Property = FindFProperty<FProperty>(GetClass(), ChangedProperty);
+
+			if (Property == nullptr)
+			{
+				continue;
+			}
+			
+			FPropertyChangedEvent PropertyChangedEvent(Property);
+			PostEditChangeProperty(PropertyChangedEvent);
+		}
+	}
+}
+
 void ALandscapeProxy::PostEditChangeProperty(FPropertyChangedEvent& PropertyChangedEvent)
 {
 	const FName MemberPropertyName = PropertyChangedEvent.MemberProperty ? PropertyChangedEvent.MemberProperty->GetFName() : NAME_None;
