Fix very long stall when saving a level containing landscape while in a mobile render preview mode.

The problem was that the grass generation shaders were not compiled.  In this case it was because the preview mode platform was not detected as an editor platform.

In the general case, we force compile shaders when in the BuildGrassMapsNowForComponents(), so if the shaders remain unavailable it is because they failed to compile for some reason.

#changelist validated
#virtualized

diff --git a/Engine/Source/Runtime/Landscape/Private/LandscapeGrassMapsBuilder.cpp b/Engine/Source/Runtime/Landscape/Private/LandscapeGrassMapsBuilder.cpp
--- a/Engine/Source/Runtime/Landscape/Private/LandscapeGrassMapsBuilder.cpp
+++ b/Engine/Source/Runtime/Landscape/Private/LandscapeGrassMapsBuilder.cpp
@@ -784,11 +784,14 @@ bool FLandscapeGrassMapsBuilder::BuildGrassMapsNowForComponents(
 	const TArray<FVector> EmptyCamerasArray;
 	int32 LastUpToDateCount = 0;
 	int32 UpToDateCount = 0;
-	
+
+	// track any components that have failed to build
+	TSet<FComponentState*> FailedStates;
+
 	const double StartTime = FPlatformTime::Seconds();
 	double LastFlush = StartTime;
 	double LastChangeTime = StartTime;
-	while (UpToDateCount != LandscapeComponents.Num())
+	while (UpToDateCount + FailedStates.Num() != LandscapeComponents.Num())
 	{
 		// ensure we are making progress within a reasonable amount of time TODO [chris.tchou] there should be a better way to detect non-progress here
 		const double CurTime = FPlatformTime::Seconds();
@@ -836,8 +839,28 @@ bool FLandscapeGrassMapsBuilder::BuildGrassMapsNowForComponents(
 				// guaranteed by UpdateTrackedComponents(), as long as we update all of the components
 				check(Component->ComputeGrassMapGenerationHash() == Component->GrassData->GenerationHash);
 #endif // WITH_EDITOR
+				if (FailedStates.Contains(State))
+				{
+					FailedStates.Remove(State);
+				}
 				UpToDateCount++;
 			}
+			if (State->Stage == EComponentStage::NotReady)
+			{
+				// if it's not ready because of shader reasons
+				if (!FailedStates.Contains(State) && !UE::Landscape::CanRenderGrassMap(Component))
+				{
+#if WITH_EDITOR
+					// in editor, try to force compilation to complete
+					UE::Landscape::CompileGrassMapShader(State->Component);
+					if (!UE::Landscape::CanRenderGrassMap(Component))
+#endif // WITH_EDITOR
+					{
+						// failed to compile shaders... we won't be able to build the grass map for this component
+						FailedStates.Add(State);
+					}
+				}
+			}
 		}
 
 		if (LastUpToDateCount != UpToDateCount)
@@ -875,6 +898,11 @@ bool FLandscapeGrassMapsBuilder::BuildGrassMapsNowForComponents(
 
 	UE_LOG(LogGrass, Verbose, TEXT("BuildGrassMapsNowForComponents() updated %d/%d components in %f seconds"), UpToDateCount, LandscapeComponents.Num(), FPlatformTime::Seconds() - StartTime);
 
+	if (UpToDateCount != LandscapeComponents.Num())
+	{
+		UE_LOG(LogGrass, Warning, TEXT("Failed to build grass maps for %d landscape components, check if you are using a render preview mode, or a non-SM5 capable render device"), LandscapeComponents.Num() - UpToDateCount);
+	}
+
 	return (UpToDateCount == LandscapeComponents.Num());
 }
 
