Fix crash when creating a landscape larger than the max allowed render target size, while using global edit layer merge.

This adds a size check when we initialize or resize the layer merge resources, and disables the layer merge update when it cannot create the proper size resources.

#changelist validated
#virtualized

diff --git a/Engine/Source/Runtime/Landscape/Classes/Landscape.h b/Engine/Source/Runtime/Landscape/Classes/Landscape.h
--- a/Engine/Source/Runtime/Landscape/Classes/Landscape.h
+++ b/Engine/Source/Runtime/Landscape/Classes/Landscape.h
@@ -417,7 +417,7 @@ private:
 	bool SupportsEditLayersLocalMerge();
 	bool HasNormalCaptureBPBrushLayer();
 
-	void CreateLayersRenderingResource(bool bUseNormalCapture);
+	bool CreateLayersRenderingResource(bool bUseNormalCapture);
 	void PrepareEditLayersLocalMergeResources();
 	void UpdateLayersContent(bool bInWaitForStreaming = false, bool bInSkipMonitorLandscapeEdModeChanges = false, bool bIntermediateRender = false, bool bFlushRender = false);
 	void MonitorShaderCompilation();
@@ -541,6 +541,9 @@ public:
 	UPROPERTY(Transient)
 	bool bEnableEditorLayersTick = true;
 
+	UPROPERTY(Transient, DuplicateTransient, TextExportTransient, NonPIEDuplicateTransient)
+	bool bWarnedGlobalMergeDimensionsExceeded = false;
+
 	UPROPERTY()
 	TArray<FLandscapeLayer> LandscapeLayers;
 
diff --git a/Engine/Source/Runtime/Landscape/Private/LandscapeEditLayers.cpp b/Engine/Source/Runtime/Landscape/Private/LandscapeEditLayers.cpp
--- a/Engine/Source/Runtime/Landscape/Private/LandscapeEditLayers.cpp
+++ b/Engine/Source/Runtime/Landscape/Private/LandscapeEditLayers.cpp
@@ -1832,12 +1832,12 @@ bool ALandscape::HasNormalCaptureBPBrushLayer()
 	return false;
 }
 
-void ALandscape::CreateLayersRenderingResource(bool bUseNormalCapture)
+bool ALandscape::CreateLayersRenderingResource(bool bUseNormalCapture)
 {
 	ULandscapeInfo* Info = GetLandscapeInfo();
 	if (Info == nullptr)
 	{
-		return;
+		return false;
 	}
 
 	const FIntPoint ComponentCounts = ComputeComponentCounts();
@@ -1845,7 +1845,7 @@ void ALandscape::CreateLayersRenderingResource(bool bUseNormalCapture)
 	// No components, can't update the render targets	
 	if (ComponentCounts.X <= 0 || ComponentCounts.Y <= 0)
 	{
-		return;
+		return false;
 	}
 
 	ALandscape* Landscape = GetLandscapeActor();
@@ -1882,6 +1882,20 @@ void ALandscape::CreateLayersRenderingResource(bool bUseNormalCapture)
 			}
 		}
 
+		//static bool bWarned = false;
+		if (CurrentMipSizeX > GRHIGlobals.MaxTextureDimensions || CurrentMipSizeY > GRHIGlobals.MaxTextureDimensions)
+		{
+			if (!bWarnedGlobalMergeDimensionsExceeded)
+			{
+				UE_LOG(LogLandscape, Error, TEXT("Cannot initialize resources for Landscape Layer Merge because the current device does not support render targets of the required size.  Please reduce landscape size, or use a different render device, or try to enable local merge with `landscape.EditLayersLocalMerge.Enable 1` (local merge works only if no landscape blueprint brushes are used)"));
+				bWarnedGlobalMergeDimensionsExceeded = true;
+			}
+			return false;
+		}
+
+		// once the issue is fixed, clear the warn flag
+		bWarnedGlobalMergeDimensionsExceeded = false;
+
 		bool bCreateFromScratch = (Landscape->HeightmapRTList.Num() == 0);
 		if (bCreateFromScratch)
 		{
@@ -1986,6 +2000,7 @@ void ALandscape::CreateLayersRenderingResource(bool bUseNormalCapture)
 
 		InitializeLayersWeightmapResources();
 	}
+	return true;
 }
 
 void ALandscape::ToggleCanHaveLayersContent()
@@ -8278,6 +8293,12 @@ void ALandscape::UpdateLayersContent(bool bInWaitForStreaming, bool bInSkipMonit
 		InitializeLayers(bUseNormalCapture);
 	}
 
+	if (!bLandscapeLayersAreInitialized)
+	{
+		// we failed to initialize layers, cannot continue
+		return;
+	}
+
 	if (!bInSkipMonitorLandscapeEdModeChanges)
 	{
 		MonitorLandscapeEdModeChanges();
@@ -8632,10 +8653,11 @@ void ALandscape::InitializeLayers(bool bUseNormalCapture)
 {
 	check(HasLayersContent());
 
-	CreateLayersRenderingResource(bUseNormalCapture);
-	InitializeLandscapeLayersWeightmapUsage();
-
-	bLandscapeLayersAreInitialized = true;
+	if (CreateLayersRenderingResource(bUseNormalCapture))
+	{
+		InitializeLandscapeLayersWeightmapUsage();
+		bLandscapeLayersAreInitialized = true;
+	}
 }
 
 void ALandscape::OnPreSave()
