Control Nanite Position Precision from specific project Blueprint script

diff --git a/Engine/Source/Runtime/Landscape/Classes/Landscape.h b/Engine/Source/Runtime/Landscape/Classes/Landscape.h
--- a/Engine/Source/Runtime/Landscape/Classes/Landscape.h
+++ b/Engine/Source/Runtime/Landscape/Classes/Landscape.h
@@ -408,6 +408,9 @@ public:
 
 	LANDSCAPE_API void EnableNaniteSkirts(bool bInEnable, float InSkirtDepth, bool bInShouldDirtyPackage);
 
+	/** Set the target precision on nanite vertex position.  Precision is set to approximately (2^-InPrecision) in world units. */
+	LANDSCAPE_API void SetNanitePositionPrecision(int32 InPrecision,  bool bInShouldDirtyPackage);
+
 	LANDSCAPE_API void SetDisableRuntimeGrassMapGeneration(bool bInDisableRuntimeGrassMapGeneration);
 
 protected:
diff --git a/Engine/Source/Runtime/Landscape/Private/LandscapeEdit.cpp b/Engine/Source/Runtime/Landscape/Private/LandscapeEdit.cpp
--- a/Engine/Source/Runtime/Landscape/Private/LandscapeEdit.cpp
+++ b/Engine/Source/Runtime/Landscape/Private/LandscapeEdit.cpp
@@ -3836,6 +3836,31 @@ TSharedRef<UE::Landscape::Nanite::FAsyncBuildData> ALandscapeProxy::MakeAsyncNan
 	return AsyncBuildData;
 }
 
+LANDSCAPE_API void ALandscape::SetNanitePositionPrecision(int32 InPrecision, bool bInShouldDirtyPackage)
+{
+	NanitePositionPrecision = InPrecision;
+
+	// TODO [chris.tchou] : We should make a consolidated 'value changed' path, unifying this with PostEditChangeProperty
+	InvalidateOrUpdateNaniteRepresentation(/*bInCheckContentId*/true, /*InTargetPlatform*/nullptr);
+	MarkComponentsRenderStateDirty();
+	Modify(bInShouldDirtyPackage);
+
+	if (ULandscapeInfo* LandscapeInfo = GetLandscapeInfo())
+	{
+		LandscapeInfo->ForEachLandscapeProxy([&](ALandscapeProxy* Proxy)
+			{
+				if (Proxy != nullptr)
+				{
+					Proxy->SynchronizeSharedProperties(this);
+					Proxy->InvalidateOrUpdateNaniteRepresentation(/*bInCheckContentId*/true, /*InTargetPlatform*/nullptr);
+					Proxy->MarkComponentsRenderStateDirty();
+					Proxy->Modify(bInShouldDirtyPackage);
+				}
+				return true;
+			});
+	}
+}
+
 bool ALandscapeProxy::ExportToRawMesh(const FRawMeshExportParams& InExportParams, FMeshDescription& OutRawMesh) const
 {
 	FRawMeshExportParams ExportParams = InExportParams;
