Fix for rare crash accessing null pointer when painting on landscape

diff --git a/Engine/Source/Runtime/Landscape/Private/LandscapeEditInterface.cpp b/Engine/Source/Runtime/Landscape/Private/LandscapeEditInterface.cpp
--- a/Engine/Source/Runtime/Landscape/Private/LandscapeEditInterface.cpp
+++ b/Engine/Source/Runtime/Landscape/Private/LandscapeEditInterface.cpp
@@ -4364,12 +4364,18 @@ void FLandscapeTextureDataInterface::CopyTextureChannel(UTexture2D* Dest, int32
 
 	for( int32 MipIdx=0;MipIdx<DestDataInfo->NumMips();MipIdx++ )
 	{
-		uint8* DestTextureData = (uint8*)DestDataInfo->GetMipData(MipIdx) + ChannelOffsets[DestChannel];
-		uint8* SrcTextureData = (uint8*)SrcDataInfo->GetMipData(MipIdx) + ChannelOffsets[SrcChannel];
+		uint8* DestTextureData = (uint8*)DestDataInfo->GetMipData(MipIdx);
+		uint8* SrcTextureData = (uint8*)SrcDataInfo->GetMipData(MipIdx);
 
-		for( int32 i=0;i<FMath::Square(MipSize);i++ )
+		if (DestTextureData && SrcTextureData)
 		{
-			DestTextureData[i*4] = SrcTextureData[i*4];
+			DestTextureData += ChannelOffsets[DestChannel];
+			SrcTextureData += ChannelOffsets[SrcChannel];
+
+			for (int32 i = 0; i < FMath::Square(MipSize); i++)
+			{
+				DestTextureData[i * 4] = SrcTextureData[i * 4];
+			}
 		}
 
 		DestDataInfo->AddMipUpdateRegion(MipIdx, 0, 0, MipSize-1, MipSize-1);
