Fix to allow Landscape grass maps to update in editor when an out-of-date landscape is level instanced into the current level.

diff --git a/Engine/Source/Runtime/Landscape/Private/LandscapeSubsystem.cpp b/Engine/Source/Runtime/Landscape/Private/LandscapeSubsystem.cpp
--- a/Engine/Source/Runtime/Landscape/Private/LandscapeSubsystem.cpp
+++ b/Engine/Source/Runtime/Landscape/Private/LandscapeSubsystem.cpp
@@ -78,11 +78,21 @@ ULandscapeSubsystem::~ULandscapeSubsystem()
 void ULandscapeSubsystem::RegisterActor(ALandscapeProxy* Proxy)
 {
 	Proxies.AddUnique(TWeakObjectPtr<ALandscapeProxy>(Proxy));
+	
+	if (ALandscape* LandscapeActor = Cast<ALandscape>(Proxy))
+	{
+		LandscapeActors.AddUnique(TWeakObjectPtr<ALandscape>(LandscapeActor));
+	}
 }
 
 void ULandscapeSubsystem::UnregisterActor(ALandscapeProxy* Proxy)
 {
 	Proxies.Remove(TWeakObjectPtr<ALandscapeProxy>(Proxy));
+
+	if (ALandscape* LandscapeActor = Cast<ALandscape>(Proxy))
+	{
+		LandscapeActors.Remove(TWeakObjectPtr<ALandscape>(LandscapeActor));
+	}
 }
 
 void ULandscapeSubsystem::Initialize(FSubsystemCollectionBase& Collection)
@@ -179,6 +189,7 @@ void ULandscapeSubsystem::Deinitialize()
 	delete NotificationManager;
 #endif
 	Proxies.Empty();
+	LandscapeActors.Empty();
 
 	Super::Deinitialize();
 }
@@ -356,6 +367,21 @@ void ULandscapeSubsystem::Tick(float DeltaTime)
 	// TODO [chris.tchou] We should drop the usage of TWeakObjectPtr to reduce the dereference cost, can rely on register/unregister instead
 	bool bAllProxiesReadyForGrassMapGeneration = true;
 	bool bAllProxiesRuntimeGrassMapsDisabled = true;
+
+	for (TWeakObjectPtr<ALandscape> ActorPtr : LandscapeActors)
+	{
+		if (ALandscape* Landscape = ActorPtr.Get())
+		{
+			// if either of these things are true, then we wait for them to complete before running ANY grass map updates..
+			bool bLandscapeToolIsModifyingLandscape = !Landscape->bGrassUpdateEnabled;
+			bool bLandscapeLayerMergeWillBePerformedSoon = !Landscape->IsUpToDate() && Landscape->GetLandscapeInfo()->SupportsLandscapeEditing();
+			if (bLandscapeToolIsModifyingLandscape || bLandscapeLayerMergeWillBePerformedSoon)
+			{
+				bAllProxiesReadyForGrassMapGeneration = false;
+			}
+		}
+	}
+
 	static TArray<ALandscapeProxy*> ActiveProxies;
 	ActiveProxies.Reset(Proxies.Num());
 	for (TWeakObjectPtr<ALandscapeProxy> ProxyPtr : Proxies)
@@ -379,14 +405,6 @@ void ULandscapeSubsystem::Tick(float DeltaTime)
 					{
 						Component->UpdateGrassTypes();
 					}
-
-					// before starting grass map generation in editor, ensure layers are up to date
-					// and grass update is enabled (if disabled then a tool is currently operating)
-					ALandscape* Landscape = Proxy->GetLandscapeActor();
-					if ((Landscape == nullptr) || !Landscape->bGrassUpdateEnabled || !Landscape->IsUpToDate())
-					{
-						bAllProxiesReadyForGrassMapGeneration = false;
-					}
 				}
 #endif // WITH_EDITOR
 
diff --git a/Engine/Source/Runtime/Landscape/Public/LandscapeSubsystem.h b/Engine/Source/Runtime/Landscape/Public/LandscapeSubsystem.h
--- a/Engine/Source/Runtime/Landscape/Public/LandscapeSubsystem.h
+++ b/Engine/Source/Runtime/Landscape/Public/LandscapeSubsystem.h
@@ -8,6 +8,7 @@
 #include "LandscapeSubsystem.generated.h"
 
 class ALandscapeProxy;
+class ALandscape;
 class AWorldSettings;
 class IConsoleVariable;
 class ULandscapeInfo;
@@ -153,6 +154,7 @@ private:
 	void HandlePostGarbageCollect();
 
 	bool bIsGrassCreationPrioritized = false;
+	TArray<TWeakObjectPtr<ALandscape>> LandscapeActors;
 	TArray<TWeakObjectPtr<ALandscapeProxy>> Proxies;
 	FDelegateHandle OnNaniteWorldSettingsChangedHandle;
 
