Build fix for non-editor compile issue caused by 31435399

diff --git a/Engine/Source/Runtime/Landscape/Private/LandscapeSubsystem.cpp b/Engine/Source/Runtime/Landscape/Private/LandscapeSubsystem.cpp
--- a/Engine/Source/Runtime/Landscape/Private/LandscapeSubsystem.cpp
+++ b/Engine/Source/Runtime/Landscape/Private/LandscapeSubsystem.cpp
@@ -368,6 +368,7 @@ void ULandscapeSubsystem::Tick(float DeltaTime)
 	bool bAllProxiesReadyForGrassMapGeneration = true;
 	bool bAllProxiesRuntimeGrassMapsDisabled = true;
 
+#if WITH_EDITOR
 	for (TWeakObjectPtr<ALandscape> ActorPtr : LandscapeActors)
 	{
 		if (ALandscape* Landscape = ActorPtr.Get())
@@ -381,6 +382,7 @@ void ULandscapeSubsystem::Tick(float DeltaTime)
 			}
 		}
 	}
+#endif // WITH_EDITOR
 
 	static TArray<ALandscapeProxy*> ActiveProxies;
 	ActiveProxies.Reset(Proxies.Num());
