Remove unnecessary and potentially harmful data stripping in ALandscapeProxy::PreSave()

They have been replaced by other methods that don't destroy the data in memory.
The collision component uses NeedsLoadForClient/Server() callbacks, and the grass data is omitted from Serialize() on cook.

Also fixes an issue where the cvar override on physics stripping did not affect the new path via callback.

diff --git a/Engine/Source/Runtime/Landscape/Private/Landscape.cpp b/Engine/Source/Runtime/Landscape/Private/Landscape.cpp
--- a/Engine/Source/Runtime/Landscape/Private/Landscape.cpp
+++ b/Engine/Source/Runtime/Landscape/Private/Landscape.cpp
@@ -201,11 +201,6 @@ static FAutoConsoleVariable CVarStripLayerTextureMipsOnLoad(
 	false,
 	TEXT("Remove (on load) the mip chain from textures used in layers which don't require them"));
 
-static FAutoConsoleVariable CVarAllowPhysicsStripping(
-	TEXT("landscape.AllowPhysicsStripping"),
-	true,
-	TEXT("Enables the conditional stripping of physics data during cook.  Disabling this means the bStripPhysicsWhenCooked* will be ignored."));
-
 static FAutoConsoleVariable CVarAllowGrassStripping(
 	TEXT("landscape.AllowGrassStripping"),
 	true,
@@ -3681,59 +3676,6 @@ void ALandscapeProxy::PreSave(FObjectPreSaveContext ObjectSaveContext)
 			}
 		}
 	}
-
-	// Strip data according to flags when cooked (and not cooking for editor)
-	if (ObjectSaveContext.IsCooking() && !ObjectSaveContext.GetTargetPlatform()->AllowsEditorObjects())
-	{
-		if (CVarAllowPhysicsStripping->GetBool() &&
-			((bStripPhysicsWhenCookedClient && bStripPhysicsWhenCookedServer) ||
-			(bStripPhysicsWhenCookedClient && ObjectSaveContext.GetTargetPlatform()->IsClientOnly()) ||
-			(bStripPhysicsWhenCookedServer && ObjectSaveContext.GetTargetPlatform()->IsServerOnly())))
-		{
-			// Clear old CollisionComponent containers
-			CollisionComponents.Empty();
-
-			// Destroy any owned collision components
-			TInlineComponentArray<ULandscapeHeightfieldCollisionComponent*> CollisionComps;
-			GetComponents(CollisionComps);
-			for (ULandscapeHeightfieldCollisionComponent* Component : CollisionComps)
-			{
-				Component->DestroyComponent();
-			}
-
-			// Destroy any attached but un-owned collision components (not sure what would be here that is not already caught above ^^)
-			TArray<USceneComponent*> AttachedCollisionComponents = RootComponent->GetAttachChildren().FilterByPredicate(
-				[](USceneComponent* Component)
-				{
-					return Cast<ULandscapeHeightfieldCollisionComponent>(Component);
-				});
-
-			for (USceneComponent* Component : AttachedCollisionComponents)
-			{
-				Component->DestroyComponent();
-			}
-
-			// ULandscapeComponents have refs... null them
-			for (ULandscapeComponent* Component : LandscapeComponents)
-			{
-				Component->DestroyCollisionData();
-				Component->SetCollisionComponent(nullptr);
-			}
-		}
-		
-		if (CVarAllowGrassStripping->GetBool() &&
-			((bStripGrassWhenCookedClient && bStripGrassWhenCookedServer) ||
-			(bStripGrassWhenCookedClient && ObjectSaveContext.GetTargetPlatform()->IsClientOnly()) ||
-			(bStripGrassWhenCookedServer && ObjectSaveContext.GetTargetPlatform()->IsServerOnly())))
-		{
-			for (ULandscapeComponent* Component : LandscapeComponents)
-			{
-				TUniquePtr<FLandscapeComponentGrassData> NewGrassData = MakeUnique<FLandscapeComponentGrassData>();
-				Component->GrassData = MakeShareable(NewGrassData.Release());
-				Component->GrassData->NumElements = 0;
-			}
-		}
-	}
 #endif // WITH_EDITOR
 }
 
diff --git a/Engine/Source/Runtime/Landscape/Private/LandscapeCollision.cpp b/Engine/Source/Runtime/Landscape/Private/LandscapeCollision.cpp
--- a/Engine/Source/Runtime/Landscape/Private/LandscapeCollision.cpp
+++ b/Engine/Source/Runtime/Landscape/Private/LandscapeCollision.cpp
@@ -102,6 +102,11 @@ static TAutoConsoleVariable<bool> CVarLandscapeCollisionMeshShowPhysicalMaterial
 	TEXT("When enabled, vertex colors of the collision mesh are chosen based on the physical material"),
 	ECVF_RenderThreadSafe);
 
+static FAutoConsoleVariable CVarAllowPhysicsStripping(
+	TEXT("landscape.AllowPhysicsStripping"),
+	true,
+	TEXT("Enables the conditional stripping of physics data during cook.  Disabling this means the bStripPhysicsWhenCooked* will be ignored."));
+
 #if ENABLE_COOK_STATS
 namespace LandscapeCollisionCookStats
 {
@@ -2537,7 +2542,8 @@ bool ULandscapeHeightfieldCollisionComponent::NeedsLoadForClient() const
 	ALandscapeProxy* Proxy = GetLandscapeProxy();
 	if (ensure(Proxy))
 	{
-		return !Proxy->bStripPhysicsWhenCookedClient;
+		bool bStrip = Proxy->bStripPhysicsWhenCookedClient && CVarAllowPhysicsStripping->GetBool();
+		return !bStrip;
 	}
 	return true;
 }
@@ -2547,7 +2553,8 @@ bool ULandscapeHeightfieldCollisionComponent::NeedsLoadForServer() const
 	ALandscapeProxy* Proxy = GetLandscapeProxy();
 	if (ensure(Proxy))
 	{
-		return !Proxy->bStripPhysicsWhenCookedServer;
+		bool bStrip = Proxy->bStripPhysicsWhenCookedServer && CVarAllowPhysicsStripping->GetBool();
+		return !bStrip;
 	}
 	return true;
 }
