Added system to invalidate VSM pages when using (non-Nanite) landscape, to hide shadow artifacts induced by the vertex morphing system of standard landscape :
* Relies on pre-computing max height delas from mip-to-mip for every landscape component
* Invalidation occurs when the evaluated max delta between the heights at the LOD value that was active when VSM was last cached is different enough from the heights at the current LOD value (for a given persistent view), based on a height threshold that is tweakable per landscape and overridable per landscape proxy
* Invalidation doesn't occur when Nanite landscape is used
* The invalidation rate is decreased as the LOD value goes up, controlled by a screen size parameter in the landscape actor (overridable per proxy), under which no invalidation will occur. This avoids over-invalidating VSM on higher LOD values, since they tend to occupy less real estate and therefore don't need to have perfect shadows
* Added per landscape (overridable per-proxy) shadow map bias to help with this problem too
* Added 3 non-shipping CVars to help tweak those 3 parameters in-game (landscape.OverrideNonNaniteVirtualShadowMapConstantDepthBiasOverride, landscape.OverrideNonNaniteVirtualShadowMapInvalidationHeightErrorThreshold, landscape.OverrideNonNaniteVirtualShadowMapInvalidationScreenSizeLimit)
* The whole invalidation system can be enabled/disabled via CVar landscape.AllowNonNaniteVirtualShadowMapInvalidation
* Added another CVar (landscape.NonNaniteVirtualShadowMapInvalidationLODAttenuationExponent) to tweak the screen-size-dependent invalidation rate curve shape

Misc:
* Fixed landscape.DumpLODs command : now works without parameter and can be used several times
* Changed CachedSectionLODValues system for landscape to be usable both for ViewState-less views (one-shot renders like water info custom render passes) and "persistent" views (those that have  FViewStateInterface associated). This allows to keep track of the previous frame's CachedSectionLODValues for a given persistent view and removes the previous hacky method
* Also, prevented landscape render systems to have their LOD section values cleared when a render is made on another scene (e.g. any time a thumbnail renders in editor or when custom render passes execute) : now only the render systems that pertain to the scene being rendered get udpated
* Fixed game thread data being accessed in render thread landscape scene proxy (component)
* Landscape UX fixes
* Fixed some landscape CVars which were not render-thread safe
* Removed member RHI buffer refs LandscapeLODDataBuffer and LandscapeIndirectionBuffer which didn't serve any purpose
#tests Editor, cooked game,, cooked game

diff --git a/Engine/Shaders/Private/LandscapeVertexFactory.ush b/Engine/Shaders/Private/LandscapeVertexFactory.ush
--- a/Engine/Shaders/Private/LandscapeVertexFactory.ush
+++ b/Engine/Shaders/Private/LandscapeVertexFactory.ush
@@ -828,4 +828,12 @@ uint VertexFactoryGetPrimitiveId(FVertexFactoryInterpolantsVSToPS Interpolants)
 #endif // VF_USE_PRIMITIVE_SCENE_DATA
 }
 
+// Tell the VSM depth shader that we are supplying a constant bias
+#define VF_USE_VSM_CONSTANT_BIAS 1
+
+float VertexFactoryGetNonNaniteVirtualShadowMapConstantDepthBias()
+{
+	return LandscapeParameters.NonNaniteVirtualShadowMapConstantDepthBias;
+}
+
 #include "VertexFactoryDefaultInterface.ush"
diff --git a/Engine/Source/Editor/LandscapeEditor/Private/LandscapeEdMode.cpp b/Engine/Source/Editor/LandscapeEditor/Private/LandscapeEdMode.cpp
--- a/Engine/Source/Editor/LandscapeEditor/Private/LandscapeEdMode.cpp
+++ b/Engine/Source/Editor/LandscapeEditor/Private/LandscapeEdMode.cpp
@@ -3707,6 +3707,9 @@ ALandscape* FEdModeLandscape::ChangeComponentSetting(int32 NumComponentsX, int32
 			NewLandscape->bFillCollisionUnderLandscapeForNavmesh = OldLandscape->bFillCollisionUnderLandscapeForNavmesh;
 			NewLandscape->NavigationGeometryGatheringMode = OldLandscape->NavigationGeometryGatheringMode;
 			NewLandscape->bUseLandscapeForCullingInvisibleHLODVertices = OldLandscape->bUseLandscapeForCullingInvisibleHLODVertices;
+			NewLandscape->NonNaniteVirtualShadowMapConstantDepthBias = OldLandscape->NonNaniteVirtualShadowMapConstantDepthBias;
+			NewLandscape->NonNaniteVirtualShadowMapInvalidationHeightErrorThreshold = OldLandscape->NonNaniteVirtualShadowMapInvalidationHeightErrorThreshold;
+			NewLandscape->NonNaniteVirtualShadowMapInvalidationScreenSizeLimit = OldLandscape->NonNaniteVirtualShadowMapInvalidationScreenSizeLimit;
 
 			NewLandscape->BodyInstance.SetCollisionProfileName(OldLandscape->BodyInstance.GetCollisionProfileName());
 			if (NewLandscape->BodyInstance.DoesUseCollisionProfile() == false)
diff --git a/Engine/Source/Runtime/Landscape/Classes/LandscapeComponent.h b/Engine/Source/Runtime/Landscape/Classes/LandscapeComponent.h
--- a/Engine/Source/Runtime/Landscape/Classes/LandscapeComponent.h
+++ b/Engine/Source/Runtime/Landscape/Classes/LandscapeComponent.h
@@ -466,6 +466,17 @@ class ULandscapeComponent : public UPrimitiveComponent
 	UPROPERTY()
 	FBox CachedLocalBox;
 
+	/** Maximum deltas between vertices and their counterparts from other mips. This mip-to-mip data is laid out in a contiguous array following the following pattern : 
+	*  Say, we have 5 "relevant" mips and [N -> M] is the delta from mip N to M (where M > N and M < (NumRelevantMips - 1)) then the array will contain : 
+	*  [0 -> 1], [0 -> 2], [0 -> 3], [1 -> 2], [1 -> 3], [2 -> 3]
+	*  i.e. for mip 0 : (NumRelevantMips - 1) deltas, then for mip 1 : (NumRelevantMips - 2) deltas, until mip == (NumRelevantMips - 2) : 1 delta
+	*  Note: a "relevant" mip is one with more than 1 vertex. i.e.:
+	*   - In the case of a 1x1 subsection, the last mip index (NumMips - 1) has a single pixel and is therefore not relevant (we cannot draw a landscape component with a single vertex!), hence the last relevant mip index will be NumMips - 2
+	*   - In the case of 2x2 subsections, the penultimate mip index (NumMips - 2) has 4 pixels, which means 4 subsections, each with a single pixel, and is therefore not relevant either, hence the last relevant mip index will be NumMips - 3
+	*/
+	UPROPERTY()
+	TArray<double> MipToMipMaxDeltas;
+
 #if WITH_EDITORONLY_DATA
 	UPROPERTY()
 	TLazyObjectPtr<ULandscapeHeightfieldCollisionComponent> CollisionComponent_DEPRECATED;
@@ -531,9 +542,9 @@ private:
 #endif // WITH_EDITORONLY_DATA
 
 	/** Cached list of grass types supported by the component's material.
-	  *	This is needed in a cooked build, as the grass types list is not available
-	  * on the cooked material.
-	  * Call UpdateGrassTypes() to ensure this array is up to date */
+	* This is needed in a cooked build, as the grass types list is not available
+	* on the cooked material.
+	* Call UpdateGrassTypes() to ensure this array is up to date */
 	UPROPERTY()
 	TArray<TObjectPtr<ULandscapeGrassType>> GrassTypes;
 
@@ -956,6 +967,13 @@ public:
 		SectionBaseY = InSectionBase.Y;
 	}
 
+	/** 
+	* Computes the number of mips that are actually usable, i.e.:
+	*  - For 1x1 subsection, the last mip is not usable (it has a single vertex)
+	*  - For 2x2 subsections, the last 2 mips are not usable (a single vertex per subsection)
+	*/
+	int32 GetNumRelevantMips() const;
+
 	/** @todo document */
 	const FGuid& GetLightingGuid() const
 	{
diff --git a/Engine/Source/Runtime/Landscape/Classes/LandscapeProxy.h b/Engine/Source/Runtime/Landscape/Classes/LandscapeProxy.h
--- a/Engine/Source/Runtime/Landscape/Classes/LandscapeProxy.h
+++ b/Engine/Source/Runtime/Landscape/Classes/LandscapeProxy.h
@@ -463,30 +463,35 @@ public:
 	float ComponentScreenSizeToUseSubSections;
 
 	/** This is the starting screen size used to calculate the distribution. You can increase the value if you want less LOD0 component, and you use very large landscape component. */
-	UPROPERTY(EditAnywhere, Category = "LOD Distribution", meta = (EditCondition = "!bUseScalableLODSettings", DisplayName = "LOD 0 Screen Size", ClampMin = "0.1", ClampMax = "10.0", UIMin = "1.0", UIMax = "10.0", LandscapeInherited))
+	UPROPERTY(EditAnywhere, Category = "LOD Distribution", meta = (EditCondition = "!bUseScalableLODSettings", DisplayName = "LOD 0 Screen Size", ClampMin = "0.1", ClampMax = "10.0", UIMin = "0.1", UIMax = "10.0", LandscapeInherited))
 	float LOD0ScreenSize = 0.5f;
 	
-	UPROPERTY(EditAnywhere, Category = LOD, AdvancedDisplay, meta = (ToolTip = "Specifies the LOD Group (Zero is No Group). All landscapes in the same group calculate their LOD together, allowing matching border LODs to fix geometry seams.", LandscapeInherited))
+	/** Specifies the LOD Group (Zero is No Group). All landscapes in the same group calculate their LOD together, allowing matching border LODs to fix geometry seams. */
+	UPROPERTY(EditAnywhere, Category = LOD, AdvancedDisplay, meta = (LandscapeInherited))
 	uint32 LODGroupKey = 0;
 
-	/** The distribution setting used to change the LOD 0 generation, 1.75 is the normal distribution, numbers influence directly the LOD0 proportion on screen. */
+	/** The distribution setting used to change the LOD 0 generation, 1.25 is the normal distribution, numbers influence directly the LOD0 proportion on screen. */
 	UPROPERTY(EditAnywhere, Category = "LOD Distribution", meta = (EditCondition = "!bUseScalableLODSettings", DisplayName = "LOD 0", ClampMin = "1.0", ClampMax = "10.0", UIMin = "1.0", UIMax = "10.0", LandscapeInherited))
 	float LOD0DistributionSetting = 1.25f;
 
-	/** The distribution setting used to change the LOD generation, 2 is the normal distribution, small number mean you want your last LODs to take more screen space and big number mean you want your first LODs to take more screen space. */
+	/** The distribution setting used to change the LOD generation, 3 is the normal distribution, small number mean you want your last LODs to take more screen space and big number mean you want your first LODs to take more screen space. */
 	UPROPERTY(EditAnywhere, Category = "LOD Distribution", meta = (EditCondition = "!bUseScalableLODSettings", DisplayName = "Other LODs", ClampMin = "1.0", ClampMax = "10.0", UIMin = "1.0", UIMax = "10.0", LandscapeInherited))
 	float LODDistributionSetting = 3.0f;
 
-	UPROPERTY(EditAnywhere, Category = "LOD Distribution", meta = (EditCondition = "bUseScalableLODSettings", DisplayName = "Scalable LOD 0 Screen Size", ClampMin = "0.1", ClampMax = "10.0", UIMin = "1.0", UIMax = "10.0", LandscapeInherited))
+	/** Scalable (per-quality) version of 'LOD 0 Screen Size'. */
+	UPROPERTY(EditAnywhere, Category = "LOD Distribution", meta = (EditCondition = "bUseScalableLODSettings", DisplayName = "Scalable LOD 0 Screen Size", ClampMin = "0.1", ClampMax = "10.0", UIMin = "0.1", UIMax = "10.0", LandscapeInherited))
 	FPerQualityLevelFloat ScalableLOD0ScreenSize = 0.5f;
 	
+	/** Scalable (per-quality) version of 'LOD 0'. */
 	UPROPERTY(EditAnywhere, Category = "LOD Distribution", meta = (EditCondition = "bUseScalableLODSettings", DisplayName = "Scalable LOD 0", ClampMin = "1.0", ClampMax = "10.0", UIMin = "1.0", UIMax = "10.0", LandscapeInherited))
 	FPerQualityLevelFloat ScalableLOD0DistributionSetting = 1.25f;
 	
+	/** Scalable (per-quality) version of 'Other LODs'. */
 	UPROPERTY(EditAnywhere, Category = "LOD Distribution", meta = (EditCondition = "bUseScalableLODSettings", DisplayName = "Scalable Other LODs", ClampMin = "1.0", ClampMax = "10.0", UIMin = "1.0", UIMax = "10.0", LandscapeInherited))
 	FPerQualityLevelFloat ScalableLODDistributionSetting = 3.0f;
 
-	UPROPERTY(EditAnywhere, Category = "LOD Distribution", meta = (LandscapeInherited, tooltip="Use ScalableLODDistributionSetting, ScalableLOD0DistributionSetting & ScalableLODDistributionSetting (ignoring r.LandscapeLOD0DistributionScale) instead of LODDistributionSetting, LOD0DistributionSetting & LODDistributionSetting respectively"))
+	/** Allows to specify LOD distribution settings per quality level. Using this will ignore the r.LandscapeLOD0DistributionScale CVar. */
+	UPROPERTY(EditAnywhere, Category = "LOD Distribution", meta = (LandscapeInherited))
 	bool bUseScalableLODSettings = false;
 
 	/** This controls the area that blends LOD between neighboring sections. At 1.0 it blends across the entire section, and lower numbers reduce the blend region to be closer to the boundary. */
@@ -678,7 +683,7 @@ public:
 	uint8 bCastStaticShadow : 1;
 
 	/** Control shadow invalidation behavior, in particular with respect to Virtual Shadow Maps and material effects like World Position Offset. */
-	UPROPERTY(EditAnywhere, BlueprintReadOnly, Category=Lighting, AdvancedDisplay, meta=(EditCondition="CastShadow", LandscapeOverridable))
+	UPROPERTY(EditAnywhere, BlueprintReadOnly, Category = Lighting, AdvancedDisplay, meta=(EditCondition="CastShadow", LandscapeOverridable))
 	EShadowCacheInvalidationBehavior ShadowCacheInvalidationBehavior;
 
 	/** Whether the object should cast contact shadows. This flag is only used if CastShadow is true. */
@@ -709,13 +714,34 @@ public:
 	FLightingChannels LightingChannels;
 
 	/** Whether to use the landscape material's vertical world position offset when calculating static lighting.
-		Note: Only z (vertical) offset is supported. XY offsets are ignored.
-		Does not work correctly with an XY offset map (mesh collision) */
-	UPROPERTY(EditAnywhere, AdvancedDisplay, Category=Lighting, meta = (LandscapeInherited))
+	* Note: Only z (vertical) offset is supported. XY offsets are ignored.
+	* Does not work correctly with an XY offset map (mesh collision) */
+	UPROPERTY(EditAnywhere, AdvancedDisplay, Category = Lighting, meta = (LandscapeInherited))
 	uint32 bUseMaterialPositionOffsetInStaticLighting:1;
 
+	/** Constant bias to handle the worst artifacts of the continuous LOD morphing when rendering to VSM.  
+	* Only applies when using non-Nanite landscape and VSM. */
+	UPROPERTY(EditAnywhere, AdvancedDisplay, Category = Lighting, meta = (LandscapeOverridable))
+	float NonNaniteVirtualShadowMapConstantDepthBias = 150.0f;
+
+	/** For non-Nanite landscape, cached VSM pages need to be invalidated when continuous LOD morphing introduces a height difference that is too large between the current landscape component's profile and the one that was used when the shadow was shadow was last cached.
+	* This height threshold (in Unreal units) controls this invalidation rate (a smaller threshold will reduce the likeliness of shadow artifacts, but will make the invalidations occur more frequently, which is not desirable in terms of performance.
+	* Disabled if 0.0.
+	* Only applies when using non-Nanite landscape and VSM. */
+	UPROPERTY(EditAnywhere, AdvancedDisplay, Category = Lighting, meta = (LandscapeOverridable, UIMin = 0.0, ClampMin = 0.0))
+	float NonNaniteVirtualShadowMapInvalidationHeightErrorThreshold = 250.0f;
+
+	/** Screen size under which VSM invalidation stops occurring.
+	* As the height difference between 2 mip levels increases when the LOD level increases (because of undersampling), VSM pages tend to be invalidated more frequently even though it's getting less and less relevant to do so, since this will mean that the screen size of the landscape section decreases, thus the artifacts actually become less noticeable.
+	* We therefore artificially attenuate the VSM invalidation rate as the screen size decreases, to avoid invalidating VSM pages too often, as it becomes less and less impactful. 
+	* A higher value will accentuate this attenuation (better performance but more artifacts) and vice versa.
+	* If 0.0, the attenuation of the VSM invalidation rate will be disabled entirely.
+	* Only applies when using non-Nanite landscape and VSM. */
+	UPROPERTY(EditAnywhere, AdvancedDisplay, Category = Lighting, meta = (LandscapeOverridable, ClampMin = "0.0", ClampMax = "10.0", UIMin = "0.0", UIMax = "10.0"))
+	float NonNaniteVirtualShadowMapInvalidationScreenSizeLimit = 0.2f;
+
 	/** If true, the Landscape will be rendered in the CustomDepth pass (usually used for outlines) */
-	UPROPERTY(EditAnywhere, AdvancedDisplay, BlueprintReadOnly, Category=Rendering, meta=(DisplayName = "Render CustomDepth Pass", LandscapeOverridable))
+	UPROPERTY(EditAnywhere, AdvancedDisplay, BlueprintReadOnly, Category = Rendering, meta=(DisplayName = "Render CustomDepth Pass", LandscapeOverridable))
 	uint32 bRenderCustomDepth:1;
 
 	/** Mask used for stencil buffer writes. */
diff --git a/Engine/Source/Runtime/Landscape/Private/Landscape.cpp b/Engine/Source/Runtime/Landscape/Private/Landscape.cpp
--- a/Engine/Source/Runtime/Landscape/Private/Landscape.cpp
+++ b/Engine/Source/Runtime/Landscape/Private/Landscape.cpp
@@ -1491,6 +1491,12 @@ void ULandscapeComponent::PostLoad()
 		CollisionComponentRef = CollisionComponent_DEPRECATED.Get();
 		CollisionComponent_DEPRECATED = nullptr;
 	}
+
+	// If mip-to-mip info is missing, recompute them (they were introduced later) :
+	if (MipToMipMaxDeltas.IsEmpty())
+	{
+		UpdateCachedBounds();
+	}
 #endif // !WITH_EDITORONLY_DATA
 
 #endif // WITH_EDITOR
@@ -1880,6 +1886,17 @@ ALandscapeProxy* ULandscapeComponent::GetLandscapeProxy() const
 	return CastChecked<ALandscapeProxy>(GetOuter());
 }
 
+int32 ULandscapeComponent::GetNumRelevantMips() const
+{
+	const int32 TextureSize = (SubsectionSizeQuads + 1) * NumSubsections;
+	const int32 NumTextureMips = FMath::FloorLog2(TextureSize) + 1;
+	// We actually only don't care about the last texture mip, since a 1 vertex landscape is meaningless. When using 2x2 subsections, we can even drop an additional mip 
+	//  as the 4 texels of the penultimate mip will be identical (i.e. 4 sub-sections of 1 vertex are equally meaningless) :
+	const int32 NumRelevantMips = (NumSubsections > 1) ? (NumTextureMips - 2) : (NumTextureMips - 1);
+	check(NumRelevantMips > 0);
+	return NumRelevantMips;
+}
+
 const FMeshMapBuildData* ULandscapeComponent::GetMeshMapBuildData() const
 {
 	AActor* Owner = GetOwner();
@@ -4427,6 +4444,9 @@ void ALandscapeProxy::GetSharedProperties(ALandscapeProxy* Landscape)
 		VirtualTextureRenderPassType = Landscape->VirtualTextureRenderPassType;
 		bEnableNanite = Landscape->bEnableNanite;
 		ShadowCacheInvalidationBehavior = Landscape->ShadowCacheInvalidationBehavior;
+		NonNaniteVirtualShadowMapConstantDepthBias = Landscape->NonNaniteVirtualShadowMapConstantDepthBias;
+		NonNaniteVirtualShadowMapInvalidationHeightErrorThreshold = Landscape->NonNaniteVirtualShadowMapInvalidationHeightErrorThreshold;
+		NonNaniteVirtualShadowMapInvalidationScreenSizeLimit = Landscape->NonNaniteVirtualShadowMapInvalidationScreenSizeLimit;
 
 		bUseCompressedHeightmapStorage = Landscape->bUseCompressedHeightmapStorage;
 #if WITH_EDITORONLY_DATA
diff --git a/Engine/Source/Runtime/Landscape/Private/LandscapeCulling.cpp b/Engine/Source/Runtime/Landscape/Private/LandscapeCulling.cpp
--- a/Engine/Source/Runtime/Landscape/Private/LandscapeCulling.cpp
+++ b/Engine/Source/Runtime/Landscape/Private/LandscapeCulling.cpp
@@ -616,7 +616,7 @@ static void ComputeSectionIntermediateData(FRDGBuilder& GraphBuilder, TArrayView
 			ViewStateIntermediates.SectionsBufferRDG = nullptr;
 
 			const FLandscapeRenderSystem& RenderSystem = *LandscapeRenderSystems.FindChecked(CullingEntry.LandscapeKey);
-			const TResourceArray<float>& SectionLODValues = RenderSystem.CachedSectionLODValues.FindRef(&View);
+			const TResourceArray<float>& SectionLODValues = RenderSystem.GetCachedSectionLODValues(View);
 
 			for (int32 SectionIdx = 0; SectionIdx < SectionLODValues.Num(); ++SectionIdx)
 			{
diff --git a/Engine/Source/Runtime/Landscape/Private/LandscapeEdit.cpp b/Engine/Source/Runtime/Landscape/Private/LandscapeEdit.cpp
--- a/Engine/Source/Runtime/Landscape/Private/LandscapeEdit.cpp
+++ b/Engine/Source/Runtime/Landscape/Private/LandscapeEdit.cpp
@@ -89,6 +89,7 @@ LandscapeEdit.cpp: Landscape editing
 #include "DataDrivenShaderPlatformInfo.h"
 #include "Misc/TransactionObjectEvent.h"
 #endif
+#include "Algo/Compare.h"
 #include "Algo/Count.h"
 #include "Algo/Transform.h"
 #include "Algo/ForEach.h"
@@ -146,8 +147,125 @@ void ULandscapeComponent::Init(int32 InBaseX, int32 InBaseY, int32 InComponentSi
 }
 
 #if WITH_EDITOR
+namespace UE::Landscape::Private
+{
+	/** 
+	* Struct to hold, for a given 4 neighboring pixels of a given mip of the heightmap, information about the difference between those and the resulting pixel in the next mip. 
+	*  This allows landscape components to evaluate the error (height delta) between any 2 mip levels
+	*/
+	struct FQuadHeightInfo
+	{
+		FQuadHeightInfo() = default;
+		FQuadHeightInfo(double InQuadVertices[4])
+		{
+			Min = FMath::Min(InQuadVertices[0], FMath::Min3(InQuadVertices[1], InQuadVertices[2], InQuadVertices[3]));
+			Max = FMath::Max(InQuadVertices[0], FMath::Max3(InQuadVertices[1], InQuadVertices[2], InQuadVertices[3]));
+			Average = (InQuadVertices[0] + InQuadVertices[1] + InQuadVertices[2] + InQuadVertices[3]) / 4.0;
+		}
+
+		double Min = DBL_MAX;
+		double Max = DBL_MIN;
+		double Average = 0.0f;
+	};
+
+	int32 ComputeQuadInfosOffset(int32 InMipIndex, int32 InTextureSize)
+	{
+		int32 NumQuadsForMip = FMath::Square(InTextureSize / 2);
+		int32 Offset = 0;
+		for (int32 X = 0; X < InMipIndex; ++X)
+		{
+			Offset += NumQuadsForMip;
+			NumQuadsForMip /= 4;
+		}
+		return Offset;
+	}
+
+	int32 ComputeQuadInfosCount(int32 InNumRelevantMips, int32 InTextureSize)
+	{
+		int32 Count = 0; 
+		int32 NumQuadsForMip = FMath::Square(InTextureSize / 2);
+		for (int32 MipIndex = 0; MipIndex < InNumRelevantMips; ++MipIndex)
+		{
+			Count += NumQuadsForMip;
+			NumQuadsForMip /= 4;
+		}
+		return Count;
+	}
+
+	TArrayView<FQuadHeightInfo> GetMipQuadInfosForMip(const TArrayView<FQuadHeightInfo>& InQuadInfos, int32 InMipIndex, int32 InTextureSize)
+	{
+		int32 MipOffset = ComputeQuadInfosOffset(InMipIndex, InTextureSize);
+		int32 NextMipOffset = ComputeQuadInfosOffset(InMipIndex + 1, InTextureSize);
+		return MakeArrayView(InQuadInfos.GetData() + MipOffset, NextMipOffset - MipOffset);
+	}
+
+	TArrayView<const FQuadHeightInfo> GetMipQuadInfosForMipConst(const TArrayView<const FQuadHeightInfo>& InQuadInfos, int32 InMipIndex, int32 InTextureSize)
+	{
+		TArrayView<FQuadHeightInfo> Result = GetMipQuadInfosForMip(MakeArrayView(const_cast<FQuadHeightInfo*>(InQuadInfos.GetData()), InQuadInfos.Num()), InMipIndex, InTextureSize);
+		return MakeArrayView<const FQuadHeightInfo>(Result.GetData(), Result.Num());
+	}
+
+	void ComputeMaxDelta(const FQuadHeightInfo& InSourceQuadInfo, const FQuadHeightInfo& InDestinationQuadInfo, double& InOutMaxDelta)
+	{
+		InOutMaxDelta = FMath::Max3(
+			InOutMaxDelta,
+			FMath::Abs(InSourceQuadInfo.Min - InDestinationQuadInfo.Average), 
+			FMath::Abs(InSourceQuadInfo.Max - InDestinationQuadInfo.Average));
+	}
+
+	TArray<double> ComputeMipToMipMaxDeltas(const TArrayView<const FQuadHeightInfo>& InAllMipsQuadInfos, int32 InNumTextureMips, int32 InNumRelevantMips)
+	{
+		TRACE_CPUPROFILER_EVENT_SCOPE(ComputeMipToMipMaxDeltas);
+		const int32 TextureSize = 1 << (InNumTextureMips - 1);
+
+		TArray<double> MipToMipMaxDeltas;
+		MipToMipMaxDeltas.AddZeroed(ComputeMipToMipMaxDeltasCount(InNumRelevantMips));
+
+		for (int32 SourceMipIndex = 0; SourceMipIndex < InNumRelevantMips - 1; ++SourceMipIndex)
+		{
+			TArrayView<const FQuadHeightInfo> SourceMipQuadInfos = GetMipQuadInfosForMipConst(InAllMipsQuadInfos, SourceMipIndex, TextureSize);
+			const int32 NumQuadsForSourceMip = SourceMipQuadInfos.Num();
+			const int32 SourceMipQuadsStride = (TextureSize >> SourceMipIndex) / 2;
+			check(SourceMipQuadsStride * SourceMipQuadsStride == NumQuadsForSourceMip);
+			const int32 NumMaxDeltasForSourceMip = ComputeMaxDeltasCountForMip(SourceMipIndex, InNumRelevantMips);
+			const int32 SourceMipMaxDeltasOffset = ComputeMaxDeltasOffsetForMip(SourceMipIndex, InNumRelevantMips);
+			check(SourceMipMaxDeltasOffset + NumMaxDeltasForSourceMip <= MipToMipMaxDeltas.Num());
+			TArrayView<double> SourceMipToDestinationMipMaxDeltas = MakeArrayView(MipToMipMaxDeltas.GetData() + SourceMipMaxDeltasOffset, NumMaxDeltasForSourceMip);
+			// Iterate through all remaining mips and find the max delta between the source mip and destination mips :
+			for (int32 SourceMipQuadIndex = 0; SourceMipQuadIndex < NumQuadsForSourceMip; ++SourceMipQuadIndex)
+			{
+				const FIntPoint SourceMipQuadCoords(SourceMipQuadIndex % SourceMipQuadsStride, SourceMipQuadIndex / SourceMipQuadsStride);
+				const FQuadHeightInfo& SourceMipQuadInfo = SourceMipQuadInfos[SourceMipQuadIndex];
+
+				// Special case for N to N+1 because all the info is located within the row of InAllMipsQuadInfos already so we don't need to fetch from another mip: FQuadHeightInfo's Average is the next mip's value:
+				int32 DestinationMipQuadsStride = SourceMipQuadsStride / 2;
+				FIntPoint DestinationMipQuadCoords = SourceMipQuadCoords / 2;
+				ComputeMaxDelta(SourceMipQuadInfo, SourceMipQuadInfo, SourceMipToDestinationMipMaxDeltas[0]);
+
+				for (int32 DestinationMipIndex = SourceMipIndex + 2; DestinationMipIndex < InNumRelevantMips; ++DestinationMipIndex)
+				{
+					DestinationMipQuadsStride /= 2;
+					DestinationMipQuadCoords /= 2;
+					const int32 DestinationMipQuadIndex = DestinationMipQuadCoords.X + DestinationMipQuadCoords.Y * DestinationMipQuadsStride;
+
+					TArrayView<const FQuadHeightInfo> DestinationMipQuadInfos = GetMipQuadInfosForMipConst(InAllMipsQuadInfos, DestinationMipIndex, TextureSize);
+					checkSlow(DestinationMipQuadInfos.Num() == DestinationMipQuadsStride * DestinationMipQuadsStride);
+
+					const FQuadHeightInfo& DestinationMipQuadInfo = DestinationMipQuadInfos[DestinationMipQuadIndex];
+					int32 DestinationMipRelativeIndex = DestinationMipIndex - SourceMipIndex - 1;
+					ComputeMaxDelta(SourceMipQuadInfo, DestinationMipQuadInfo, SourceMipToDestinationMipMaxDeltas[DestinationMipRelativeIndex]);
+				}
+			}
+		}
+
+		return MipToMipMaxDeltas;
+	}
+}
+
 void ULandscapeComponent::UpdateCachedBounds(bool bInApproximateBounds)
 {
+	TRACE_CPUPROFILER_EVENT_SCOPE(ULandscapeComponent::UpdateCachedBounds);
+
 	// Update local-space bounding box
 	CachedLocalBox.Init();
 	if (bInApproximateBounds && GetLandscapeProxy()->HasLayersContent())
@@ -158,16 +276,71 @@ void ULandscapeComponent::UpdateCachedBounds(bool bInApproximateBounds)
 	}
 	else
 	{
-		const int32 MipLevel = 0;
-		const bool bWorkOnEditingLayer = false; // We never want to compute bounds based on anything else that final landscape layer's height data
-		FLandscapeComponentDataInterface CDI(this, MipLevel, bWorkOnEditingLayer);
+		// We purposefully don't reset MipToMipMaxDeltas when bInApproximateBounds is true because we can totally live with a MipToMipMaxDeltas that is not up-to-date
+		//  and we want to minimize the render state changes so we track if there was any change to that data here : 
+		TArray<double> PreviousMipToMipMaxDeltas;
+		Swap(MipToMipMaxDeltas, PreviousMipToMipMaxDeltas); 
+
+		// TODO [jonathan.bard] : multithread and make this async : 
+		using namespace UE::Landscape::Private;
+		TArray<FQuadHeightInfo> AllMipsQuadInfos;
+		const int32 TextureSize = (SubsectionSizeQuads + 1) * NumSubsections;
+		const int32 NumTextureMips = FMath::FloorLog2(TextureSize) + 1;
+		// We actually only don't need to process the last texture mip, since a 1 vertex landscape is meaningless. When using 2x2 subsections, we can even drop an additional mip 
+		//  as the 4 texels of the penultimate mip will be identical (i.e. 4 sub-sections of 1 vertex are equally meaningless) :
+		const int32 NumRelevantMips = GetNumRelevantMips();
+		const int32 FinalMipIndex = NumRelevantMips - 1;
+		check(FinalMipIndex > 0);
+
+		{
+			TRACE_CPUPROFILER_EVENT_SCOPE(FetchMipQuads);
+
+			double LocalMin = DBL_MAX;
+			double LocalMax = DBL_MIN;
+			AllMipsQuadInfos.AddZeroed(ComputeQuadInfosCount(NumRelevantMips, TextureSize));
+			for (int32 MipIndex = 0; MipIndex < NumRelevantMips; ++MipIndex)
+			{
+				FLandscapeComponentDataInterface CDI(this, MipIndex, /*bWorkOnEditingLayer = */false);
+				TArrayView<FQuadHeightInfo> MipQuadInfos = GetMipQuadInfosForMip(AllMipsQuadInfos, MipIndex, TextureSize);
+				const int32 NumQuadsForMip = MipQuadInfos.Num();
+				const int32 MipTextureSize = TextureSize >> MipIndex;
+				const int32 MipTextureSubSectionSize = MipTextureSize / NumSubsections;
+				const int32 MipQuadsStride = MipTextureSize / 2;
+				for (int32 MipQuadIndex = 0; MipQuadIndex < NumQuadsForMip; ++MipQuadIndex)
+				{
+					int32 QuadY = MipQuadIndex / MipQuadsStride;
+					int32 QuadX = MipQuadIndex - QuadY * MipQuadsStride;
+					int32 X = QuadX * 2;
+					int32 Y = QuadY * 2;
+					// 2x2 Subsections have a duplicate pixel in the middle so we have to subtract one for the pixels in those subsections, since FLandscapeComponentDataInterface provides a view into 
+					//  pixels as if there was no subsection (hence the pixel indices for a 128 heightmap is in the range [0, 127] when using a single (1x1) subsection, but [0, 126] when using 2x2 subsections) :
+					X -= (X >= MipTextureSubSectionSize) ? 1 : 0;
+					Y -= (Y >= MipTextureSubSectionSize) ? 1 : 0;
+					double QuadVertices[4] = {
+						CDI.GetLocalHeight(X + 0, Y + 0),
+						CDI.GetLocalHeight(X + 1, Y + 0),
+						CDI.GetLocalHeight(X + 0, Y + 1),
+						CDI.GetLocalHeight(X + 1, Y + 1) };
+					FQuadHeightInfo QuadInfo(QuadVertices);
+					MipQuadInfos[MipQuadIndex] = QuadInfo;
+
+					if (MipIndex == 0)
+					{
+						LocalMin = FMath::Min(LocalMin, QuadInfo.Min);
+						LocalMax = FMath::Max(LocalMax, QuadInfo.Max);
+					}
+				}
+			}
+
+			CachedLocalBox = FBox(FVector(0.0, 0.0, LocalMin), FVector(ComponentSizeQuads, ComponentSizeQuads, LocalMax));
+		}
 
-		for (int32 y = 0; y < ComponentSizeQuads + 1; y++)
+		MipToMipMaxDeltas = ComputeMipToMipMaxDeltas(MakeArrayView<const FQuadHeightInfo>(AllMipsQuadInfos.GetData(), AllMipsQuadInfos.Num()), NumTextureMips, NumRelevantMips);
+
+		if (!Algo::Compare(MipToMipMaxDeltas, PreviousMipToMipMaxDeltas, [](double InLHS, double InRHS) { return FMath::IsNearlyEqual(InLHS, InRHS); }))
 		{
-			for (int32 x = 0; x < ComponentSizeQuads + 1; x++)
-			{
-				CachedLocalBox += CDI.GetLocalVertex(x, y);
-			}
+			// MipToMipMaxDeltas is used on the render thread, we need to reflect the change there : 
+			MarkRenderStateDirty();
 		}
 	}
 	if (CachedLocalBox.GetExtent().Z == 0)
@@ -5816,7 +5989,10 @@ void ALandscapeProxy::PostEditChangeProperty(FPropertyChangedEvent& PropertyChan
 		|| (MemberPropertyName == GET_MEMBER_NAME_CHECKED(ALandscapeProxy, bVirtualTextureRenderWithQuadHQ))
 		|| (MemberPropertyName == GET_MEMBER_NAME_CHECKED(ALandscapeProxy, VirtualTextureNumLods))
 		|| (MemberPropertyName == GET_MEMBER_NAME_CHECKED(ALandscapeProxy, VirtualTextureLodBias))
-		|| (MemberPropertyName == GET_MEMBER_NAME_CHECKED(ALandscapeProxy, bUseDynamicMaterialInstance)))
+		|| (MemberPropertyName == GET_MEMBER_NAME_CHECKED(ALandscapeProxy, bUseDynamicMaterialInstance))
+		|| (MemberPropertyName == GET_MEMBER_NAME_CHECKED(ALandscapeProxy, NonNaniteVirtualShadowMapConstantDepthBias))
+		|| (MemberPropertyName == GET_MEMBER_NAME_CHECKED(ALandscapeProxy, NonNaniteVirtualShadowMapInvalidationHeightErrorThreshold))
+		|| (MemberPropertyName == GET_MEMBER_NAME_CHECKED(ALandscapeProxy, NonNaniteVirtualShadowMapInvalidationScreenSizeLimit)))
 	{		
 		MarkComponentsRenderStateDirty();
 	}
diff --git a/Engine/Source/Runtime/Landscape/Private/LandscapeRender.cpp b/Engine/Source/Runtime/Landscape/Private/LandscapeRender.cpp
--- a/Engine/Source/Runtime/Landscape/Private/LandscapeRender.cpp
+++ b/Engine/Source/Runtime/Landscape/Private/LandscapeRender.cpp
@@ -54,6 +54,8 @@ LandscapeRender.cpp: New terrain rendering
 #include "RenderGraphBuilder.h"
 #include "Scalability.h"
 #include "Rendering/CustomRenderPass.h"
+#include "LandscapeUtils.h"
+#include "SceneRendererInterface.h"
 
 using namespace UE::Landscape;
 
@@ -62,15 +64,13 @@ IMPLEMENT_GLOBAL_SHADER_PARAMETER_STRUCT(FLandscapeFixedGridUniformShaderParamet
 IMPLEMENT_TYPE_LAYOUT(FLandscapeVertexFactoryVertexShaderParameters);
 IMPLEMENT_TYPE_LAYOUT(FLandscapeVertexFactoryPixelShaderParameters);
 
-#if !UE_BUILD_SHIPPING
-static void OnLODDistributionScaleChanged(IConsoleVariable* CVar)
+static void OnCVarNeedingRenderStateInvalidationChanged(IConsoleVariable* CVar)
 {
-	for (auto* LandscapeComponent : TObjectRange<ULandscapeComponent>(RF_ClassDefaultObject | RF_ArchetypeObject, true, EInternalObjectFlags::Garbage))
+	for (ULandscapeComponent* LandscapeComponent : TObjectRange<ULandscapeComponent>(RF_ClassDefaultObject | RF_ArchetypeObject, true, EInternalObjectFlags::Garbage))
 	{
 		LandscapeComponent->MarkRenderStateDirty();
 	}
 }
-#endif
 
 #if !UE_BUILD_SHIPPING
 float GLandscapeLOD0ScreenSizeOverride = -1.f;
@@ -78,7 +78,7 @@ FAutoConsoleVariableRef CVarLandscapeLOD0ScreenSizeOverride(
 	TEXT("landscape.OverrideLOD0ScreenSize"),
 	GLandscapeLOD0ScreenSizeOverride,
 	TEXT("When > 0, force override the landscape LOD0ScreenSize property on all landscapes"),
-	FConsoleVariableDelegate::CreateStatic(&OnLODDistributionScaleChanged),
+	FConsoleVariableDelegate::CreateStatic(&OnCVarNeedingRenderStateInvalidationChanged),
 	ECVF_Cheat
 );
 
@@ -87,7 +87,7 @@ FAutoConsoleVariableRef CVarLandscapeLOD0DistributionOverride(
 	TEXT("landscape.OverrideLOD0Distribution"),
 	GLandscapeLOD0DistributionOverride,
 	TEXT("When > 0, force override the LOD0DistributionSetting property on all landscapes, and ignore r.LandscapeLOD0DistributionScale"),
-	FConsoleVariableDelegate::CreateStatic(&OnLODDistributionScaleChanged),
+	FConsoleVariableDelegate::CreateStatic(&OnCVarNeedingRenderStateInvalidationChanged),
 	ECVF_Cheat
 );
 
@@ -96,7 +96,7 @@ FAutoConsoleVariableRef CVarLandscapeLODDistributionOverride(
 	TEXT("landscape.OverrideLODDistribution"),
 	GLandscapeLODDistributionOverride,
 	TEXT("When > 0, force override the landscape LODDistributionSetting property on all landscapes, and ignore r.LandscapeLODDistributionScale"),
-	FConsoleVariableDelegate::CreateStatic(&OnLODDistributionScaleChanged),
+	FConsoleVariableDelegate::CreateStatic(&OnCVarNeedingRenderStateInvalidationChanged),
 	ECVF_Cheat
 );
 
@@ -105,7 +105,34 @@ FAutoConsoleVariableRef CVarLandscapeLODBlendRangeOverride(
 	TEXT("landscape.OverrideLODBlendRange"),
 	GLandscapeLODBlendRangeOverride,
 	TEXT("When > 0, force the LODBlendRange property on all landscapes"),
-	FConsoleVariableDelegate::CreateStatic(&OnLODDistributionScaleChanged),
+	FConsoleVariableDelegate::CreateStatic(&OnCVarNeedingRenderStateInvalidationChanged),
+	ECVF_Cheat
+);
+
+float GLandscapeNonNaniteVirtualShadowMapConstantDepthBiasOverride = -1.f;
+FAutoConsoleVariableRef CVarLandscapeNonNaniteVirtualShadowMapConstantDepthBiasOverrideOverride(
+	TEXT("landscape.OverrideNonNaniteVirtualShadowMapConstantDepthBiasOverride"),
+	GLandscapeNonNaniteVirtualShadowMapConstantDepthBiasOverride,
+	TEXT("When > 0, force override the landscape NonNaniteVirtualShadowMapConstantDepthBias property on all landscapes"),
+	FConsoleVariableDelegate::CreateStatic(&OnCVarNeedingRenderStateInvalidationChanged),
+	ECVF_Cheat
+);
+
+float GLandscapeNonNaniteVirtualShadowMapInvalidationHeightErrorThresholdOverride = -1.f;
+FAutoConsoleVariableRef CVarLandscapeNonNaniteVirtualShadowMapInvalidationHeightErrorThresholdOverride(
+	TEXT("landscape.OverrideNonNaniteVirtualShadowMapInvalidationHeightErrorThreshold"),
+	GLandscapeNonNaniteVirtualShadowMapInvalidationHeightErrorThresholdOverride,
+	TEXT("When > 0, force override the landscape NonNaniteVirtualShadowMapInvalidationHeightErrorThreshold property on all landscapes"),
+	FConsoleVariableDelegate::CreateStatic(&OnCVarNeedingRenderStateInvalidationChanged),
+	ECVF_Cheat
+);
+
+float GLandscapeNonNaniteVirtualShadowMapInvalidationScreenSizeLimitOverride = -1.f;
+FAutoConsoleVariableRef CVarLandscapeNonNaniteVirtualShadowMapInvalidationScreenSizeLimitOverride(
+	TEXT("landscape.OverrideNonNaniteVirtualShadowMapInvalidationScreenSizeLimit"),
+	GLandscapeNonNaniteVirtualShadowMapInvalidationScreenSizeLimitOverride,
+	TEXT("When > 0, force override the landscape NonNaniteVirtualShadowMapInvalidationScreenSizeLimit property on all landscapes"),
+	FConsoleVariableDelegate::CreateStatic(&OnCVarNeedingRenderStateInvalidationChanged),
 	ECVF_Cheat
 );
 #else
@@ -117,9 +144,7 @@ FAutoConsoleVariableRef CVarLandscapeLOD0DistributionScale(
 	TEXT("r.LandscapeLOD0DistributionScale"),
 	GLandscapeLOD0DistributionScale,
 	TEXT("Multiplier for the landscape LOD0DistributionSetting property"),
-#if !UE_BUILD_SHIPPING
-	FConsoleVariableDelegate::CreateStatic(&OnLODDistributionScaleChanged),
-#endif
+	FConsoleVariableDelegate::CreateStatic(&OnCVarNeedingRenderStateInvalidationChanged),
 	ECVF_Scalability
 );
 
@@ -128,9 +153,7 @@ FAutoConsoleVariableRef CVarLandscapeLODDistributionScale(
 	TEXT("r.LandscapeLODDistributionScale"),
 	GLandscapeLODDistributionScale,
 	TEXT("Multiplier for the landscape LODDistributionSetting property"),
-#if !UE_BUILD_SHIPPING
-	FConsoleVariableDelegate::CreateStatic(&OnLODDistributionScaleChanged),
-#endif
+	FConsoleVariableDelegate::CreateStatic(&OnCVarNeedingRenderStateInvalidationChanged),
 	ECVF_Scalability
 );
 
@@ -138,20 +161,41 @@ int32 GAllowLandscapeShadows = 1;
 static FAutoConsoleVariableRef CVarAllowLandscapeShadows(
 	TEXT("r.AllowLandscapeShadows"),
 	GAllowLandscapeShadows,
-	TEXT("Allow Landscape Shadows")
+	TEXT("Allow Landscape Shadows"),
+	ECVF_RenderThreadSafe | ECVF_Scalability
 );
 
 int32 GLandscapeUseAsyncTasksForLODComputation = 1;
 FAutoConsoleVariableRef CVarLandscapeUseAsyncTasksForLODComputation(
 	TEXT("r.LandscapeUseAsyncTasksForLODComputation"),
 	GLandscapeUseAsyncTasksForLODComputation,
-	TEXT("Use async tasks for computing per-landscape component LOD biases."));
+	TEXT("Use async tasks for computing per-landscape component LOD biases."), 
+	ECVF_RenderThreadSafe
+);
 
 int32 GDisableLandscapeNaniteGI = 1;
 static FAutoConsoleVariableRef CVarDisableLandscapeNaniteGI(
 	TEXT("r.DisableLandscapeNaniteGI"),
 	GDisableLandscapeNaniteGI,
-	TEXT("Disable Landscape Nanite GI")
+	TEXT("Disable Landscape Nanite GI"),
+	FConsoleVariableDelegate::CreateStatic(&OnCVarNeedingRenderStateInvalidationChanged), 
+	ECVF_Scalability
+);
+
+bool GLandscapeAllowNonNaniteVirtualShadowMapInvalidation = true;
+FAutoConsoleVariableRef CVarLandscapeAllowNonNaniteVirtualShadowMapInvalidation(
+	TEXT("landscape.AllowNonNaniteVirtualShadowMapInvalidation"),
+	GLandscapeAllowNonNaniteVirtualShadowMapInvalidation,
+	TEXT("For non-Nanite landscape, cached virtual shadow map pages need to be invalidated when the vertex morphing introduces a height difference that is too large. This enables or disables this behavior entirely"),
+	ECVF_RenderThreadSafe | ECVF_Scalability
+);
+
+float GLandscapeNonNaniteVirtualShadowMapInvalidationLODAttenuationExponent = 2.0f;
+FAutoConsoleVariableRef CVarLandscapeNonNaniteVirtualShadowMapInvalidationLODAttenuationExponent(
+	TEXT("landscape.NonNaniteVirtualShadowMapInvalidationLODAttenuationExponent"),
+	GLandscapeNonNaniteVirtualShadowMapInvalidationLODAttenuationExponent,
+	TEXT("For non-Nanite landscape, controls the shape of the curve of the attenuation of the virtual shadow map pages' invalidation rate (1 - X^N), where X is the relative LOD value (LODValue/NumMips in the [0,1] range) and N, the CVar"),
+	ECVF_RenderThreadSafe | ECVF_Scalability
 );
 
 #if WITH_EDITOR
@@ -164,19 +208,19 @@ extern int32 GGrassMapUseRuntimeGeneration;
 extern int32 GGrassMapAlwaysBuildRuntimeGenerationResources;
 
 #if !UE_BUILD_SHIPPING
-int32 GVarDumpLandscapeLODsCurrentFrame = 0;
+uint32 GVarDumpLandscapeLODsCurrentFrame = 0;
 
 static void OnDumpLandscapeLODs(const TArray< FString >& Args)
 {
-	if (Args.Num() >= 1)
+	if (GVarDumpLandscapeLODsCurrentFrame == 0)
 	{
 		// Add some buffer to be able to correctly catch the frame during the rendering
-		GVarDumpLandscapeLODsCurrentFrame = FCString::Atoi(*Args[0]) != 0 ? GFrameNumberRenderThread + 3 : INDEX_NONE;
+		GVarDumpLandscapeLODsCurrentFrame = GFrameNumberRenderThread + 3;
 	}
 }
 
 static FAutoConsoleCommand CVarDumpLandscapeLODs(
-	TEXT("Landscape.DumpLODs"),
+	TEXT("landscape.DumpLODs"),
 	TEXT("Will dump the current status of LOD value and current texture streaming status"),
 	FConsoleCommandWithArgsDelegate::CreateStatic(&OnDumpLandscapeLODs)
 );
@@ -438,8 +482,26 @@ const static FName NAME_LandscapeResourceNameForDebugging(TEXT("Landscape"));
 
 IMPLEMENT_GLOBAL_SHADER_PARAMETER_STRUCT(FLandscapeSectionLODUniformParameters, "LandscapeContinuousLODParameters");
 
+// TODO [jonathan.bard] Move this to a ISceneExtension, where it belongs : the FLandscapeRenderSystem are associated with a scene and should have the same lifetime
 TMap<uint32, FLandscapeRenderSystem*> LandscapeRenderSystems;
 
+float FLandscapeRenderSystem::ComputeLODFromScreenSize(const LODSettingsComponent& InLODSettings, float InScreenSizeSquared)
+{
+	if (InScreenSizeSquared <= InLODSettings.LastLODScreenSizeSquared)
+	{
+		return InLODSettings.LastLODIndex;
+	}
+	else if (InScreenSizeSquared > InLODSettings.LOD1ScreenSizeSquared)
+	{
+		return (InLODSettings.LOD0ScreenSizeSquared - FMath::Min(InScreenSizeSquared, InLODSettings.LOD0ScreenSizeSquared)) / (InLODSettings.LOD0ScreenSizeSquared - InLODSettings.LOD1ScreenSizeSquared);
+	}
+	else
+	{
+		// No longer linear fraction, but worth the cache misses
+		return 1.0f + FMath::LogX(InLODSettings.LODOnePlusDistributionScalarSquared, InLODSettings.LOD1ScreenSizeSquared / InScreenSizeSquared);
+	}
+}
+
 TBitArray<> FLandscapeRenderSystem::LandscapeIndexAllocator;
 
 #if RHI_RAYTRACING
@@ -535,13 +597,15 @@ public:
 //
 // FLandscapeRenderSystem
 //
-FLandscapeRenderSystem::FLandscapeRenderSystem()
+FLandscapeRenderSystem::FLandscapeRenderSystem(uint32 InLandscapeKey, FSceneInterface* InScene)
 	: Min(MAX_int32, MAX_int32)
 	, Size(EForceInit::ForceInitToZero)
 	, ReferenceCount(0)
 	, RegisteredCount(0)
 	, ForcedLODOverride(-1)
 	, SectionsRemovedSinceLastCompact(0)
+	, LandscapeKey(InLandscapeKey)
+	, Scene(InScene)
 {
 	SectionLODBiases.SetAllowCPUAccess(true);
 
@@ -573,7 +637,7 @@ void FLandscapeRenderSystem::CreateResources(FRHICommandListBase& RHICmdList, FL
 	FLandscapeRenderSystem*& LandscapeRenderSystem = LandscapeRenderSystems.FindOrAdd(SectionInfo->LandscapeKey);
 	if (!LandscapeRenderSystem)
 	{
-		LandscapeRenderSystem = new FLandscapeRenderSystem();
+		LandscapeRenderSystem = new FLandscapeRenderSystem(SectionInfo->LandscapeKey, SectionInfo->Scene);
 	}
 
 	LandscapeRenderSystem->CreateResources_Internal(RHICmdList, SectionInfo);
@@ -888,22 +952,87 @@ bool FLandscapeRenderSystem::AnySectionsInRangeInclusive(FIntPoint RangeMin, FIn
 	return false;
 }
 
-const TResourceArray<float>& FLandscapeRenderSystem::ComputeSectionsLODForView(const FSceneView& InView)
+float FLandscapeRenderSystem::GetSectionLODValue(const FSceneView& InView, FIntPoint InRenderCoord) const
+{
+	return GetCachedSectionLODValues(InView)[GetSectionLinearIndex(InRenderCoord)];
+}
+
+const TResourceArray<float>& FLandscapeRenderSystem::GetCachedSectionLODValues(const FSceneView& InView) const
+{
+	const uint32 ViewStateKey = InView.GetViewKey();
+	const TResourceArray<float>* CachedSectionLODValues = (ViewStateKey != 0) ? PerViewStateCachedSectionLODValues.Find(ViewStateKey) : PerViewCachedSectionLODValues.Find(&InView);
+	checkf(CachedSectionLODValues != nullptr, TEXT("No section LOD value cached for this view. Make sure FLandscapeRenderSystem::ComputeSectionsLODForView (FLandscapeSceneViewExtension::PreRenderView_RenderThread) was called"));
+	return *CachedSectionLODValues;
+}
+
+float FLandscapeRenderSystem::GetSectionLODBias(FIntPoint InRenderCoord) const
+{
+	return SectionLODBiases[GetSectionLinearIndex(InRenderCoord)];
+}
+
+const TResourceArray<float>& FLandscapeRenderSystem::ComputeSectionsLODForView(const FSceneView& InView, UE::Renderer::Private::IShadowInvalidatingInstances* InShadowInvalidatingInstances)
 {
 	TRACE_CPUPROFILER_EVENT_SCOPE(FLandscapeRenderSystem::ComputeSectionsLODForView);
 
-	TResourceArray<float>& SectionLODValues = CachedSectionLODValues.Add(&InView);
-	SectionLODValues.AddZeroed(SectionInfos.Num());
+	// Find where the CachedSectionLODValues lie (either in the persistent map, if we have a persistent view state or in the transient one) : 
+	TResourceArray<float>* CachedSectionLODValues = nullptr;
+	SectionKeyToLODValueMap* LastShadowInvalidationSectionLODValues = nullptr;
+	TArray<const FPrimitiveSceneInfo*>* PrimitivesToInvalidateForShadows = nullptr;
+	const uint32 ViewStateKey = InView.GetViewKey();
+	if (ViewStateKey != 0)
+	{
+		CachedSectionLODValues = &PerViewStateCachedSectionLODValues.FindOrAdd(ViewStateKey);
+		if (InShadowInvalidatingInstances != nullptr)
+		{
+			LastShadowInvalidationSectionLODValues = &PerViewStateLastShadowInvalidationSectionLODValues.FindOrAdd(ViewStateKey);
+			PrimitivesToInvalidateForShadows = &ShadowInvalidationRequests.FindOrAdd(InShadowInvalidatingInstances);
+		}
+	}
+	else
+	{
+		CachedSectionLODValues = &PerViewCachedSectionLODValues.FindOrAdd(&InView);
+	}
 
-	for (int32 SectionIndex = 0; SectionIndex < SectionInfos.Num(); SectionIndex++)
+	CachedSectionLODValues->Reset(SectionInfos.Num());
+	for (FLandscapeSectionInfo* SectionInfo : SectionInfos)
 	{
-		const float DefaultLODValue = 0.0f;
+		constexpr float DefaultLODValue = 0.0f;
+		float& LODSectionValue = CachedSectionLODValues->Add_GetRef(DefaultLODValue);
+		if (SectionInfo != nullptr)
+		{
+			LODSectionValue = SectionInfo->ComputeLODForView(InView);
 
-		FLandscapeSectionInfo* SectionInfo = SectionInfos[SectionIndex];
-		SectionLODValues[SectionIndex] = SectionInfo ? SectionInfo->ComputeLODForView(InView) : DefaultLODValue;
+			if ((InShadowInvalidatingInstances != nullptr) && (PrimitivesToInvalidateForShadows != nullptr))
+			{
+				// The shadow invalidation system requires the actual LOD that will be used in the shader, so we need to "correct" the desired LODSectionValue with the value that we *can* 
+				//  actually render the landscape with (i.e. LODSectionBias) : 
+				const float LODSectionBias = GetSectionLODBias(SectionInfo->RenderCoord);
+				const float FinalLODSectionValue = FMath::Max(LODSectionBias, LODSectionValue);
+
+				float& LastShadowInvalidationLODValue = LastShadowInvalidationSectionLODValues->FindOrAdd(SectionInfo->RenderCoord, FinalLODSectionValue);
+				if (SectionInfo->ShouldInvalidateShadows(InView, FinalLODSectionValue, LastShadowInvalidationLODValue))
+				{
+					LastShadowInvalidationLODValue = FinalLODSectionValue;
+					PrimitivesToInvalidateForShadows->Add(SectionInfo->GetPrimitiveSceneInfo());
+				}
+			}
+		}
 	}
 
-	return SectionLODValues;
+	return *CachedSectionLODValues;
+}
+
+void FLandscapeRenderSystem::PerformShadowInvalidations(UE::Renderer::Private::IShadowInvalidatingInstances& InShadowInvalidatingInstances)
+{
+	checkf(IsInRenderingThread(), TEXT("Using IShadowInvalidatingInstances is only allowed from the rendering thread!"));
+	if (TArray<const FPrimitiveSceneInfo*>* PrimitivesToInvalidateForShadows = ShadowInvalidationRequests.Find(&InShadowInvalidatingInstances))
+	{
+		for (const FPrimitiveSceneInfo* Primitive : *PrimitivesToInvalidateForShadows)
+		{
+			InShadowInvalidatingInstances.AddPrimitive(Primitive);
+		}
+		ShadowInvalidationRequests.Remove(&InShadowInvalidatingInstances);
+	}
 }
 
 void FLandscapeRenderSystem::FetchHeightmapLODBiases()
@@ -957,12 +1086,24 @@ void FLandscapeRenderSystem::UpdateBuffers(FRHICommandListBase& RHICmdList)
 }
 
 
+//
+// FLandscapeSceneViewExtension::FLandscapeViewData
+//
+FLandscapeSceneViewExtension::FLandscapeViewData::FLandscapeViewData(FSceneView& InView)
+	: View(&InView)
+{
+	if (ISceneRenderer* ScenerRenderer = InView.Family->GetSceneRenderer())
+	{
+		ShadowInvalidatingInstances = ScenerRenderer->GetShadowInvalidatingInstancesInterface(&InView);
+	}
+}
+
+
 //
 // FLandscapeSceneViewExtension
 //
 FLandscapeSceneViewExtension::FLandscapeSceneViewExtension(const FAutoRegister& AutoReg) : FSceneViewExtensionBase(AutoReg)
 {
-
 	FCoreDelegates::OnEndFrame.AddRaw(this, &FLandscapeSceneViewExtension::EndFrame_GameThread);
 	FCoreDelegates::OnEndFrameRT.AddRaw(this, &FLandscapeSceneViewExtension::EndFrame_RenderThread);
 }
@@ -1002,60 +1143,53 @@ void FLandscapeSceneViewExtension::PreRenderView_RenderThread(FRDGBuilder& Graph
 	}
 #endif	// RHI_RAYTRACING
 
+	TArray<FLandscapeRenderSystem*> SceneLandscapeRenderSystems = GetLandscapeRenderSystems(InView.Family->Scene);
+
 	// Kick the job once all views have been collected.
-	if (!LandscapeRenderSystems.IsEmpty() && LandscapeViews.Num() == InView.Family->AllViews.Num())
+	if (!SceneLandscapeRenderSystems.IsEmpty() && LandscapeViews.Num() == InView.Family->AllViews.Num())
 	{
-		LandscapeSetupTask = GraphBuilder.AddCommandListSetupTask([this] (FRHICommandListBase& RHICmdList)
+		LandscapeSetupTask = GraphBuilder.AddCommandListSetupTask([SceneRenderSystems = MoveTemp(SceneLandscapeRenderSystems), LocalLandscapeViewsPtr = &LandscapeViews](FRHICommandListBase& RHICmdList)
 		{
 			TRACE_CPUPROFILER_EVENT_SCOPE(FLandscapeRenderSystem::ComputeLODs);
 			FOptionalTaskTagScope Scope(ETaskTag::EParallelRenderingThread);
+			check(!SceneRenderSystems.IsEmpty())
 
-			for (auto& Pair : LandscapeRenderSystems)
+			for (FLandscapeRenderSystem* RenderSystem : SceneRenderSystems)
 			{
-				FLandscapeRenderSystem& RenderSystem = *Pair.Value;
-				RenderSystem.CachedSectionLODValues.Reset();
-				RenderSystem.FetchHeightmapLODBiases();
+				RenderSystem->FetchHeightmapLODBiases();
 			}
 
-			for (FLandscapeViewData& LandscapeView : LandscapeViews)
+			for (FLandscapeViewData& LandscapeView : *LocalLandscapeViewsPtr)
 			{
+				const uint32 ViewStateKey = LandscapeView.View->GetViewKey() ;
+
 				LandscapeView.LandscapeIndirection.SetNum(FLandscapeRenderSystem::LandscapeIndexAllocator.Num());
 
-				for (auto& Pair : LandscapeRenderSystems)
+				for (FLandscapeRenderSystem* RenderSystem : SceneRenderSystems)
 				{
-					FLandscapeRenderSystem& RenderSystem = *Pair.Value;
-
 					// Store index where the LOD data for this landscape starts
-					LandscapeView.LandscapeIndirection[RenderSystem.LandscapeIndex] = LandscapeView.LandscapeLODData.Num();
+					LandscapeView.LandscapeIndirection[RenderSystem->LandscapeIndex] = LandscapeView.LandscapeLODData.Num();
 
 					// Compute sections lod values for this view & append to the global landscape LOD data
-					LandscapeView.LandscapeLODData.Append(RenderSystem.ComputeSectionsLODForView(*LandscapeView.View));
+					const TResourceArray<float>& CachedSectionLODValues = RenderSystem->ComputeSectionsLODForView(*LandscapeView.View, LandscapeView.ShadowInvalidatingInstances);
+					LandscapeView.LandscapeLODData.Append(CachedSectionLODValues);
 				}
 			}
 
-			for (auto& Pair : LandscapeRenderSystems)
+			for (FLandscapeRenderSystem* RenderSystem : SceneRenderSystems)
 			{
-				FLandscapeRenderSystem& RenderSystem = *Pair.Value;
-				RenderSystem.UpdateBuffers(RHICmdList);
+				RenderSystem->UpdateBuffers(RHICmdList);
 			}
 
-			for (FLandscapeViewData& LandscapeView : LandscapeViews)
+			for (FLandscapeViewData& LandscapeView : *LocalLandscapeViewsPtr)
 			{
-				if (!LandscapeRenderSystems.IsEmpty())
-				{
-					FRHIResourceCreateInfo CreateInfoLODBuffer(TEXT("LandscapeLODDataBuffer"), &LandscapeView.LandscapeLODData);
-					LandscapeLODDataBuffer = RHICmdList.CreateVertexBuffer(LandscapeView.LandscapeLODData.GetResourceDataSize(), BUF_ShaderResource | BUF_Volatile, CreateInfoLODBuffer);
-					LandscapeView.View->LandscapePerComponentDataBuffer = RHICmdList.CreateShaderResourceView(LandscapeLODDataBuffer, sizeof(float), PF_R32_FLOAT);
+				FRHIResourceCreateInfo CreateInfoLODBuffer(TEXT("LandscapeLODDataBuffer"), &LandscapeView.LandscapeLODData);
+				FBufferRHIRef LandscapeLODDataBuffer = RHICmdList.CreateVertexBuffer(LandscapeView.LandscapeLODData.GetResourceDataSize(), BUF_ShaderResource | BUF_Volatile, CreateInfoLODBuffer);
+				LandscapeView.View->LandscapePerComponentDataBuffer = RHICmdList.CreateShaderResourceView(LandscapeLODDataBuffer, sizeof(float), PF_R32_FLOAT);
 
-					FRHIResourceCreateInfo CreateInfoIndirection(TEXT("LandscapeIndirectionBuffer"), &LandscapeView.LandscapeIndirection);
-					LandscapeIndirectionBuffer = RHICmdList.CreateVertexBuffer(LandscapeView.LandscapeIndirection.GetResourceDataSize(), BUF_ShaderResource | BUF_Volatile, CreateInfoIndirection);
-					LandscapeView.View->LandscapeIndirectionBuffer = RHICmdList.CreateShaderResourceView(LandscapeIndirectionBuffer, sizeof(uint32), PF_R32_UINT);
-				}
-				else
-				{
-					LandscapeView.View->LandscapePerComponentDataBuffer = GWhiteVertexBufferWithSRV->ShaderResourceViewRHI;
-					LandscapeView.View->LandscapeIndirectionBuffer = GWhiteVertexBufferWithSRV->ShaderResourceViewRHI;
-				}
+				FRHIResourceCreateInfo CreateInfoIndirection(TEXT("LandscapeIndirectionBuffer"), &LandscapeView.LandscapeIndirection);
+				FBufferRHIRef LandscapeIndirectionBuffer = RHICmdList.CreateVertexBuffer(LandscapeView.LandscapeIndirection.GetResourceDataSize(), BUF_ShaderResource | BUF_Volatile, CreateInfoIndirection);
+				LandscapeView.View->LandscapeIndirectionBuffer = RHICmdList.CreateShaderResourceView(LandscapeIndirectionBuffer, sizeof(uint32), PF_R32_UINT);
 			}
 
 		}, GIsThreadedRendering && GLandscapeUseAsyncTasksForLODComputation);
@@ -1064,13 +1198,23 @@ void FLandscapeSceneViewExtension::PreRenderView_RenderThread(FRDGBuilder& Graph
 
 void FLandscapeSceneViewExtension::PreInitViews_RenderThread(FRDGBuilder& GraphBuilder)
 {
+	TRACE_CPUPROFILER_EVENT_SCOPE(FLandscapeSceneViewExtension::PreInitViews_RenderThread);
 	LandscapeSetupTask.Wait();
 
 	TArray<const FSceneView*, TInlineAllocator<2>> SceneViews;
-
 	for (FLandscapeViewData& LandscapeView : LandscapeViews)
 	{
 		SceneViews.Add(LandscapeView.View);
+
+		// Perform the accumulated shadow invalidations for this view: 
+		if (LandscapeView.ShadowInvalidatingInstances != nullptr)
+		{
+			for (auto& Pair : LandscapeRenderSystems)
+			{
+				FLandscapeRenderSystem* RenderSystem = Pair.Value;
+				RenderSystem->PerformShadowInvalidations(*LandscapeView.ShadowInvalidatingInstances);
+			}
+		}
 	}
 
 	if (Culling::UseCulling(SceneViews[0]->GetShaderPlatform()))
@@ -1087,12 +1231,19 @@ void FLandscapeSceneViewExtension::EndFrame_GameThread()
 	NumViewsWithShowCollisionAcc = 0;
 }
 
-// TODO [jonathan.bard] Ideally this should be symmetrical with FLandscapeSceneViewExtension::PreRenderView_RenderThread and should be called in FLandscapeSceneViewExtension::PostRenderView_RenderThread
-//  but PostRenderView_RenderThread is not currently called on all code paths (scene/reflection captures, etc.)
 void FLandscapeSceneViewExtension::EndFrame_RenderThread()
 {
-	LandscapeLODDataBuffer = FBufferRHIRef();
-	LandscapeIndirectionBuffer = FBufferRHIRef();
+	for (auto& Pair : LandscapeRenderSystems)
+	{
+		FLandscapeRenderSystem* RenderSystem = Pair.Value;
+		// Cleanup the transient list of cached section LOD values : they are only valid for the frame since they are associated with views that have no persistent view state :
+		RenderSystem->PerViewCachedSectionLODValues.Reset();
+		// COMMENT [jonathan.bard] : we cannot do the same for the LOD values associated with persistent view states since it's possible to run certain frames without rendering any view 
+		//  at all, which would lead us to needlessly lose the history (PreviousPerViewStateCachedSectionLODValues). We need a proper point to attach per-view state data like this, 
+		//  with the same lifetime as the view state itself, but we don't have anything for this ATM. This is good enough for now, as we don't have too many view states, so this data shouldn't
+		//  grow too much unless creating/destroying several view states, which should usually never really occur
+		checkf(RenderSystem->ShadowInvalidationRequests.IsEmpty(), TEXT("All shadow invalidation requests should have been processed by now"));
+	}
 }
 
 const TMap<uint32, FLandscapeRenderSystem*>& FLandscapeSceneViewExtension::GetLandscapeRenderSystems() const
@@ -1101,6 +1252,23 @@ const TMap<uint32, FLandscapeRenderSystem*>& FLandscapeSceneViewExtension::GetLa
 	return LandscapeRenderSystems;
 }
 
+TArray<FLandscapeRenderSystem*> FLandscapeSceneViewExtension::GetLandscapeRenderSystems(const FSceneInterface* InScene) const
+{
+	checkf(IsInParallelRenderingThread(), TEXT("Accessing the Landscape render systems is only valid from the rendering thread!"));
+	check(InScene != nullptr);
+
+	TArray<FLandscapeRenderSystem*> Result;
+	Result.Reserve(LandscapeRenderSystems.Num());
+	for (auto& Pair : LandscapeRenderSystems)
+	{
+		FLandscapeRenderSystem* RenderSystem = Pair.Value;
+		if (RenderSystem->Scene == InScene)
+		{
+			Result.Add(RenderSystem);
+		}
+	}
+	return Result;
+}
 
 //
 // FLandscapeVisibilityHelper
@@ -1153,7 +1321,6 @@ FLandscapeComponentSceneProxy::FLandscapeComponentSceneProxy(ULandscapeComponent
 	, ComponentSizeQuads(InComponent->ComponentSizeQuads)
 	, ComponentSizeVerts(InComponent->ComponentSizeQuads + 1)
 	, SectionBase(InComponent->GetSectionBase())
-	, LandscapeComponent(InComponent)
 	, bUsesLandscapeCulling(false)
 	, WeightmapScaleBias(InComponent->WeightmapScaleBias)
 	, VisibilityWeightmapTexture(nullptr)
@@ -1165,9 +1332,9 @@ FLandscapeComponentSceneProxy::FLandscapeComponentSceneProxy(ULandscapeComponent
 	, SharedBuffers(nullptr)
 	, VertexFactory(nullptr)
 	, ComponentLightInfo(nullptr)
+	, NumRelevantMips(InComponent->GetNumRelevantMips())
 #if WITH_EDITORONLY_DATA
 	, EditToolRenderData(InComponent->EditToolRenderData)
-	, LODFalloff_DEPRECATED(InComponent->GetLandscapeProxy()->LODFalloff_DEPRECATED)
 #endif
 #if WITH_EDITOR || !(UE_BUILD_SHIPPING || UE_BUILD_TEST)
 	, CollisionMipLevel(InComponent->CollisionMipLevel)
@@ -1246,6 +1413,9 @@ FLandscapeComponentSceneProxy::FLandscapeComponentSceneProxy(ULandscapeComponent
 	float LOD0ScreenSize;
 	float LOD0Distribution;
 	float LODDistribution;
+	VirtualShadowMapConstantDepthBias = Proxy->NonNaniteVirtualShadowMapConstantDepthBias;
+	VirtualShadowMapInvalidationHeightErrorThreshold = Proxy->NonNaniteVirtualShadowMapInvalidationHeightErrorThreshold;
+	float NonNaniteVirtualShadowMapInvalidationScreenSizeLimit = Proxy->NonNaniteVirtualShadowMapInvalidationScreenSizeLimit;
 	
 	if (Proxy->bUseScalableLODSettings)
 	{
@@ -1275,27 +1445,47 @@ FLandscapeComponentSceneProxy::FLandscapeComponentSceneProxy(ULandscapeComponent
 	{
 		LODDistribution = GLandscapeLODDistributionOverride;
 	}
-#endif // !UE_BUILD_SHIPPING
+	if (GLandscapeNonNaniteVirtualShadowMapConstantDepthBiasOverride > 0.0)
+	{
+		VirtualShadowMapConstantDepthBias = GLandscapeNonNaniteVirtualShadowMapConstantDepthBiasOverride;
+	}
+	if (GLandscapeNonNaniteVirtualShadowMapInvalidationHeightErrorThresholdOverride > 0.0)
+	{
+		VirtualShadowMapInvalidationHeightErrorThreshold = GLandscapeNonNaniteVirtualShadowMapInvalidationHeightErrorThresholdOverride;
+	}
+	if (GLandscapeNonNaniteVirtualShadowMapInvalidationScreenSizeLimitOverride > 0.0)
+	{
+		NonNaniteVirtualShadowMapInvalidationScreenSizeLimit = GLandscapeNonNaniteVirtualShadowMapInvalidationScreenSizeLimitOverride;
+	}
 
-	float ScreenSizeRatioDivider = FMath::Max(LOD0Distribution, 1.01f);
-	// Cancel out so that landscape is not affected by r.StaticMeshLODDistanceScale
-	float CurrentScreenSizeRatio = LOD0ScreenSize / CVarStaticMeshLODDistanceScale.GetValueOnAnyThread();
+	// For the display name, some timing issues can lead to a temporarily invalid parent actor, so fallback to the parent actor's name if that's the case, no big deal : 
+	const FString& LandscapeName = InComponent->GetLandscapeActor() ? InComponent->GetLandscapeActor()->GetName() : Proxy->GetName();
+	const FString& ComponentName = InComponent->GetName();
+	DebugName = FName(FString::Printf(TEXT(" Landscape: %s, Component: %s [%s]"), *ComponentName, *LandscapeName, *SectionBase.ToString()));
+#endif // !UE_BUILD_SHIPPING
 
-	LODScreenRatioSquared.AddUninitialized(MaxLOD + 1);
+	// Precompute screen ratios for each LOD level : 
+	{
+		float ScreenSizeRatioDivider = FMath::Max(LOD0Distribution, 1.01f);
+		// Cancel out so that landscape is not affected by r.StaticMeshLODDistanceScale
+		float CurrentScreenSizeRatio = LOD0ScreenSize / CVarStaticMeshLODDistanceScale.GetValueOnAnyThread();
 
-	// LOD 0 handling
-	LODScreenRatioSquared[0] = FMath::Square(CurrentScreenSizeRatio);
-	LODSettings.LOD0ScreenSizeSquared = FMath::Square(CurrentScreenSizeRatio);
-	CurrentScreenSizeRatio /= ScreenSizeRatioDivider;
-	LODSettings.LOD1ScreenSizeSquared = FMath::Square(CurrentScreenSizeRatio);
-	ScreenSizeRatioDivider = FMath::Max(LODDistribution, 1.01f);
-	LODSettings.LODOnePlusDistributionScalarSquared = FMath::Square(ScreenSizeRatioDivider);
+		LODScreenRatioSquared.AddUninitialized(MaxLOD + 1);
 
-	// Other LODs
-	for (int32 LODIndex = 1; LODIndex <= MaxLOD; ++LODIndex) // This should ALWAYS be calculated from the component size, not user MaxLOD override
-	{
-		LODScreenRatioSquared[LODIndex] = FMath::Square(CurrentScreenSizeRatio);
+		// LOD 0 handling
+		LODScreenRatioSquared[0] = FMath::Square(CurrentScreenSizeRatio);
+		LODSettings.LOD0ScreenSizeSquared = FMath::Square(CurrentScreenSizeRatio);
 		CurrentScreenSizeRatio /= ScreenSizeRatioDivider;
+		LODSettings.LOD1ScreenSizeSquared = FMath::Square(CurrentScreenSizeRatio);
+		ScreenSizeRatioDivider = FMath::Max(LODDistribution, 1.01f);
+		LODSettings.LODOnePlusDistributionScalarSquared = FMath::Square(ScreenSizeRatioDivider);
+
+		// Other LODs
+		for (int32 LODIndex = 1; LODIndex <= MaxLOD; ++LODIndex) // This should ALWAYS be calculated from the component size, not user MaxLOD override
+		{
+			LODScreenRatioSquared[LODIndex] = FMath::Square(CurrentScreenSizeRatio);
+			CurrentScreenSizeRatio /= ScreenSizeRatioDivider;
+		}
 	}
 
 	FirstLOD = 0;
@@ -1338,13 +1528,47 @@ FLandscapeComponentSceneProxy::FLandscapeComponentSceneProxy(ULandscapeComponent
 	LODSettings.DrawCollisionPawnLOD = LODSettings.DrawCollisionVisibilityLOD = -1;
 #endif
 
-	ComponentMaxExtend = static_cast<float>(SubsectionSizeQuads * FMath::Max(InComponent->GetComponentTransform().GetScale3D().X, InComponent->GetComponentTransform().GetScale3D().Y));
+	// Pre-compute the last LOD value at which VSM invalidation will stop occurring : 
+	if (NonNaniteVirtualShadowMapInvalidationScreenSizeLimit <= 0.0f)
+	{
+		// Special case to disable attenuation of invalidation altogether : 
+		LODSettings.VirtualShadowMapInvalidationLimitLOD = -1.0f;
+	}
+	else
+	{
+		// Screen size under which VSM invalidation stops occurring :
+		float VirtualShadowMapInvalidationScreenSizeLimitSquared = FMath::Square(NonNaniteVirtualShadowMapInvalidationScreenSizeLimit);
+		LODSettings.VirtualShadowMapInvalidationLimitLOD = 0.0f;
+		if (VirtualShadowMapInvalidationScreenSizeLimitSquared <= LODScreenRatioSquared[MaxLOD])
+		{
+			LODSettings.VirtualShadowMapInvalidationLimitLOD = MaxLOD;
+		}
+		else
+		{
+			// Since the screen size distribution between LOD levels is not linear, we have to iterate through the LOD levels to find 
+			//  where NonNaniteVirtualShadowMapInvalidationScreenSizeLimit lays and convert back to a LOD value : 
+			float CurrentScreenSizeRatioSquared = LODScreenRatioSquared[0];
+			for (int32 LODIndex = 1; LODIndex < MaxLOD; ++LODIndex) 
+			{
+				float NextScreenSizeRatioSquared = LODScreenRatioSquared[LODIndex];
+				if ((VirtualShadowMapInvalidationScreenSizeLimitSquared > NextScreenSizeRatioSquared) && (VirtualShadowMapInvalidationScreenSizeLimitSquared <= CurrentScreenSizeRatioSquared))
+				{
+					LODSettings.VirtualShadowMapInvalidationLimitLOD = LODIndex + (CurrentScreenSizeRatioSquared - VirtualShadowMapInvalidationScreenSizeLimitSquared) / (CurrentScreenSizeRatioSquared - NextScreenSizeRatioSquared);
+					break;
+				}
+				CurrentScreenSizeRatioSquared = NextScreenSizeRatioSquared;
+			}
+		}
+	}
+
+	const FVector ComponentScale = InComponent->GetComponentTransform().GetScale3D();
+	ComponentMaxExtend = static_cast<float>(SubsectionSizeQuads * FMath::Max(ComponentScale.X, ComponentScale.Y));
 
 	if (NumSubsections > 1)
 	{
-		FRotator ComponentRotator = LandscapeComponent->GetComponentRotation();
+		FRotator ComponentRotator = InComponent->GetComponentRotation();
 		float SubSectionMaxExtend = ComponentMaxExtend / 2.0f;
-		FVector ComponentTopLeftCorner = LandscapeComponent->Bounds.Origin - ComponentRotator.RotateVector(FVector(SubSectionMaxExtend, SubSectionMaxExtend, 0.0f));
+		FVector ComponentTopLeftCorner = InComponent->Bounds.Origin - ComponentRotator.RotateVector(FVector(SubSectionMaxExtend, SubSectionMaxExtend, 0.0f));
 
 		SubSectionScreenSizeTestingPosition.AddUninitialized(MAX_SUBSECTION_COUNT);
 
@@ -1449,6 +1673,10 @@ FLandscapeComponentSceneProxy::FLandscapeComponentSceneProxy(ULandscapeComponent
 
 	bOpaqueOrMasked = true; // Landscape is always opaque
 	UpdateVisibleInLumenScene();
+
+	// Read and transform MipToMipMaxDeltas to world space to avoid doing this at runtime: 
+	WorldSpaceMipToMipMaxDeltas.Reserve(InComponent->MipToMipMaxDeltas.Num());
+	Algo::Transform(InComponent->MipToMipMaxDeltas, WorldSpaceMipToMipMaxDeltas, [ZScale = ComponentScale.Z](double InMaxDelta) { return ZScale * InMaxDelta; });
 }
 
 void FLandscapeComponentSceneProxy::CreateRenderThreadResources(FRHICommandListBase& RHICmdList)
@@ -1469,8 +1697,13 @@ void FLandscapeComponentSceneProxy::CreateRenderThreadResources(FRHICommandListB
 	SharedBuffers = FLandscapeComponentSceneProxy::SharedBuffersMap.FindRef(SharedBuffersKey);
 	if (SharedBuffers == nullptr)
 	{
+		FName BufferOwnerName;
+#if !UE_BUILD_SHIPPING
+		BufferOwnerName = DebugName;
+#endif // !UE_BUILD_SHIPPING
+
 		SharedBuffers = new FLandscapeSharedBuffers(
-			RHICmdList, SharedBuffersKey, SubsectionSizeQuads, NumSubsections, FeatureLevel, FName(LandscapeComponent->GetPathName()));
+			RHICmdList, SharedBuffersKey, SubsectionSizeQuads, NumSubsections, FeatureLevel, BufferOwnerName);
 
 		FLandscapeComponentSceneProxy::SharedBuffersMap.Add(SharedBuffersKey, SharedBuffers);
 
@@ -1890,9 +2123,18 @@ FPrimitiveViewRelevance FLandscapeComponentSceneProxy::GetViewRelevance(const FS
 	Result.bShadowRelevance = (GAllowLandscapeShadows > 0) && IsShadowCast(View) && View->Family->EngineShowFlags.Landscape;
 
 #if !UE_BUILD_SHIPPING
-	if (GVarDumpLandscapeLODsCurrentFrame == GFrameNumberRenderThread)
+	if (GVarDumpLandscapeLODsCurrentFrame != 0)
 	{
-		Result.bDynamicRelevance = true;
+		if (GVarDumpLandscapeLODsCurrentFrame == GFrameNumberRenderThread)
+		{
+			// Enable dynamic relevance to let the dump code run on GetDynamicMeshElements : 
+			Result.bDynamicRelevance = true;
+		}
+		// Assume we've dumped the info already, reset the counter : 
+		else if (GVarDumpLandscapeLODsCurrentFrame < GFrameNumberRenderThread)
+		{
+			GVarDumpLandscapeLODsCurrentFrame = 0;
+		}
 	}
 #endif // !UE_BUILD_SHIPPING
 
@@ -2032,7 +2274,7 @@ void FLandscapeComponentSceneProxy::OnTransformChanged(FRHICommandListBase& RHIC
 	LocalToWorldNoScaling = LtoW;
 	LocalToWorldNoScaling.RemoveScaling();
 
-	// Set FLandscapeUniformVSParameters for this subsection
+	// Set FLandscapeUniformShaderParameters for this subsection
 	FLandscapeUniformShaderParameters LandscapeParams;
 	LandscapeParams.ComponentBaseX = RenderCoord.X;
 	LandscapeParams.ComponentBaseY = RenderCoord.Y;
@@ -2044,6 +2286,7 @@ void FLandscapeComponentSceneProxy::OnTransformChanged(FRHICommandListBase& RHIC
 	LandscapeParams.WeightmapUVScaleBias = WeightmapScaleBias;
 	LandscapeParams.LocalToWorldNoScaling = FMatrix44f(LocalToWorldNoScaling);			// LWC_TODO: Precision loss
 	LandscapeParams.InvLODBlendRange = InvLODBlendRange;
+	LandscapeParams.NonNaniteVirtualShadowMapConstantDepthBias = VirtualShadowMapConstantDepthBias;
 
 	LandscapeParams.LandscapeLightmapScaleBias = FVector4f(
 		LightmapScaleX,
@@ -2308,11 +2551,10 @@ void FLandscapeComponentSceneProxy::DrawStaticElements(FStaticPrimitiveDrawInter
 	check(StaticBatchParamArray.Num() <= TotalBatchCount);
 }
 
+// Deprecated : do not use
 int8 FLandscapeComponentSceneProxy::GetLODFromScreenSize(float InScreenSizeSquared, float InViewLODScale) const
 {
-	float FractionalLOD;
-
-	return FLandscapeRenderSystem::GetLODFromScreenSize(LODSettings, InScreenSizeSquared, InViewLODScale, FractionalLOD);
+	return static_cast<int8>(FLandscapeRenderSystem::ComputeLODFromScreenSize(LODSettings, InScreenSizeSquared / (InViewLODScale * InViewLODScale)));
 }
 
 namespace
@@ -2707,31 +2949,38 @@ void FLandscapeComponentSceneProxy::GetDynamicMeshElements(const TArray<const FS
 #if !UE_BUILD_SHIPPING
 	if (GVarDumpLandscapeLODsCurrentFrame == GFrameNumberRenderThread)
 	{
-		for (const FSceneView* View : Views)
+		// Dump the mip-to-mip info for this component to evaluate how different mips are from one another: 
+		FString MipToMipInfoString;
+		for (int32 SourceMipIndex = 0; SourceMipIndex < NumRelevantMips - 1; ++SourceMipIndex)
 		{
-			const FString& LandscapeName = LandscapeComponent->GetLandscapeInfo()->LandscapeActor.IsValid()
-											 ? LandscapeComponent->GetLandscapeInfo()->LandscapeActor->GetName()
-											 : LexToString(LandscapeComponent->GetLandscapeInfo()->LandscapeGuid);
-			const FString& ComponentName = LandscapeComponent->GetName();
+			for (int32 DestinationMipIndex = SourceMipIndex + 1; DestinationMipIndex < NumRelevantMips; ++DestinationMipIndex)
+			{
+				int32 MipToMipDeltaIndex = UE::Landscape::ComputeMipToMipMaxDeltasIndex(SourceMipIndex, DestinationMipIndex, NumRelevantMips);
+				MipToMipInfoString += FString::Printf(TEXT("- %i->%i: %f\n"), SourceMipIndex, DestinationMipIndex, WorldSpaceMipToMipMaxDeltas[MipToMipDeltaIndex]);
+			}
+		}
+		UE_LOG(LogLandscape, Display, TEXT("%s, WorldSpaceMipToMipMaxDeltas:\n%s"), *DebugName.ToString(), *MipToMipInfoString);
 
+		for (const FSceneView* View : Views)
+		{
 			const float LODValue = ComputeLODForView(*View);
 			const int32 LOD = FMath::FloorToInt(LODValue);
 			const int32 Resolution = (ComponentSizeQuads + 1) >> LOD;
 
 			const int32 LoadedHeightmapResolution = [this]
 			{
-				if (!(LandscapeComponent->HeightmapTexture && LandscapeComponent->HeightmapTexture->GetResource()))
+				if (!(HeightmapTexture && HeightmapTexture->GetResource()))
 				{
 					return 0;
 				}
-				const int32 MipCount = LandscapeComponent->HeightmapTexture->GetResource()->GetCurrentMipCount();
+				const int32 MipCount = HeightmapTexture->GetResource()->GetCurrentMipCount();
 				return MipCount > 0 ? 1 << (MipCount - 1) : 0;
 			}();
 
 			const int32 LoadedWeightmapResolution = [this]
 			{
 				int32 MaxMipCount = 0;
-				for (const UTexture2D* WeightmapTexture : LandscapeComponent->WeightmapTextures)
+				for (const UTexture2D* WeightmapTexture : WeightmapTextures)
 				{
 					if (!(WeightmapTexture && WeightmapTexture->GetResource()))
 					{
@@ -2742,10 +2991,8 @@ void FLandscapeComponentSceneProxy::GetDynamicMeshElements(const TArray<const FS
 				return MaxMipCount > 0 ? 1 << (MaxMipCount - 1) : 0;
 			}();
 
-			UE_LOG(LogLandscape, Display, TEXT("\nView: %d, Landscape: %s, Component: %s [%s], "
-				       "LODValue: %f, LOD: %d, Resolution: %d, LoadedHeightmapMIP: %d, LoadedWeightmapMIP: %d"),
-			       View->GetViewKey(), *LandscapeName, *ComponentName, *SectionBase.ToString(),
-			       LODValue, LOD, Resolution, LoadedHeightmapResolution, LoadedWeightmapResolution);
+			UE_LOG(LogLandscape, Display, TEXT("\nView: %d, %s, LODValue: %f (LOD: %d), Resolution: %d, LoadedHeightmapMIP: %d, LoadedWeightmapMIP: %d"),
+			       View->GetViewKey(), *DebugName.ToString(), LODValue, LOD, Resolution, LoadedHeightmapResolution, LoadedWeightmapResolution);
 		}
 	}
 #endif // !UE_BUILD_SHIPPING
@@ -4194,8 +4441,6 @@ void FLandscapeComponentSceneProxy::GetLCIs(FLCIArray& LCIs)
 
 float FLandscapeComponentSceneProxy::ComputeLODForView(const FSceneView& InView) const
 {
-	// TODO: this function generates A LOT OF cache misses - it should be much better if we have an event of FTexture2DResource::UpdateTexture
-
 	int32 ViewLODOverride = GetViewLodOverride(InView, LandscapeKey);
 	float ViewLODDistanceFactor = InView.LODDistanceFactor;
 	bool ViewEngineShowFlagCollisionPawn = InView.Family->EngineShowFlags.CollisionPawn;
@@ -4217,11 +4462,9 @@ float FLandscapeComponentSceneProxy::ComputeLODForView(const FSceneView& InView)
 	float LODLevel = static_cast<float>(ForcedLODLevel);
 	if (ForcedLODLevel < 0)
 	{
-		float MeshScreenSizeSquared = ComputeBoundsScreenRadiusSquared(GetBounds().Origin, static_cast<float>(GetBounds().SphereRadius), ViewOrigin, ViewProjectionMatrix);
-
-		float FractionalLOD;
-		FLandscapeRenderSystem::GetLODFromScreenSize(LODSettings, MeshScreenSizeSquared, LODScale * LODScale, FractionalLOD);
-		LODLevel = FractionalLOD;
+		float SectionScreenSizeSquared = ComputeBoundsScreenRadiusSquared(GetBounds().Origin, static_cast<float>(GetBounds().SphereRadius), ViewOrigin, ViewProjectionMatrix);
+		SectionScreenSizeSquared /= FMath::Max(LODScale * LODScale, UE_SMALL_NUMBER);
+		LODLevel = FLandscapeRenderSystem::ComputeLODFromScreenSize(LODSettings, SectionScreenSizeSquared);
 	}
 
 	return FMath::Max(LODLevel, 0.f);
@@ -4245,6 +4488,11 @@ float FLandscapeComponentSceneProxy::ComputeLODBias() const
 	return ComputedLODBias;
 }
 
+const FPrimitiveSceneInfo* FLandscapeComponentSceneProxy::GetPrimitiveSceneInfo() const
+{
+	return FPrimitiveSceneProxy::GetPrimitiveSceneInfo();
+}
+
 void FLandscapeComponentSceneProxy::OnRenderCoordsChanged(FRHICommandListBase& RHICmdList)
 {
 	// need to rebuild the uniforms that contain render coords
@@ -4281,14 +4529,95 @@ void FLandscapeComponentSceneProxy::GetSectionCenterAndVectors(FVector& OutSecti
 	OutSectionYVectorWorldSpace = ComponentLocalToWorld.TransformVector(FVector::YAxisVector) * ComponentResolution;
 }
 
+bool FLandscapeComponentSceneProxy::ShouldInvalidateShadows(const FSceneView& InView, float InLODValue, float InLastShadowInvalidationLODValue) const
+{
+	if (WorldSpaceMipToMipMaxDeltas.IsEmpty() // Only apply if we have computed the error estimates
+		|| !GLandscapeAllowNonNaniteVirtualShadowMapInvalidation // Global switch
+		|| (bNaniteActive && InView.Family->EngineShowFlags.NaniteMeshes) // Only applies if Nanite is not active
+		|| (VirtualShadowMapInvalidationHeightErrorThreshold <= 0.0f)) // Only applies if the threshold is valid
+	{
+		return false;
+	}
+
+	// We want to estimate the height error between the height at the current LOD transition and the height at which the last shadow invalidation took place
+	// Start by clamping the values so that we're always in the [0, last relevant mip[ range (we only have mip-to-mip data up until that mip) 
+	const float LODValueClamped = FMath::Min(InLODValue, static_cast<float>(NumRelevantMips - 1) - UE_KINDA_SMALL_NUMBER);
+	const float LastShadowInvalidationLODValueClamped = FMath::Min(InLastShadowInvalidationLODValue, static_cast<float>(NumRelevantMips - 1) - UE_KINDA_SMALL_NUMBER);
+
+	// We have at our disposal the max error between a mip and any of its higher mips (N->N+1, N->N+2, etc.) so let's evaluate the error of both LODValue and LastShadowInvalidationLODValue with the 
+	//  lowest of the 2 mips as the common basis :
+	//  e.g.:
+	//  |          Mip N           |         Mip N + 1        |         Mip N + 2        |
+	//  -------------------------------------------------------------------------------------...
+	//  ^       ^                  ^                          ^  ^                       ^
+	//  |       |                  |                          |  |                       |
+	//  |       SourceLODValue     |                          |  DestinationLODValue     |    
+	//  SourceMipIndex             SourceMipIndex + 1         DestinationMipIndex        DestinationMipIndex + 1
+
+	// It doesn't matter the direction of the change, the height difference is computed in absolute values, so invert the 2 if necessary :
+	float SourceLODValue = (LODValueClamped < LastShadowInvalidationLODValueClamped) ? LODValueClamped : LastShadowInvalidationLODValueClamped;
+	float DestinationLODValue = (LODValueClamped < LastShadowInvalidationLODValueClamped) ? LastShadowInvalidationLODValueClamped : LODValueClamped;
+
+	const int32 SourceMipIndex = FMath::FloorToInt32(SourceLODValue);
+	check(SourceMipIndex + 1 < NumRelevantMips);
+	const int32 DestinationMipIndex = FMath::FloorToInt32(DestinationLODValue);
+	check(DestinationMipIndex + 1 < NumRelevantMips);
+
+	// Evaluate the max delta for both SourceLODValue and DestinationLODValue against SourceMipIndex :
+	const int32 SourceMipToMipMaxDeltaIndex = UE::Landscape::ComputeMipToMipMaxDeltasIndex(SourceMipIndex, SourceMipIndex + 1, NumRelevantMips);
+	const double SourceMipToMipMaxDelta = WorldSpaceMipToMipMaxDeltas[SourceMipToMipMaxDeltaIndex];
+	// MipToMipMaxDelta represents the maximum delta if we were to transition from SourceMipIndex to SourceMipIndex + 1 but we want to compute the error at SourceLODValue
+	//  so re-scale the delta within that range to evaluate the actual error : 
+	const double SourceMaxDelta = SourceMipToMipMaxDelta * (SourceLODValue - SourceMipIndex);
+
+	const int32 DestinationMipToMipMaxDeltaIndex = UE::Landscape::ComputeMipToMipMaxDeltasIndex(SourceMipIndex, DestinationMipIndex + 1, NumRelevantMips);
+	const double DestinationMipToMipMaxDelta = WorldSpaceMipToMipMaxDeltas[DestinationMipToMipMaxDeltaIndex];
+	// MipToMipMaxDelta represents the maximum delta if we were to transition from SourceMipIndex to DestinationMipIndex + 1 but we want to compute the error at DestinationLODValue
+	//  so re-scale the delta within that range to evaluate the actual error : 
+	check(SourceMipIndex < DestinationMipIndex + 1);
+	const double DestinationMaxDelta = DestinationMipToMipMaxDelta * (DestinationLODValue - SourceMipIndex) / (DestinationMipIndex + 1 - SourceMipIndex);
+
+	// Now we estimate that the MaxDelta between those 2 is the difference : 
+	const double MaxDelta = FMath::Abs(DestinationMaxDelta - SourceMaxDelta);
+
+	// Perform screen size-based attenuation in order to decrease the invalidation rate as the screen size decreases (as error tends to grow on higher mips, the invalidation rate increases, which 
+	//  is not desirable, since it means the screen size of the landscape section actually decreases, and so the shadow artifacts actually become less noticeable)
+	double MaxDeltaLODAttenuationFactor = 1.0;
+	if (LODSettings.VirtualShadowMapInvalidationLimitLOD > 0.0f)
+	{
+		// The closer we get to the limit LOD, the more we attenuate : 
+		MaxDeltaLODAttenuationFactor = FMath::Clamp(1.0f - (SourceLODValue / LODSettings.VirtualShadowMapInvalidationLimitLOD), 0.0f, 1.0f);
+		// Now that we're in the [0,1] range, apply an exponent to have a non-linear attenuation : 
+		MaxDeltaLODAttenuationFactor = FMath::Pow(MaxDeltaLODAttenuationFactor, GLandscapeNonNaniteVirtualShadowMapInvalidationLODAttenuationExponent);
+	}
+
+	const double AttenuatedMaxDelta = MaxDelta * MaxDeltaLODAttenuationFactor;
+
+	bool bShouldInvalidateShadow = (AttenuatedMaxDelta > VirtualShadowMapInvalidationHeightErrorThreshold);
+
+#if !UE_BUILD_SHIPPING
+	if (bShouldInvalidateShadow)
+	{
+		UE_LOG(LogLandscape, Verbose, TEXT("Shadow invalidation occured: View: %d, %s, LODValue: %f, LastShadowInvalidationLODValue: %f, \n"
+			"MipToMipMaxDelta(%d<->%d): %f (unscaled: %f), OtherMipToMipMaxDelta(%d<->%d): %f (unscaled: %f)\n"
+			"AttenuatedMaxDelta: %f, (MaxDelta: %f, MaxDeltaLODAttenuationFactor: %f, InvalidationLimitLOD: %f)"),
+			InView.GetViewKey(), *DebugName.ToString(), InLODValue, InLastShadowInvalidationLODValue,
+			SourceMipIndex, SourceMipIndex + 1, SourceMaxDelta, SourceMipToMipMaxDelta, SourceMipIndex, DestinationMipIndex + 1, DestinationMaxDelta, DestinationMipToMipMaxDelta,
+			AttenuatedMaxDelta, MaxDelta, MaxDeltaLODAttenuationFactor, LODSettings.VirtualShadowMapInvalidationLimitLOD);
+	}
+#endif // !UE_BUILD_SHIPPING
+
+	return bShouldInvalidateShadow;
+}
+
 //
 // FLandscapeSectionInfo
 //
 FLandscapeSectionInfo::FLandscapeSectionInfo(const UWorld* InWorld, const FGuid& InLandscapeGuid, const FIntPoint& InComponentBase, uint32 LODGroupKey, uint32 InLandscapeKey)
 	: LandscapeKey(InLandscapeKey)
 	, LODGroupKey(LODGroupKey)
-	, RenderCoord(INT32_MIN, INT32_MIN)
 	, ComponentBase(InComponentBase)
+	, Scene(InWorld->Scene)
 	, bResourcesCreated(false)
 	, bRegistered(false)
 {
@@ -4347,6 +4676,11 @@ public:
 	{
 	}
 
+	virtual const FPrimitiveSceneInfo* GetPrimitiveSceneInfo() const override
+	{
+		return nullptr;
+	}
+
 private:
 	int8 ProxyLOD;
 	int32 ComponentResolution;
@@ -4531,4 +4865,3 @@ FPrimitiveSceneProxy* ULandscapeNaniteComponent::CreateSceneProxy()
 	// We *only* want a Nanite proxy for this component, otherwise return null to prevent fallback rendering.
 	return nullptr;
 }
-
diff --git a/Engine/Source/Runtime/Landscape/Private/LandscapeUtils.cpp b/Engine/Source/Runtime/Landscape/Private/LandscapeUtils.cpp
--- a/Engine/Source/Runtime/Landscape/Private/LandscapeUtils.cpp
+++ b/Engine/Source/Runtime/Landscape/Private/LandscapeUtils.cpp
@@ -28,6 +28,38 @@ bool DoesPlatformSupportEditLayers(EShaderPlatform InShaderPlatform)
 		&& !IsMobilePlatform(InShaderPlatform);
 }
 
+int32 ComputeMaxDeltasOffsetForMip(int32 InMipIndex, int32 InNumRelevantMips)
+{
+	int32 Offset = 0;
+	for (int32 X = 0; X < InMipIndex; ++X)
+	{
+		Offset += InNumRelevantMips - 1 - X;
+	}
+	return Offset;
+}
+
+int32 ComputeMaxDeltasCountForMip(int32 InMipIndex, int32 InNumRelevantMips)
+{
+	return InNumRelevantMips - 1 - InMipIndex;
+}
+
+int32 ComputeMipToMipMaxDeltasIndex(int32 InSourceMipIndex, int32 InDestinationMipIndex, int32 InNumRelevantMips)
+{
+	check((InSourceMipIndex >= 0) && (InSourceMipIndex < InNumRelevantMips));
+	check((InDestinationMipIndex > InSourceMipIndex) && (InDestinationMipIndex < InNumRelevantMips));
+	return ComputeMaxDeltasOffsetForMip(InSourceMipIndex, InNumRelevantMips) + InDestinationMipIndex - InSourceMipIndex - 1;
+}
+
+int32 ComputeMipToMipMaxDeltasCount(int32 InNumRelevantMips)
+{
+	int32 Count = 0;
+	for (int32 MipIndex = 0; MipIndex < InNumRelevantMips - 1; ++MipIndex)
+	{
+		Count += InNumRelevantMips - 1 - MipIndex;
+	}
+	return Count;
+}
+
 #if WITH_EDITOR
 
 FString GetSharedAssetsPath(const FString& InPath)
diff --git a/Engine/Source/Runtime/Landscape/Public/LandscapeDataAccess.h b/Engine/Source/Runtime/Landscape/Public/LandscapeDataAccess.h
--- a/Engine/Source/Runtime/Landscape/Public/LandscapeDataAccess.h
+++ b/Engine/Source/Runtime/Landscape/Public/LandscapeDataAccess.h
@@ -25,9 +25,6 @@ namespace LandscapeDataAccess
 {
 	inline constexpr int32 MaxValue = 65535;
 	inline constexpr float MidValue = 32768.f;
-	// Reserved 2 bits for other purpose
-	// Most significant bit - Visibility, 0 is visible(default), 1 is invisible
-	// 2nd significant bit - Triangle flip, not implemented yet
 	FORCEINLINE float GetLocalHeight(uint16 Height)
 	{
 		return (static_cast<float>(Height) - MidValue) * LANDSCAPE_ZSCALE;
@@ -340,6 +337,16 @@ struct FLandscapeComponentDataInterface : public FLandscapeComponentDataInterfac
 		return GetHeight( X, Y );
 	}
 
+	float GetLocalHeight(int32 LocalX, int32 LocalY) const
+	{
+		return LandscapeDataAccess::GetLocalHeight(GetHeight(LocalX, LocalY));
+	}
+
+	float GetLocalHeight(int32 VertexIndex) const
+	{
+		return LandscapeDataAccess::GetLocalHeight(GetHeight(VertexIndex));
+	}
+
 	void GetXYOffset( int32 X, int32 Y, float& XOffset, float& YOffset ) const
 	{
 		if (XYOffsetMipData)
diff --git a/Engine/Source/Runtime/Landscape/Public/LandscapeRender.h b/Engine/Source/Runtime/Landscape/Public/LandscapeRender.h
--- a/Engine/Source/Runtime/Landscape/Public/LandscapeRender.h
+++ b/Engine/Source/Runtime/Landscape/Public/LandscapeRender.h
@@ -32,9 +32,6 @@ LandscapeRender.h: New terrain rendering
 #include "Rendering/CustomRenderPass.h"
 #include "Tasks/Task.h"
 
-// This defines the number of border blocks to surround terrain by when generating lightmaps
-#define TERRAIN_PATCH_EXPAND_SCALAR	1
-
 #define LANDSCAPE_LOD_LEVELS 8
 #define LANDSCAPE_MAX_SUBSECTION_NUM 2
 
@@ -42,6 +39,8 @@ class FLandscapeComponentSceneProxy;
 enum class ERuntimeVirtualTextureMaterialType : uint8;
 enum EShaderPlatform : uint16;
 
+namespace UE::Renderer::Private { class IShadowInvalidatingInstances; }
+
 #if RHI_RAYTRACING
 struct FLandscapeRayTracingImpl;
 #endif
@@ -110,6 +109,7 @@ BEGIN_GLOBAL_SHADER_PARAMETER_STRUCT(FLandscapeUniformShaderParameters, LANDSCAP
 	SHADER_PARAMETER(int32, LastLOD)
 	SHADER_PARAMETER(uint32, VirtualTexturePerPixelHeight)
 	SHADER_PARAMETER(float, InvLODBlendRange)
+	SHADER_PARAMETER(float, NonNaniteVirtualShadowMapConstantDepthBias)
 	SHADER_PARAMETER(FVector4f, HeightmapTextureSize)
 	SHADER_PARAMETER(FVector4f, HeightmapUVScaleBias)
     SHADER_PARAMETER(FVector4f, WeightmapUVScaleBias)
@@ -408,12 +408,12 @@ public:
 	virtual float ComputeLODForView(const FSceneView& InView) const = 0;
 	virtual float ComputeLODBias() const = 0;
 	virtual int32 GetSectionPriority() const { return INDEX_NONE; }
+	virtual const FPrimitiveSceneInfo* GetPrimitiveSceneInfo() const = 0;
 
 	/** Computes the worldspace units per vertex of the landscape section. */
 	virtual double ComputeSectionResolution() const { return -1.0; }
 	
 	virtual void GetSectionBoundsAndLocalToWorld(FBoxSphereBounds& LocalBounds, FMatrix& LocalToWorld) const = 0;
-
 	virtual void GetSectionCenterAndVectors(FVector& OutSectionCenterWorldSpace, FVector& OutSectionXVectorWorldSpace, FVector& OutSectionYVectorWorldSpace) const = 0;
 
 	/* return the resolution of a component, in vertices (-1 for any sections that are not grid based, i.e. mesh sections) */
@@ -422,14 +422,22 @@ public:
 	/* Used to notify derived classes when render coords are calculated */
 	virtual void OnRenderCoordsChanged(FRHICommandListBase& RHICmdList) = 0;
 
-public:
-	uint32 LandscapeKey;					// a hash of the world and (LandscapeGUID or LOD Group Key)
-	uint32 LODGroupKey;						// LOD Group Key (0 if no group)
-	FIntPoint RenderCoord;					// coordinate in the RenderSystem
-	FIntPoint ComponentBase;				// component base coordinate (relative to the ALandscape actor)
+	virtual bool ShouldInvalidateShadows(const FSceneView& InView, float InLODValue, float InLastShadowInvalidationLODValue) const { return false; }
 
-	bool bResourcesCreated;
-	bool bRegistered;
+public:
+	// A hash of the world and (LandscapeGUID or LOD Group Key)
+	uint32 LandscapeKey = 0;
+	// LOD Group Key (0 if no group)
+	uint32 LODGroupKey = 0;
+	// Coordinate in the RenderSystem
+	FIntPoint RenderCoord = FIntPoint(INT32_MIN, INT32_MIN);
+	// Component base coordinate (relative to the ALandscape actor)
+	FIntPoint ComponentBase = FIntPoint(ForceInit);
+	// Scene that this landscape section belongs to
+	FSceneInterface* Scene = nullptr;
+
+	bool bResourcesCreated = false;
+	bool bRegistered = false;
 };
 
 //
@@ -445,33 +453,14 @@ struct FLandscapeRenderSystem
 		float LOD1ScreenSizeSquared;
 		float LODOnePlusDistributionScalarSquared;
 		float LastLODScreenSizeSquared;
+		float VirtualShadowMapInvalidationLimitLOD;
 		int8 LastLODIndex;
 		int8 ForcedLOD;
 		int8 DrawCollisionPawnLOD;
 		int8 DrawCollisionVisibilityLOD;
 	};
 
-	static int8 GetLODFromScreenSize(LODSettingsComponent LODSettings, float InScreenSizeSquared, float InViewLODScale, float& OutFractionalLOD)
-	{
-		float ScreenSizeSquared = InScreenSizeSquared / InViewLODScale;
-		
-		if (ScreenSizeSquared <= LODSettings.LastLODScreenSizeSquared)
-		{
-			OutFractionalLOD = LODSettings.LastLODIndex;
-			return LODSettings.LastLODIndex;
-		}
-		else if (ScreenSizeSquared > LODSettings.LOD1ScreenSizeSquared)
-		{
-			OutFractionalLOD = (LODSettings.LOD0ScreenSizeSquared - FMath::Min(ScreenSizeSquared, LODSettings.LOD0ScreenSizeSquared)) / (LODSettings.LOD0ScreenSizeSquared - LODSettings.LOD1ScreenSizeSquared);
-			return 0;
-		}
-		else
-		{
-			// No longer linear fraction, but worth the cache misses
-			OutFractionalLOD = 1 + FMath::LogX(LODSettings.LODOnePlusDistributionScalarSquared, LODSettings.LOD1ScreenSizeSquared / ScreenSizeSquared);
-			return (int8)OutFractionalLOD;
-		}
-	}
+	static float ComputeLODFromScreenSize(const LODSettingsComponent& InLODSettings, float InScreenSizeSquared);
 
 	static TBitArray<> LandscapeIndexAllocator;
 
@@ -482,15 +471,30 @@ struct FLandscapeRenderSystem
 
 	TResourceArray<float> SectionLODBiases;
 	TArray<FLandscapeSectionInfo*> SectionInfos;
-	int32 ReferenceCount;			// number of sections with resources created
-	int32 RegisteredCount;			// number of sections registered
+	// Number of sections with resources created
+	int32 ReferenceCount;			
+	// Number of sections registered
+	int32 RegisteredCount;
 
 	FBufferRHIRef SectionLODBiasBuffer;
 	FShaderResourceViewRHIRef SectionLODBiasSRV;
 
 	FUniformBufferRHIRef SectionLODUniformBuffer;
 
-	TMap<const FSceneView*, TResourceArray<float>> CachedSectionLODValues;
+	// For a given view, we use those 2 maps to store, for each FLandscapeSectionInfo in SectionInfos, the current LOD Values so that it can be passed down to the renderer
+	//  These are only really valid for a given render since the list of sections can change in the next frame. 
+	//  The first map is for views that don't have a persistent view state
+	TMap<const FSceneView*, TResourceArray<float>> PerViewCachedSectionLODValues; // Key = view to render, Value = linear list of LOD values for each FLandscapeSectionInfo
+	TMap<uint32, TResourceArray<float>> PerViewStateCachedSectionLODValues; // Key = view state key corresponding to the view to render, Value = linear list of LOD values for each FLandscapeSectionInfo
+
+	// This map allows to track, for each FLandscapeSectionInfo that triggered a shadow invalidation, the LOD value that was used when the invalidation occurred
+	//  Key = SectionInfo's RenderCoord (we cannot use a linear index like the TMaps above because these might change as new sections get added/removed dynamically and this needs to be tracked across frames)
+	//  Value = LOD value of the section when the last shadow invalidation occurred
+	using SectionKeyToLODValueMap = TMap<FIntPoint, float>; 
+	TMap<uint32, SectionKeyToLODValueMap> PerViewStateLastShadowInvalidationSectionLODValues; // Key = view state key corresponding to the view to render, Value = see above
+
+	// List of shadow invalidations to perform this frame (because IShadowInvalidatingInstances is not parallel-render thread-safe, we have to store them temporarily and issue them on the render thread later on) :
+	TMap<UE::Renderer::Private::IShadowInvalidatingInstances*, TArray<const FPrimitiveSceneInfo*>> ShadowInvalidationRequests;
 
 	/** Forced LOD level which overrides the ForcedLOD level of all the sections under this LandscapeRenderSystem. */
 	int8 ForcedLODOverride;
@@ -504,7 +508,10 @@ struct FLandscapeRenderSystem
 	// Counter used to reduce how often we call compact on the map when removing sections
 	int32 SectionsRemovedSinceLastCompact;
 
-	FLandscapeRenderSystem();
+	int32 LandscapeKey = 0; 
+	FSceneInterface* Scene = nullptr;
+
+	FLandscapeRenderSystem(uint32 InLandscapeKey, FSceneInterface* InScene);
 	~FLandscapeRenderSystem();
 
 	static void CreateResources(FRHICommandListBase& RHICmdList, FLandscapeSectionInfo* SectionInfo);
@@ -548,33 +555,13 @@ struct FLandscapeRenderSystem
 		return nullptr;
 	}
 
-	float GetSectionLODValue(const FSceneView& SceneView, FIntPoint InRenderCoord) const
-	{
-		if (CachedSectionLODValues.Contains(&SceneView))
-		{
-			return CachedSectionLODValues[&SceneView][GetSectionLinearIndex(InRenderCoord)];
-		}
-		else
-		{
-			// Some views such as shadow are spawned off main view and share the same view id
-			for (const TPair<const FSceneView*, TResourceArray<float>>& Pair : CachedSectionLODValues)
-			{
-				if (Pair.Key->GetViewKey() == SceneView.GetViewKey())
-				{
-					return Pair.Value[GetSectionLinearIndex(InRenderCoord)];
-				}
-			}
-		}
-		checkf(false, TEXT("No section LOD value cached for this view. Make sure FLandscapeRenderSystem::ComputeSectionsLODForView (FLandscapeSceneViewExtension::PreRenderView_RenderThread) was called"));
-		return 0.0f;
-	}
+	float GetSectionLODValue(const FSceneView& InView, FIntPoint InRenderCoord) const;
+	const TResourceArray<float>& GetCachedSectionLODValues(const FSceneView& InView) const;
 
-	float GetSectionLODBias(FIntPoint InRenderCoord) const
-	{
-		return SectionLODBiases[GetSectionLinearIndex(InRenderCoord)];
-	}
+	float GetSectionLODBias(FIntPoint InRenderCoord) const;
 
-	const TResourceArray<float>& ComputeSectionsLODForView(const FSceneView& InView);
+	const TResourceArray<float>& ComputeSectionsLODForView(const FSceneView& InView, UE::Renderer::Private::IShadowInvalidatingInstances* InShadowInvalidatingInstances);
+	void PerformShadowInvalidations(UE::Renderer::Private::IShadowInvalidatingInstances& InShadowInvalidatingInstances);
 	void FetchHeightmapLODBiases();
 	void UpdateBuffers(FRHICommandListBase& RHICmdList);
 
@@ -611,9 +598,6 @@ public:
 	FLandscapeSceneViewExtension(const FAutoRegister& AutoReg);
 	virtual ~FLandscapeSceneViewExtension();
 
-	void EndFrame_GameThread();
-	void EndFrame_RenderThread();
-
 	virtual void SetupViewFamily(FSceneViewFamily& InViewFamily) override {}
 	virtual void SetupView(FSceneViewFamily& InViewFamily, FSceneView& InView) override {}
 	virtual void BeginRenderViewFamily(FSceneViewFamily& InViewFamily) override;
@@ -623,20 +607,23 @@ public:
 	virtual void PreInitViews_RenderThread(FRDGBuilder& GraphBuilder) override;
 
 	LANDSCAPE_API const TMap<uint32, FLandscapeRenderSystem*>& GetLandscapeRenderSystems() const;
+	TArray<FLandscapeRenderSystem*> GetLandscapeRenderSystems(const class FSceneInterface* InScene) const;
 	int32 GetNumViewsWithShowCollision() const { return NumViewsWithShowCollision; }
+
+private:
+	void EndFrame_GameThread();
+	void EndFrame_RenderThread();
+
 private:
-	FBufferRHIRef LandscapeLODDataBuffer;
-	FBufferRHIRef LandscapeIndirectionBuffer;
 
 	struct FLandscapeViewData
 	{
 		FLandscapeViewData() = default;
-
-		FLandscapeViewData(FSceneView& InView)
-			: View(&InView)
-		{}
+		FLandscapeViewData(FSceneView& InView);
 
 		FSceneView* View = nullptr;
+		// Optional interface to use if the view needs to invalidate the shadow cache : 
+		UE::Renderer::Private::IShadowInvalidatingInstances* ShadowInvalidatingInstances = nullptr;
 		TResourceArray<uint32> LandscapeIndirection;
 		TResourceArray<float> LandscapeLODData;
 	};
@@ -772,19 +759,27 @@ public:
 	static LANDSCAPE_API TMap<uint32, FLandscapeSharedBuffers*> SharedBuffersMap;
 
 protected:
-	int8						MaxLOD;						// Maximum LOD level, user override possible
-	int8						NumWeightmapLayerAllocations;
-	uint8						StaticLightingLOD;
-	uint8						VirtualTexturePerPixelHeight;
-	float						WeightmapSubsectionOffset;
-	TArray<float>				LODScreenRatioSquared;		// Table of valid screen size -> LOD index
-	int32						FirstLOD;					// First LOD we have batch elements for
-	int32						LastLOD;					// Last LOD we have batch elements for
-	int32						FirstVirtualTextureLOD;
-	int32						LastVirtualTextureLOD;
-	float						ComponentMaxExtend; 		// The max extend value in any axis
-	float						ComponentSquaredScreenSizeToUseSubSections; // Size at which we start to draw in sub lod if LOD are different per sub section
-	float						InvLODBlendRange;			// 1.0 / LODBlendRange;
+	// Maximum LOD level, user override possible
+	int8 MaxLOD;
+	int8 NumWeightmapLayerAllocations;
+	uint8 StaticLightingLOD;
+	uint8 VirtualTexturePerPixelHeight;
+	float WeightmapSubsectionOffset;
+	// Table of valid screen size -> LOD index
+	TArray<float> LODScreenRatioSquared;		
+	// First LOD we have batch elements for
+	int32 FirstLOD;	
+	// Last LOD we have batch elements for
+	int32 LastLOD;					
+	int32 FirstVirtualTextureLOD;
+	int32 LastVirtualTextureLOD;
+	// The max extend value in any axis
+	float ComponentMaxExtend; 
+	// Size at which we start to draw in sub lod if LOD are different per sub section
+	// TODO [jonathan.bard] : Remove : unused
+	float ComponentSquaredScreenSizeToUseSubSections; 
+	// 1.0 / LODBlendRange
+	float InvLODBlendRange;
 
 	FLandscapeRenderSystem::LODSettingsComponent LODSettings;
 
@@ -792,27 +787,26 @@ protected:
 	 * Number of subsections within the component in each dimension, this can be 1 or 2.
 	 * Subsections exist to improve the speed at which LOD transitions can take place over distance.
 	 */
-	int32						NumSubsections;
+	int32 NumSubsections;
 	/** Number of unique heights in the subsection. */
-	int32						SubsectionSizeQuads;
+	int32 SubsectionSizeQuads;
 	/** Number of heightmap heights in the subsection. This includes the duplicate row at the end. */
-	int32						SubsectionSizeVerts;
+	int32 SubsectionSizeVerts;
 	/** Size of the component in unique heights. */
-	int32						ComponentSizeQuads;
+	int32 ComponentSizeQuads;
 	/** 
 	 * ComponentSizeQuads + 1.
 	 * Note: in the case of multiple subsections, this is not very useful, as there will be an internal duplicate row of heights in addition to the row at the end.
 	 */
-	int32						ComponentSizeVerts;
-	float						StaticLightingResolution;
+	int32 ComponentSizeVerts;
+	float StaticLightingResolution;
 	/** Address of the component within the parent Landscape in unique height texels. */
-	FIntPoint					SectionBase;
-
-	const ULandscapeComponent*	LandscapeComponent;
+	FIntPoint SectionBase;
 
-	FMatrix						LocalToWorldNoScaling;
+	FMatrix LocalToWorldNoScaling;
 
-	TArray<FVector>				SubSectionScreenSizeTestingPosition;	// Precomputed sub section testing position for screen size calculation
+	// TODO [jonathan.bard] : remove : unused :
+	TArray<FVector> SubSectionScreenSizeTestingPosition;	// Precomputed sub section testing position for screen size calculation
 
 	// Storage for static draw list batch params
 	TArray<FLandscapeBatchElementParams> StaticBatchParamArray;
@@ -820,9 +814,8 @@ protected:
 	bool bNaniteActive;
 	bool bUsesLandscapeCulling;
 
-
 	// Precomputed grass rendering MeshBatch and per-LOD params
-	FMeshBatch                           GrassMeshBatch;
+	FMeshBatch GrassMeshBatch;
 	TArray<FLandscapeBatchElementParams> GrassBatchParams;
 
 	FVector4f WeightmapScaleBias;
@@ -842,10 +835,10 @@ protected:
 
 	UTexture2D* XYOffsetmapTexture;
 
-	uint32						SharedBuffersKey;
-	FLandscapeSharedBuffers*	SharedBuffers;
-	FLandscapeVertexFactory*	VertexFactory;
-	FLandscapeVertexFactory*	FixedGridVertexFactory;
+	uint32 SharedBuffersKey;
+	FLandscapeSharedBuffers* SharedBuffers;
+	FLandscapeVertexFactory* VertexFactory;
+	FLandscapeVertexFactory* FixedGridVertexFactory;
 
 	/** All available materials, including LOD Material, Tessellation generated materials*/
 	TArray<FMaterialRenderProxy*> AvailableMaterials;
@@ -862,12 +855,22 @@ protected:
 	/** Material Relevance for each material in AvailableMaterials */
 	TArray<FMaterialRelevance> MaterialRelevances;
 
-#if WITH_EDITORONLY_DATA
-	FLandscapeEditToolRenderData EditToolRenderData;
-#endif
+	/** Number of mips that are actually usable (that have more than 1 vertex) */
+	int32 NumRelevantMips = 0;
+
+	/** Maximum deltas between vertices and their counterparts from other mips (see additional details in ULandscapeComponent)
+	*   Stored in world space to avoid useless runtime computation :
+	*/
+	TArray<double> WorldSpaceMipToMipMaxDeltas;
+
+	/** Constant bias to handle the worst artifacts of the continuous LOD morphing when rendering to VSM. */
+	float VirtualShadowMapConstantDepthBias;
+
+	/** Height threshold to invalidate VSM pages when using non-Nanite landscape. */
+	float VirtualShadowMapInvalidationHeightErrorThreshold;
 
 #if WITH_EDITORONLY_DATA
-	ELandscapeLODFalloff::Type LODFalloff_DEPRECATED;
+	FLandscapeEditToolRenderData EditToolRenderData;
 #endif
 
 	// data used in editor or visualisers
@@ -885,16 +888,21 @@ protected:
 
 	TUniformBuffer<FLandscapeUniformShaderParameters> LandscapeUniformShaderParameters;
 
-	TArray< TUniformBuffer<FLandscapeFixedGridUniformShaderParameters> > LandscapeFixedGridUniformShaderParameters;
+	TArray<TUniformBuffer<FLandscapeFixedGridUniformShaderParameters>> LandscapeFixedGridUniformShaderParameters;
 
 	// Cached versions of these
-	FMatrix					WorldToLocal;
+	FMatrix WorldToLocal;
 
 	FLandscapeVisibilityHelper VisibilityHelper;
 
+#if !UE_BUILD_SHIPPING
+	FName DebugName;
+#endif // !UE_BUILD_SHIPPING
+
 protected:
 	LANDSCAPE_API virtual ~FLandscapeComponentSceneProxy();
 	
+	UE_DEPRECATED(5.4, "Removed function")
 	LANDSCAPE_API int8 GetLODFromScreenSize(float InScreenSizeSquared, float InViewLODScale) const;
 
 	LANDSCAPE_API bool GetMeshElementForFixedGrid(int32 InLodIndex, FMaterialRenderProxy* InMaterialInterface, FMeshBatch& OutMeshBatch, TArray<FLandscapeBatchElementParams>& OutStaticBatchParamArray) const;
@@ -952,11 +960,11 @@ public:
 	LANDSCAPE_API virtual float ComputeLODBias() const override;
 	LANDSCAPE_API virtual void OnRenderCoordsChanged(FRHICommandListBase& RHICmdList) override;
 	LANDSCAPE_API virtual int32 GetComponentResolution() const override;
-
 	LANDSCAPE_API virtual double ComputeSectionResolution() const override;
 	LANDSCAPE_API virtual void GetSectionBoundsAndLocalToWorld(FBoxSphereBounds& LocalBounds, FMatrix& LocalToWorld) const override;
-
 	virtual void GetSectionCenterAndVectors(FVector& OutSectionCenterWorldSpace, FVector& OutSectionXVectorWorldSpace, FVector& OutSectionYVectorWorldSpace) const override;
+	virtual const FPrimitiveSceneInfo* GetPrimitiveSceneInfo() const override;
+	virtual bool ShouldInvalidateShadows(const FSceneView& InView, float InLODValue, float InLastShadowInvalidationLODValue) const override;
 };
 
 class FLandscapeDebugMaterialRenderProxy : public FMaterialRenderProxy
diff --git a/Engine/Source/Runtime/Landscape/Public/LandscapeUtils.h b/Engine/Source/Runtime/Landscape/Public/LandscapeUtils.h
--- a/Engine/Source/Runtime/Landscape/Public/LandscapeUtils.h
+++ b/Engine/Source/Runtime/Landscape/Public/LandscapeUtils.h
@@ -25,6 +25,11 @@ namespace UE::Landscape
 */
 LANDSCAPE_API bool DoesPlatformSupportEditLayers(EShaderPlatform InShaderPlatform);
 
+int32 ComputeMaxDeltasOffsetForMip(int32 InMipIndex, int32 InNumRelevantMips);
+int32 ComputeMaxDeltasCountForMip(int32 InMipIndex, int32 InNumRelevantMips);
+int32 ComputeMipToMipMaxDeltasIndex(int32 InSourceMipIndex, int32 InDestinationMipIndex, int32 InNumRelevantMips);
+int32 ComputeMipToMipMaxDeltasCount(int32 InNumRelevantMips);
+
 #if WITH_EDITOR
 
 struct FTextureCopyRequest
diff --git a/Engine/Source/Runtime/Renderer/Private/SceneRendering.cpp b/Engine/Source/Runtime/Renderer/Private/SceneRendering.cpp
--- a/Engine/Source/Runtime/Renderer/Private/SceneRendering.cpp
+++ b/Engine/Source/Runtime/Renderer/Private/SceneRendering.cpp
@@ -4565,6 +4565,7 @@ void FSceneRenderer::RenderThreadEnd(FRHICommandListImmediate& RHICmdList, const
 
 UE::Renderer::Private::IShadowInvalidatingInstances *FSceneRenderer::GetShadowInvalidatingInstancesInterface(const FSceneView *SceneView)
 {
+	checkf(IsInRenderingThread(), TEXT("Accessing the ShadowInvalidatingInstancesInterface should only be allowed from the rendering thread!"));
 	if (ShadowSceneRenderer)
 	{
 		return ShadowSceneRenderer->GetInvalidatingInstancesInterface(SceneView);
