Landscape : Only propagate inherited / overriden landscape properties if it's not an interactive change

diff --git a/Engine/Source/Runtime/Landscape/Private/LandscapeEdit.cpp b/Engine/Source/Runtime/Landscape/Private/LandscapeEdit.cpp
--- a/Engine/Source/Runtime/Landscape/Private/LandscapeEdit.cpp
+++ b/Engine/Source/Runtime/Landscape/Private/LandscapeEdit.cpp
@@ -6566,8 +6566,9 @@ void ALandscape::PostEditChangeProperty(FPropertyChangedEvent& PropertyChangedEv
 		bMarkAllLandscapeRenderStateDirty = true;
 	}
 	
-	// If the property that has changed is overridable or inherited, synchronize the change on all landscape proxies :
-	if (IsSharedProperty(MemberPropertyName))
+	// If the property that has changed is overridable or inherited and not interactive (dragging a slider)
+	// synchronize the change on all landscape proxies :
+	if (IsSharedProperty(MemberPropertyName) && PropertyChangedEvent.ChangeType != EPropertyChangeType::Interactive)
 	{
 		bPropagateToProxies = true;
 	}
