Fixed crash in landscape culling when collecting landscape render systems, it's now possible to end with some that are not rendered (e.g. the landscape from the Editor world when in PIE) and therefore its CachedLODSectionValues have not been computed. We now simply skip other those.
#tests Editor, PIE, mobile preview#preflight 65cc30214a670ab24cf303cf

diff --git a/Engine/Source/Runtime/Landscape/Private/LandscapeCulling.cpp b/Engine/Source/Runtime/Landscape/Private/LandscapeCulling.cpp
--- a/Engine/Source/Runtime/Landscape/Private/LandscapeCulling.cpp
+++ b/Engine/Source/Runtime/Landscape/Private/LandscapeCulling.cpp
@@ -9,8 +9,6 @@
 #include "MeshMaterialShader.h"
 #include "DataDrivenShaderPlatformInfo.h"
 
-extern TMap<uint32, FLandscapeRenderSystem*> LandscapeRenderSystems;
-
 static TAutoConsoleVariable<int32> CVarLandscapeSupportCulling(
 	TEXT("landscape.SupportGPUCulling"),
 	1,
@@ -615,39 +613,43 @@ static void ComputeSectionIntermediateData(FRDGBuilder& GraphBuilder, TArrayView
 			FViewStateIntermediateData& ViewStateIntermediates = CullingEntry.IntermediateData.AddDefaulted_GetRef();
 			ViewStateIntermediates.SectionsBufferRDG = nullptr;
 
-			const FLandscapeRenderSystem& RenderSystem = *LandscapeRenderSystems.FindChecked(CullingEntry.LandscapeKey);
-			const TResourceArray<float>& SectionLODValues = RenderSystem.GetCachedSectionLODValues(View);
-
-			for (int32 SectionIdx = 0; SectionIdx < SectionLODValues.Num(); ++SectionIdx)
+			const FLandscapeRenderSystem* RenderSystem = FLandscapeSceneViewExtension::GetLandscapeRenderSystem(View.Family->Scene, CullingEntry.LandscapeKey);
+			// This landscape render system might not correspond to the scene we're rendering, so we might end up with nothing here : 
+			if (RenderSystem != nullptr)
 			{
-				const int32 LODValue = static_cast<int32>(SectionLODValues[SectionIdx]);
-				FLandscapeSectionInfo* SectionInfo = RenderSystem.SectionInfos[SectionIdx];
+				const TResourceArray<float>& SectionLODValues = RenderSystem->GetCachedSectionLODValues(View);
 
-				if (LODValue == 0 && SectionInfo != nullptr)
+				for (int32 SectionIdx = 0; SectionIdx < SectionLODValues.Num(); ++SectionIdx)
 				{
-					FBuildLandscapeTileDataCS::FLandscapeSection& Section = SectionsData.AddDefaulted_GetRef();
-
-					FBoxSphereBounds SectionLocalBounds;
-					FMatrix SectionLocalToWorld;
-					SectionInfo->GetSectionBoundsAndLocalToWorld(SectionLocalBounds, SectionLocalToWorld);
-					const FLargeWorldRenderPosition SectionAbsoluteOrigin(SectionLocalToWorld.GetOrigin());
-					const int32 NeighborsMaxLOD = static_cast<int32>(FMath::RoundFromZero(ComputeNeighborsMaxLOD(RenderSystem, SectionLODValues, SectionInfo->RenderCoord)));
-
-					Section.LocalToRelativeWorld = FLargeWorldRenderScalar::MakeToRelativeWorldMatrix(SectionAbsoluteOrigin.GetTileOffset(), SectionLocalToWorld);
-					Section.TilePosition = SectionAbsoluteOrigin.GetTile();
-					Section.LocalZ = static_cast<float>(SectionLocalBounds.Origin.Z);
-					Section.HalfHeight = static_cast<float>(SectionLocalBounds.BoxExtent.Z);
-					// How many quads to add to each tile extent to compensate for a neighbors LOD
-					Section.NeighborLODExtent = FMath::Max(static_cast<float>((1 << NeighborsMaxLOD) - 1), 1.f);
-
-					ViewStateIntermediates.SectionRenderCoords.Add(SectionInfo->RenderCoord);
+					const int32 LODValue = static_cast<int32>(SectionLODValues[SectionIdx]);
+					FLandscapeSectionInfo* SectionInfo = RenderSystem->SectionInfos[SectionIdx];
+
+					if (LODValue == 0 && SectionInfo != nullptr)
+					{
+						FBuildLandscapeTileDataCS::FLandscapeSection& Section = SectionsData.AddDefaulted_GetRef();
+
+						FBoxSphereBounds SectionLocalBounds;
+						FMatrix SectionLocalToWorld;
+						SectionInfo->GetSectionBoundsAndLocalToWorld(SectionLocalBounds, SectionLocalToWorld);
+						const FLargeWorldRenderPosition SectionAbsoluteOrigin(SectionLocalToWorld.GetOrigin());
+						const int32 NeighborsMaxLOD = static_cast<int32>(FMath::RoundFromZero(ComputeNeighborsMaxLOD(*RenderSystem, SectionLODValues, SectionInfo->RenderCoord)));
+
+						Section.LocalToRelativeWorld = FLargeWorldRenderScalar::MakeToRelativeWorldMatrix(SectionAbsoluteOrigin.GetTileOffset(), SectionLocalToWorld);
+						Section.TilePosition = SectionAbsoluteOrigin.GetTile();
+						Section.LocalZ = static_cast<float>(SectionLocalBounds.Origin.Z);
+						Section.HalfHeight = static_cast<float>(SectionLocalBounds.BoxExtent.Z);
+						// How many quads to add to each tile extent to compensate for a neighbors LOD
+						Section.NeighborLODExtent = FMath::Max(static_cast<float>((1 << NeighborsMaxLOD) - 1), 1.f);
+
+						ViewStateIntermediates.SectionRenderCoords.Add(SectionInfo->RenderCoord);
+					}
 				}
-			}
 
-			if (SectionsData.Num() != 0)
-			{
-				ViewStateIntermediates.SectionsBufferRDG = CreateStructuredBuffer<FBuildLandscapeTileDataCS::FLandscapeSection>(GraphBuilder, TEXT("LandscapeCulling.SectionsData"), SectionsData);
-				SectionsData.Reset();
+				if (SectionsData.Num() != 0)
+				{
+					ViewStateIntermediates.SectionsBufferRDG = CreateStructuredBuffer<FBuildLandscapeTileDataCS::FLandscapeSection>(GraphBuilder, TEXT("LandscapeCulling.SectionsData"), SectionsData);
+					SectionsData.Reset();
+				}
 			}
 		}
 	}
diff --git a/Engine/Source/Runtime/Landscape/Private/LandscapeRender.cpp b/Engine/Source/Runtime/Landscape/Private/LandscapeRender.cpp
--- a/Engine/Source/Runtime/Landscape/Private/LandscapeRender.cpp
+++ b/Engine/Source/Runtime/Landscape/Private/LandscapeRender.cpp
@@ -1246,13 +1246,13 @@ void FLandscapeSceneViewExtension::EndFrame_RenderThread()
 	}
 }
 
-const TMap<uint32, FLandscapeRenderSystem*>& FLandscapeSceneViewExtension::GetLandscapeRenderSystems() const
+const TMap<uint32, FLandscapeRenderSystem*>& FLandscapeSceneViewExtension::GetLandscapeRenderSystems()
 {
 	checkf(IsInParallelRenderingThread(), TEXT("Accessing the Landscape render systems is only valid from the rendering thread!"));
 	return LandscapeRenderSystems;
 }
 
-TArray<FLandscapeRenderSystem*> FLandscapeSceneViewExtension::GetLandscapeRenderSystems(const FSceneInterface* InScene) const
+TArray<FLandscapeRenderSystem*> FLandscapeSceneViewExtension::GetLandscapeRenderSystems(const FSceneInterface* InScene)
 {
 	checkf(IsInParallelRenderingThread(), TEXT("Accessing the Landscape render systems is only valid from the rendering thread!"));
 	check(InScene != nullptr);
@@ -1270,6 +1270,13 @@ TArray<FLandscapeRenderSystem*> FLandscapeSceneViewExtension::GetLandscapeRender
 	return Result;
 }
 
+FLandscapeRenderSystem* FLandscapeSceneViewExtension::GetLandscapeRenderSystem(const class FSceneInterface* InScene, uint32 InLandscapeKey)
+{
+	TArray<FLandscapeRenderSystem*> ViewLandscapeRenderSystems = GetLandscapeRenderSystems(InScene);
+	FLandscapeRenderSystem** Found = ViewLandscapeRenderSystems.FindByPredicate([InLandscapeKey](FLandscapeRenderSystem* InRenderSystem) { return InRenderSystem->LandscapeKey == InLandscapeKey; });
+	return (Found != nullptr) ? *Found : nullptr;
+}
+
 //
 // FLandscapeVisibilityHelper
 //
diff --git a/Engine/Source/Runtime/Landscape/Public/LandscapeRender.h b/Engine/Source/Runtime/Landscape/Public/LandscapeRender.h
--- a/Engine/Source/Runtime/Landscape/Public/LandscapeRender.h
+++ b/Engine/Source/Runtime/Landscape/Public/LandscapeRender.h
@@ -508,7 +508,7 @@ struct FLandscapeRenderSystem
 	// Counter used to reduce how often we call compact on the map when removing sections
 	int32 SectionsRemovedSinceLastCompact;
 
-	int32 LandscapeKey = 0; 
+	uint32 LandscapeKey = 0; 
 	FSceneInterface* Scene = nullptr;
 
 	FLandscapeRenderSystem(uint32 InLandscapeKey, FSceneInterface* InScene);
@@ -606,8 +606,9 @@ public:
 	virtual void PreRenderView_RenderThread(FRDGBuilder& GraphBuilder, FSceneView& InView) override;
 	virtual void PreInitViews_RenderThread(FRDGBuilder& GraphBuilder) override;
 
-	LANDSCAPE_API const TMap<uint32, FLandscapeRenderSystem*>& GetLandscapeRenderSystems() const;
-	TArray<FLandscapeRenderSystem*> GetLandscapeRenderSystems(const class FSceneInterface* InScene) const;
+	LANDSCAPE_API static const TMap<uint32, FLandscapeRenderSystem*>& GetLandscapeRenderSystems();
+	static TArray<FLandscapeRenderSystem*> GetLandscapeRenderSystems(const class FSceneInterface* InScene);
+	static FLandscapeRenderSystem* GetLandscapeRenderSystem(const class FSceneInterface* InScene, uint32 InLandscapeKey);
 	int32 GetNumViewsWithShowCollision() const { return NumViewsWithShowCollision; }
 
 private:
