Removed/deprecated landscape code before 5.4 cutoff

diff --git a/Engine/Source/Editor/DetailCustomizations/Private/LandscapeUIDetails.cpp b/Engine/Source/Editor/DetailCustomizations/Private/LandscapeUIDetails.cpp
--- a/Engine/Source/Editor/DetailCustomizations/Private/LandscapeUIDetails.cpp
+++ b/Engine/Source/Editor/DetailCustomizations/Private/LandscapeUIDetails.cpp
@@ -66,12 +66,6 @@ void FLandscapeUIDetails::CustomizeDetails( IDetailLayoutBuilder& DetailBuilder
 			return;
 		}
 
-		if (Landscape->NumSubsections == 1)
-		{
-			TSharedRef<IPropertyHandle> ComponentScreenSizeToUseSubSectionsProp = DetailBuilder.GetProperty(GET_MEMBER_NAME_CHECKED(ALandscapeProxy, ComponentScreenSizeToUseSubSections));
-			DetailBuilder.HideProperty(ComponentScreenSizeToUseSubSectionsProp);
-		}
-
 		TSharedRef<IPropertyHandle> CanHaveLayersPropertyHandle = DetailBuilder.GetProperty(GET_MEMBER_NAME_CHECKED(ALandscape, bCanHaveLayersContent));
 		DetailBuilder.HideProperty(CanHaveLayersPropertyHandle);
 		const FText DisplayAndFilterText(LOCTEXT("LandscapeToggleLayerName", "Enable Edit Layers"));
diff --git a/Engine/Source/Runtime/Landscape/Classes/LandscapeProxy.h b/Engine/Source/Runtime/Landscape/Classes/LandscapeProxy.h
--- a/Engine/Source/Runtime/Landscape/Classes/LandscapeProxy.h
+++ b/Engine/Source/Runtime/Landscape/Classes/LandscapeProxy.h
@@ -458,8 +458,8 @@ public:
 	TEnumAsByte<ELandscapeLODFalloff::Type> LODFalloff_DEPRECATED;
 #endif // WITH_EDITORONLY_DATA
 
-	/** Component screen size (0.0 - 1.0) at which we should keep sub sections. This is mostly pertinent if you have large component of > 64 and component are close to the camera. The goal is to reduce draw call, so if a component is smaller than the value, we merge all subsections into 1 drawcall. */
-	UPROPERTY(EditAnywhere, Category = LOD, meta=(ClampMin = "0.01", ClampMax = "1.0", UIMin = "0.01", UIMax = "1.0", DisplayName= "SubSection Min Component ScreenSize", LandscapeInherited))
+	UE_DEPRECATED(5.4, "Unused property")
+	UPROPERTY()
 	float ComponentScreenSizeToUseSubSections;
 
 	/** This is the starting screen size used to calculate the distribution. You can increase the value if you want less LOD0 component, and you use very large landscape component. */
@@ -920,7 +920,7 @@ public:
 	virtual void ChangeLODDistanceFactor(float InLODDistanceFactor);
 
 	/** Change ComponentScreenSizeToUseSubSections value on the render proxy.*/
-	UFUNCTION(BlueprintCallable, Category = "Rendering", meta = (DeprecatedFunction, DeprecationMessage = "This value can't be changed anymore, you should edit the property ComponentScreenSizeToUseSubSections of the Landscape"))
+	UFUNCTION(BlueprintCallable, Category = "Rendering", meta = (DeprecatedFunction, DeprecationMessage = "This value can't be changed anymore and has been ineffective for several versions now. Please stop using it"))
 	virtual void ChangeComponentScreenSizeToUseSubSections(float InComponentScreenSizeToUseSubSections);
 
 	/* Setter for LandscapeMaterial. Has no effect outside the editor. */
diff --git a/Engine/Source/Runtime/Landscape/Private/Landscape.cpp b/Engine/Source/Runtime/Landscape/Private/Landscape.cpp
--- a/Engine/Source/Runtime/Landscape/Private/Landscape.cpp
+++ b/Engine/Source/Runtime/Landscape/Private/Landscape.cpp
@@ -1554,7 +1554,6 @@ ALandscapeProxy::ALandscapeProxy(const FObjectInitializer& ObjectInitializer)
 #if WITH_EDITORONLY_DATA
 	bLockLocation = true;
 #endif // WITH_EDITORONLY_DATA
-	ComponentScreenSizeToUseSubSections = 0.65f;
 	bCastStaticShadow = true;
 	ShadowCacheInvalidationBehavior = EShadowCacheInvalidationBehavior::Auto;
 	bUsedForNavigation = true;
@@ -4424,7 +4423,6 @@ void ALandscapeProxy::GetSharedProperties(ALandscapeProxy* Landscape)
 		NumSubsections = Landscape->NumSubsections;
 		SubsectionSizeQuads = Landscape->SubsectionSizeQuads;
 		MaxLODLevel = Landscape->MaxLODLevel;
-		ComponentScreenSizeToUseSubSections = Landscape->ComponentScreenSizeToUseSubSections;
 		ScalableLODDistributionSetting = Landscape->ScalableLODDistributionSetting;
 		ScalableLOD0DistributionSetting = Landscape->ScalableLOD0DistributionSetting;
 		ScalableLOD0ScreenSize = Landscape->ScalableLOD0ScreenSize;
diff --git a/Engine/Source/Runtime/Landscape/Private/LandscapeEdit.cpp b/Engine/Source/Runtime/Landscape/Private/LandscapeEdit.cpp
--- a/Engine/Source/Runtime/Landscape/Private/LandscapeEdit.cpp
+++ b/Engine/Source/Runtime/Landscape/Private/LandscapeEdit.cpp
@@ -5982,7 +5982,6 @@ void ALandscapeProxy::PostEditChangeProperty(FPropertyChangedEvent& PropertyChan
 		|| (MemberPropertyName == GET_MEMBER_NAME_CHECKED(ALandscapeProxy, LOD0ScreenSize))
 		|| (MemberPropertyName == GET_MEMBER_NAME_CHECKED(ALandscapeProxy, bUseScalableLODSettings))
 		|| (MemberPropertyName == GET_MEMBER_NAME_CHECKED(ALandscapeProxy, LODBlendRange))
-		|| (MemberPropertyName == GET_MEMBER_NAME_CHECKED(ALandscapeProxy, ComponentScreenSizeToUseSubSections))
 		|| (MemberPropertyName == GET_MEMBER_NAME_CHECKED(ALandscapeProxy, RuntimeVirtualTextures))
 		|| (MemberPropertyName == GET_MEMBER_NAME_CHECKED(ALandscapeProxy, VirtualTextureRenderPassType))
 		|| (MemberPropertyName == GET_MEMBER_NAME_CHECKED(ALandscapeProxy, bVirtualTextureRenderWithQuad))
diff --git a/Engine/Source/Runtime/Landscape/Private/LandscapeRender.cpp b/Engine/Source/Runtime/Landscape/Private/LandscapeRender.cpp
--- a/Engine/Source/Runtime/Landscape/Private/LandscapeRender.cpp
+++ b/Engine/Source/Runtime/Landscape/Private/LandscapeRender.cpp
@@ -305,7 +305,6 @@ namespace UE::Landscape
 FLandscapeDebugOptions::FLandscapeDebugOptions()
 	: bShowPatches(false)
 	, bDisableStatic(false)
-	, CombineMode(eCombineMode_Default)
 	, PatchesConsoleCommand(
 		TEXT("Landscape.Patches"),
 		TEXT("Show/hide Landscape patches"),
@@ -314,10 +313,6 @@ FLandscapeDebugOptions::FLandscapeDebugOptions()
 		TEXT("Landscape.Static"),
 		TEXT("Enable/disable Landscape static drawlists"),
 		FConsoleCommandDelegate::CreateRaw(this, &FLandscapeDebugOptions::Static))
-	, CombineConsoleCommand(
-		TEXT("Landscape.Combine"),
-		TEXT("Set landscape component combining mode : 0 = Default, 1 = Combine All, 2 = Disabled"),
-		FConsoleCommandWithArgsDelegate::CreateRaw(this, &FLandscapeDebugOptions::Combine))
 {
 }
 
@@ -333,15 +328,6 @@ void FLandscapeDebugOptions::Static()
 	UE_LOG(LogLandscape, Display, TEXT("Landscape.Static: %s"), bDisableStatic ? TEXT("Disabled") : TEXT("Enabled"));
 }
 
-void FLandscapeDebugOptions::Combine(const TArray<FString>& Args)
-{
-	if (Args.Num() >= 1)
-	{
-		CombineMode = (eCombineMode)FCString::Atoi(*Args[0]);
-		UE_LOG(LogLandscape, Display, TEXT("Landscape.Combine: %d"), (int32)CombineMode);
-	}
-}
-
 FLandscapeDebugOptions GLandscapeDebugOptions;
 
 
@@ -1320,7 +1306,6 @@ FLandscapeComponentSceneProxy::FLandscapeComponentSceneProxy(ULandscapeComponent
 	, FirstLOD(0)
 	, LastLOD(MaxLOD)
 	, ComponentMaxExtend(0.0f)
-	, ComponentSquaredScreenSizeToUseSubSections(FMath::Square(InComponent->GetLandscapeProxy()->ComponentScreenSizeToUseSubSections))
 	, InvLODBlendRange(1.0f / FMath::Max(0.01f, GLandscapeLODBlendRangeOverride > 0.0f ? GLandscapeLODBlendRangeOverride : InComponent->GetLandscapeProxy()->LODBlendRange))
 	, NumSubsections(InComponent->NumSubsections)
 	, SubsectionSizeQuads(InComponent->SubsectionSizeQuads)
@@ -1571,24 +1556,6 @@ FLandscapeComponentSceneProxy::FLandscapeComponentSceneProxy(ULandscapeComponent
 	const FVector ComponentScale = InComponent->GetComponentTransform().GetScale3D();
 	ComponentMaxExtend = static_cast<float>(SubsectionSizeQuads * FMath::Max(ComponentScale.X, ComponentScale.Y));
 
-	if (NumSubsections > 1)
-	{
-		FRotator ComponentRotator = InComponent->GetComponentRotation();
-		float SubSectionMaxExtend = ComponentMaxExtend / 2.0f;
-		FVector ComponentTopLeftCorner = InComponent->Bounds.Origin - ComponentRotator.RotateVector(FVector(SubSectionMaxExtend, SubSectionMaxExtend, 0.0f));
-
-		SubSectionScreenSizeTestingPosition.AddUninitialized(MAX_SUBSECTION_COUNT);
-
-		for (int32 SubY = 0; SubY < NumSubsections; ++SubY)
-		{
-			for (int32 SubX = 0; SubX < NumSubsections; ++SubX)
-			{
-				int32 SubSectionIndex = SubX + SubY * NumSubsections;
-				SubSectionScreenSizeTestingPosition[SubSectionIndex] = ComponentTopLeftCorner + ComponentRotator.RotateVector(FVector(ComponentMaxExtend * SubX, ComponentMaxExtend * SubY, 0.0f));
-			}
-		}
-	}
-
 	if (InComponent->StaticLightingResolution > 0.f)
 	{
 		StaticLightingResolution = InComponent->StaticLightingResolution;
@@ -2404,23 +2371,6 @@ bool FLandscapeComponentSceneProxy::GetMeshElementForFixedGrid(int32 InLodIndex,
 	return true;
 }
 
-void FLandscapeComponentSceneProxy::ApplyWorldOffset(FRHICommandListBase& RHICmdList, FVector InOffset)
-{
-	FPrimitiveSceneProxy::ApplyWorldOffset(RHICmdList, InOffset);
-
-	if (NumSubsections > 1)
-	{
-		for (int32 SubY = 0; SubY < NumSubsections; ++SubY)
-		{
-			for (int32 SubX = 0; SubX < NumSubsections; ++SubX)
-			{
-				int32 SubSectionIndex = SubX + SubY * NumSubsections;
-				SubSectionScreenSizeTestingPosition[SubSectionIndex] += InOffset;
-			}
-		}
-	}
-}
-
 template<class ArrayType>
 bool FLandscapeComponentSceneProxy::GetStaticMeshElement(int32 LODIndex, bool bForToolMesh, FMeshBatch& MeshBatch, ArrayType& OutStaticBatchParamArray) const
 {
@@ -4379,7 +4329,6 @@ void ALandscapeProxy::ChangeLODDistanceFactor(float InLODDistanceFactor)
 void FLandscapeComponentSceneProxy::ChangeComponentScreenSizeToUseSubSections_RenderThread(float InComponentScreenSizeToUseSubSections)
 {
 	// Deprecated
-	ComponentSquaredScreenSizeToUseSubSections = FMath::Square(InComponentScreenSizeToUseSubSections);
 }
 
 bool FLandscapeComponentSceneProxy::HeightfieldHasPendingStreaming() const
diff --git a/Engine/Source/Runtime/Landscape/Public/LandscapeRender.h b/Engine/Source/Runtime/Landscape/Public/LandscapeRender.h
--- a/Engine/Source/Runtime/Landscape/Public/LandscapeRender.h
+++ b/Engine/Source/Runtime/Landscape/Public/LandscapeRender.h
@@ -659,25 +659,15 @@ struct FLandscapeDebugOptions
 {
 	LANDSCAPE_API FLandscapeDebugOptions();
 
-	enum eCombineMode
-	{
-		eCombineMode_Default = 0,
-		eCombineMode_CombineAll = 1,
-		eCombineMode_Disabled = 2
-	};
-
 	bool bShowPatches;
 	bool bDisableStatic;
-	eCombineMode CombineMode;
 
 private:
 	FAutoConsoleCommand PatchesConsoleCommand;
 	FAutoConsoleCommand StaticConsoleCommand;
-	FAutoConsoleCommand CombineConsoleCommand;
 
 	void Patches();
 	void Static();
-	void Combine(const TArray<FString>& Args);
 };
 
 LANDSCAPE_API extern FLandscapeDebugOptions GLandscapeDebugOptions;
@@ -776,9 +766,6 @@ protected:
 	int32 LastVirtualTextureLOD;
 	// The max extend value in any axis
 	float ComponentMaxExtend; 
-	// Size at which we start to draw in sub lod if LOD are different per sub section
-	// TODO [jonathan.bard] : Remove : unused
-	float ComponentSquaredScreenSizeToUseSubSections; 
 	// 1.0 / LODBlendRange
 	float InvLODBlendRange;
 
@@ -806,9 +793,6 @@ protected:
 
 	FMatrix LocalToWorldNoScaling;
 
-	// TODO [jonathan.bard] : remove : unused :
-	TArray<FVector> SubSectionScreenSizeTestingPosition;	// Precomputed sub section testing position for screen size calculation
-
 	// Storage for static draw list batch params
 	TArray<FLandscapeBatchElementParams> StaticBatchParamArray;
 
@@ -914,7 +898,6 @@ public:
 	LANDSCAPE_API FLandscapeComponentSceneProxy(ULandscapeComponent* InComponent);
 
 	// FPrimitiveSceneProxy interface.
-	LANDSCAPE_API virtual void ApplyWorldOffset(FRHICommandListBase& RHICmdList, FVector InOffset) override;
 	LANDSCAPE_API virtual void DrawStaticElements(FStaticPrimitiveDrawInterface* PDI) override;
 	LANDSCAPE_API virtual void GetDynamicMeshElements(const TArray<const FSceneView*>& Views, const FSceneViewFamily& ViewFamily, uint32 VisibilityMap, FMeshElementCollector& Collector) const override;
 	LANDSCAPE_API virtual void ApplyViewDependentMeshArguments(const FSceneView& View, FMeshBatch& ViewDependentMeshBatch) const override;
@@ -937,7 +920,7 @@ public:
 	const FMeshBatch& GetGrassMeshBatch() const { return GrassMeshBatch; }
 
 	// FLandcapeSceneProxy
-	UE_DEPRECATED(5.4, "This function has been deprecated, use ComponentScreenSizeToUseSubSections on ALandscapeProxy instead")
+	UE_DEPRECATED(5.4, "This function has been deprecated and is ineffective")
 	LANDSCAPE_API void ChangeComponentScreenSizeToUseSubSections_RenderThread(float InComponentScreenSizeToUseSubSections);
 
 	LANDSCAPE_API virtual bool HeightfieldHasPendingStreaming() const override;
