World Partition - HLOD: Added Nanite support for landscape HLODs
* Previously, even if we had a Nanite static mesh, the fact that we relied on a ULandscapeMeshProxyComponent prevented the creation of a nanite scene proxy.
* Now, we create a standard static mesh component when building a Nanite HLOD landscape mesh, enabling the creation of a Nanite scene proxy. Non nanite HLODs will still use a ULandscapeMeshProxyComponent, which will ensure the landscape proxies surrounding the HLOD tiles blends properly to avoid any visible gaps. In the Nanite path, skirts will prevent those gaps from appearing.
* When creating a Nanite mesh, use the same settings as the landscape itself (skirt depth, position precision, max edge length factor)
* For flexibility, keep relying on the LOD configured for HLODs (either a fixed LOD, the last LOD or an automatic LOD computation based on the HLOD expected draw distance). While users could easily get away with LOD0 since it's Nanite, this allows more flexibility and allows them to reduce the on disk size of those HLODs.

diff --git a/Engine/Source/Runtime/Landscape/Private/LandscapeHLODBuilder.cpp b/Engine/Source/Runtime/Landscape/Private/LandscapeHLODBuilder.cpp
--- a/Engine/Source/Runtime/Landscape/Private/LandscapeHLODBuilder.cpp
+++ b/Engine/Source/Runtime/Landscape/Private/LandscapeHLODBuilder.cpp
@@ -119,6 +119,26 @@ uint32 ULandscapeHLODBuilder::ComputeHLODHash(const UActorComponent* InSourceCom
 		bool bNaniteEnabled = LSProxy->IsNaniteEnabled();
 		Ar << bNaniteEnabled;
 		UE_LOG(LogHLODBuilder, VeryVerbose, TEXT("     - NaniteEnabled = %d"), bNaniteEnabled);
+		if (bNaniteEnabled)
+		{
+			int32 NaniteLODIndex = LSProxy->GetNaniteLODIndex();
+			int32 NanitePositionPrecision = LSProxy->GetNanitePositionPrecision();
+			float NaniteMaxEdgeLengthFactor = LSProxy->GetNaniteMaxEdgeLengthFactor();
+			Ar << NaniteLODIndex << NanitePositionPrecision << NaniteMaxEdgeLengthFactor;
+			UE_LOG(LogHLODBuilder, VeryVerbose, TEXT("     - NaniteLODIndex = %d"), NaniteLODIndex);
+			UE_LOG(LogHLODBuilder, VeryVerbose, TEXT("     - NanitePositionPrecision = %d"), NanitePositionPrecision);
+			UE_LOG(LogHLODBuilder, VeryVerbose, TEXT("     - NaniteMaxEdgeLengthFactor = %f"), NaniteMaxEdgeLengthFactor);
+
+			bool bNaniteSkirtEnabled = LSProxy->IsNaniteSkirtEnabled();
+			Ar << bNaniteSkirtEnabled;
+			UE_LOG(LogHLODBuilder, VeryVerbose, TEXT("     - NaniteSkirtEnabled = %d"), bNaniteSkirtEnabled);
+			if (bNaniteSkirtEnabled)
+			{
+				float NaniteSkirtDepth = LSProxy->GetNaniteSkirtDepth();
+				Ar << NaniteSkirtDepth;
+				UE_LOG(LogHLODBuilder, VeryVerbose, TEXT("     - NaniteSkirtDepth = %f"), NaniteSkirtDepth);
+			}
+		}
 
 		// HLODTextureSize
 		ELandscapeHLODTextureSizePolicy HLODTextureSizePolicy = LSProxy->HLODTextureSizePolicy;
@@ -268,7 +288,7 @@ static int32 ComputeRequiredTextureSize(const ALandscapeProxy* InLandscapeProxy,
 }
 
 
-static UMaterialInterface* BakeLandscapeMaterial(const FHLODBuildContext& InHLODBuildContext, const FMeshDescription& InMeshDescription, const ALandscapeProxy* InLandscapeProxy, const int32 InLandscapeLOD, int32 InTextureSize)
+static UMaterialInterface* BakeLandscapeMaterial(const FHLODBuildContext& InHLODBuildContext, const FMeshDescription& InMeshDescription, const ALandscapeProxy* InLandscapeProxy, int32 InTextureSize)
 {
 	// Build landscape material
 	FFlattenMaterial LandscapeFlattenMaterial;
@@ -320,10 +340,8 @@ static UMaterialInterface* BakeLandscapeMaterial(const FHLODBuildContext& InHLOD
 	return LandscapeMaterialInstance;
 }
 
-// This is an initial implementation for UE 5.0
 // Multiple improvements could be done
 // * Currently, for each referenced landscape proxy, we generate individual HLOD meshes & textures. This should output a single mesh for all proxies
-// * Generated mesh could be Nanite-enabled
 TArray<UActorComponent*> ULandscapeHLODBuilder::Build(const FHLODBuildContext& InHLODBuildContext, const TArray<UActorComponent*>& InSourceComponents) const
 {
 	TArray<ULandscapeComponent*> SourceLandscapeComponents = FilterComponents<ULandscapeComponent>(InSourceComponents);
@@ -352,7 +370,7 @@ TArray<UActorComponent*> ULandscapeHLODBuilder::Build(const FHLODBuildContext& I
 		FMeshDescription* MeshDescription = nullptr;
 
 		// Compute source landscape LOD
-		int32 LandscapeLOD = ComputeRequiredLandscapeLOD(LandscapeProxy, static_cast<float>(InHLODBuildContext.MinVisibleDistance));
+		const int32 LandscapeLOD = ComputeRequiredLandscapeLOD(LandscapeProxy, static_cast<float>(InHLODBuildContext.MinVisibleDistance));
 
 		// Mesh
 		{
@@ -368,12 +386,16 @@ TArray<UActorComponent*> ULandscapeHLODBuilder::Build(const FHLODBuildContext& I
 	
 			ALandscapeProxy::FRawMeshExportParams ExportParams;
 			ExportParams.ExportLOD = LandscapeLOD;
+		    ExportParams.SkirtDepth = LandscapeProxy->IsNaniteEnabled() ? LandscapeProxy->GetNaniteSkirtDepth() : TOptional<float>();
 			LandscapeProxy->ExportToRawMesh(ExportParams, *MeshDescription);
 
 			StaticMesh->CommitMeshDescription(0);
 
 			// Nanite settings
+		    const FVector3d Scale = LandscapeProxy->GetTransform().GetScale3D();
 			StaticMesh->NaniteSettings.bEnabled = LandscapeProxy->IsNaniteEnabled();
+		    StaticMesh->NaniteSettings.PositionPrecision = FMath::Log2(Scale.GetAbsMax()) + LandscapeProxy->GetNanitePositionPrecision();
+		    StaticMesh->NaniteSettings.MaxEdgeLengthFactor = LandscapeProxy->GetNaniteMaxEdgeLengthFactor();
 
 			StaticMesh->ImportVersion = EImportStaticMeshVersion::LastVersion;
 		}
@@ -381,7 +403,7 @@ TArray<UActorComponent*> ULandscapeHLODBuilder::Build(const FHLODBuildContext& I
 		// Material
 		{
 			int32 TextureSize = ComputeRequiredTextureSize(LandscapeProxy, static_cast<float>(InHLODBuildContext.MinVisibleDistance), MeshDescription);
-			UMaterialInterface* LandscapeMaterial = BakeLandscapeMaterial(InHLODBuildContext, *MeshDescription, LandscapeProxy, LandscapeLOD, TextureSize);
+			UMaterialInterface* LandscapeMaterial = BakeLandscapeMaterial(InHLODBuildContext, *MeshDescription, LandscapeProxy, TextureSize);
 
 			//Assign the proxy material to the static mesh
 			StaticMesh->GetStaticMaterials().Add(FStaticMaterial(LandscapeMaterial));
@@ -389,10 +411,22 @@ TArray<UActorComponent*> ULandscapeHLODBuilder::Build(const FHLODBuildContext& I
 
 		StaticMeshes.Add(StaticMesh);
 
-		ULandscapeMeshProxyComponent* LandcapeMeshProxyComponent = NewObject<ULandscapeMeshProxyComponent>();
-		LandcapeMeshProxyComponent->InitializeForLandscape(LandscapeProxy, static_cast<int8>(LandscapeLOD));
-		LandcapeMeshProxyComponent->SetStaticMesh(StaticMesh);
-		HLODComponents.Add(LandcapeMeshProxyComponent);
+		// In case we are dealing with a Nanite LS, simply create a static mesh component
+		if (LandscapeProxy->IsNaniteEnabled())
+		{
+			UStaticMeshComponent* StaticMeshComponent = NewObject<UStaticMeshComponent>();
+			StaticMeshComponent->SetStaticMesh(StaticMesh);
+			HLODComponents.Add(StaticMeshComponent);
+		}
+		// Otherwise, we use a ULandscapeMeshProxyComponent, which will ensure the landscape proxies surrounding 
+		// the HLOD tiles blends properly to avoid any visible gaps.
+		else
+		{
+			ULandscapeMeshProxyComponent* LandcapeMeshProxyComponent = NewObject<ULandscapeMeshProxyComponent>();
+			LandcapeMeshProxyComponent->InitializeForLandscape(LandscapeProxy, static_cast<int8>(LandscapeLOD));
+			LandcapeMeshProxyComponent->SetStaticMesh(StaticMesh);
+			HLODComponents.Add(LandcapeMeshProxyComponent);
+		}
 	}
 
 	UStaticMesh::BatchBuild(StaticMeshes);
