Fix for landscape crash when using more than one section per component, and a collision mip level that is greater than zero

Instead of rendering physical material IDs at full resolution and downsampling afterwards, we render at the desired collision resolution, and only downsample for simple collision if needed.
Also fixed the downsample to not require an exact integer ratio downsample.

diff --git a/Engine/Source/Runtime/Landscape/Private/LandscapeEdit.cpp b/Engine/Source/Runtime/Landscape/Private/LandscapeEdit.cpp
--- a/Engine/Source/Runtime/Landscape/Private/LandscapeEdit.cpp
+++ b/Engine/Source/Runtime/Landscape/Private/LandscapeEdit.cpp
@@ -1958,21 +1958,19 @@ uint32 ULandscapeComponent::CalculatePhysicalMaterialTaskHash() const
 	
 	// Take into account any material changes.
 	UMaterialInterface* Material = GetLandscapeMaterial();
-	for (UMaterialInstanceConstant* MIC = Cast<UMaterialInstanceConstant>(Material); MIC; MIC = Cast<UMaterialInstanceConstant>(Material))
-	{
-		Hash = FCrc::TypeCrc32(MIC->ParameterStateId, Hash);
-		Material = MIC->Parent;
-	}
-	UMaterial* MaterialBase = Cast<UMaterial>(Material);
-	if (MaterialBase != nullptr)
-	{
-		Hash = FCrc::TypeCrc32(MaterialBase->StateId, Hash);
-	}
+	uint32 MaterialAllStateCRC = Material->ComputeAllStateCRC();
+
+	// also take into account the collision mip level
+	Hash = FCrc::TypeCrc32(CollisionMipLevel, Hash);
+
+	int32 SimpleCollisionMipLevelForHash = (SimpleCollisionMipLevel > CollisionMipLevel) ? SimpleCollisionMipLevel : 0;
+	Hash = FCrc::TypeCrc32(SimpleCollisionMipLevelForHash, Hash);
 
 	// We could take into account heightmap and weightmap changes here by adding to the hash.
 	// Instead we are resetting the stored hash in UpdateCollisionHeightData() and UpdateCollisionLayerData().
+	// (this is not ideal, as if we forget to reset the stored hash somewhere it will not recalculate collision...)
 
-	return Hash;
+	return MaterialAllStateCRC;
 }
 
 bool ULandscapeComponent::GetRenderPhysicalMaterials(TArray<UPhysicalMaterial*>& OutPhysicalMaterials) const 
@@ -2113,34 +2111,42 @@ void ULandscapeComponent::UpdateCollisionPhysicalMaterialData(TArray<UPhysicalMa
 	CollisionComponent->PhysicalMaterialRenderObjects = InPhysicalMaterials;
 
 	// Copy the physical material IDs for both the full and (optional) simple collision.
-	const int32 SizeVerts = SubsectionSizeQuads * NumSubsections + 1;
-	check(InMaterialIds.Num() == SizeVerts * SizeVerts);
-	const int32 FullCollisionSizeVerts = CollisionComponent->CollisionSizeQuads + 1;
+	const int32 CollisionVerts = CollisionComponent->CollisionSizeQuads + 1;
+
+	// the material ids passed in should already be at the desired collision resolution
+	check(InMaterialIds.Num() == CollisionVerts * CollisionVerts);
+
+	// if simple collision is used, this is how many verts it needs
 	const int32 SimpleCollisionSizeVerts = CollisionComponent->SimpleCollisionSizeQuads > 0 ? CollisionComponent->SimpleCollisionSizeQuads + 1 : 0;
-	const int32 BulkDataSize = FullCollisionSizeVerts * FullCollisionSizeVerts + SimpleCollisionSizeVerts * SimpleCollisionSizeVerts;
+	const int32 BulkDataSize = CollisionVerts * CollisionVerts + SimpleCollisionSizeVerts * SimpleCollisionSizeVerts;
 
 	void* Data = CollisionComponent->PhysicalMaterialRenderData.Lock(LOCK_READ_WRITE);
 	Data = CollisionComponent->PhysicalMaterialRenderData.Realloc(BulkDataSize);
 	uint8* WritePtr = (uint8*)Data;
 
-	const int32 CollisionSizes[2] = { FullCollisionSizeVerts, SimpleCollisionSizeVerts };
+	const int32 CollisionSizes[2] = { CollisionVerts, SimpleCollisionSizeVerts };
 	for (int32 i = 0; i < 2; ++i)
 	{
 		const int32 CollisionSizeVerts = CollisionSizes[i];
-		if (CollisionSizeVerts == SizeVerts)
+		if (CollisionSizeVerts == CollisionVerts)
 		{
-			FMemory::Memcpy(WritePtr, InMaterialIds.GetData(), SizeVerts * SizeVerts);
-			WritePtr += SizeVerts * SizeVerts;
+			FMemory::Memcpy(WritePtr, InMaterialIds.GetData(), CollisionVerts * CollisionVerts);
+			WritePtr += CollisionVerts * CollisionVerts;
 		}
 		else if (CollisionSizeVerts > 0)
 		{
-			const int32 StepSize = SizeVerts / CollisionSizeVerts;
-			check(CollisionSizeVerts * StepSize == SizeVerts);
-			for (int32 y = 0; y < SizeVerts; y += StepSize)
+			const float ScaleFactor = (CollisionVerts - 1.0f) / (CollisionSizeVerts - 1.0f);
+
+			// For simple collision, we just do a point-sampled resample to the desired resolution
+			// This may result in off-by-half-a-pixel misalignments, but we assume the simple collision would be fine with that
+			// (the alternative would be to directly render the MaterialIDs at both regular and simple resolutions)
+			for (int32 y = 0; y < CollisionSizeVerts; y++)					// [0 .. CollisionSizeVerts-1]
 			{
-				for (int32 x = 0; x < SizeVerts; x += StepSize)
+				int32 sampleY = FMath::RoundToInt32(y * ScaleFactor);		// [0 .. CollisionVerts-1]
+				for (int32 x = 0; x < CollisionSizeVerts; x++)				// [0 .. CollisionSizeVerts-1]
 				{
-					*WritePtr++ = InMaterialIds[y * SizeVerts + x];
+					int32 sampleX = FMath::RoundToInt32(x * ScaleFactor);	// [0 .. CollisionVerts-1]
+					*WritePtr++ = InMaterialIds[sampleY * CollisionVerts + sampleX];
 				}
 			}
 		}
diff --git a/Engine/Source/Runtime/Landscape/Private/LandscapePhysicalMaterial.cpp b/Engine/Source/Runtime/Landscape/Private/LandscapePhysicalMaterial.cpp
--- a/Engine/Source/Runtime/Landscape/Private/LandscapePhysicalMaterial.cpp
+++ b/Engine/Source/Runtime/Landscape/Private/LandscapePhysicalMaterial.cpp
@@ -342,8 +342,13 @@ namespace
 				Task.CompletionState = ECompletionState::None;
 
 				const FTransform& ComponentTransform = InLandscapeComponent->GetComponentTransform();
-				const int32 ComponentSizeVerts = InLandscapeComponent->SubsectionSizeQuads * InLandscapeComponent->NumSubsections + 1;
-				const FIntPoint TargetSize(ComponentSizeVerts, ComponentSizeVerts);
+
+				// request physical material render at the desired collision resolution
+				const int32 CollisionSubsectionVerts = ((InLandscapeComponent->SubsectionSizeQuads + 1) >> InLandscapeComponent->CollisionMipLevel);
+				const int32 CollisionSubsectionQuads = CollisionSubsectionVerts - 1;
+				const int32 CollisionComponentVerts = CollisionSubsectionQuads * InLandscapeComponent->NumSubsections + 1;
+
+				const FIntPoint TargetSize(CollisionComponentVerts, CollisionComponentVerts);
 				const FIntPoint TargetSizeMinusOne(TargetSize - FIntPoint(1, 1));
 				const FVector TargetCenter = ComponentTransform.TransformPosition(FVector(TargetSizeMinusOne, 0.f) * 0.5f);
 				const FVector TargetExtent = FVector(TargetSize, 0.0f) * ComponentTransform.GetScale3D() * 0.5f;
