#UE Do not synchronize a landscape component proxy's scale if the landscape has not yet been initialized since ComponentToWorld is not yet set in that circumstance

diff --git a/Engine/Source/Runtime/Landscape/Private/LandscapeEditLayers.cpp b/Engine/Source/Runtime/Landscape/Private/LandscapeEditLayers.cpp
--- a/Engine/Source/Runtime/Landscape/Private/LandscapeEditLayers.cpp
+++ b/Engine/Source/Runtime/Landscape/Private/LandscapeEditLayers.cpp
@@ -9028,7 +9028,7 @@ TArray<FName> ALandscapeProxy::SynchronizeUnmarkedSharedProperties(ALandscapePro
 	USceneComponent* OwnRootComponent = GetRootComponent();
 	USceneComponent* ProxyRootComponent = InLandscape->GetRootComponent();
 
-	if ((OwnRootComponent != nullptr) && (ProxyRootComponent != nullptr))
+	if ((OwnRootComponent != nullptr) && (ProxyRootComponent != nullptr) && ProxyRootComponent->HasBeenInitialized())
 	{
 		FVector ProxyScale3D = ProxyRootComponent->GetComponentToWorld().GetScale3D();
 
