Add null checks for the world in ULandscapeHeightfieldCollisionComponent::CookCollisionData

diff --git a/Engine/Source/Runtime/Landscape/Private/LandscapeCollision.cpp b/Engine/Source/Runtime/Landscape/Private/LandscapeCollision.cpp
--- a/Engine/Source/Runtime/Landscape/Private/LandscapeCollision.cpp
+++ b/Engine/Source/Runtime/Landscape/Private/LandscapeCollision.cpp
@@ -833,7 +833,7 @@ void ULandscapeHeightfieldCollisionComponent::CreateCollisionObject()
 		UWorld* World = GetWorld();
 
 #if WITH_EDITOR
-		const bool bNeedsEditorHeightField = !World->IsGameWorld() && !GetOutermost()->bIsCookedForEditor;
+		const bool bNeedsEditorHeightField = World && !World->IsGameWorld() && !GetOutermost()->bIsCookedForEditor;
 #endif // WITH_EDITOR
 		FHeightfieldGeometryRef* ExistingHeightfieldRef = nullptr;
 		bool bCheckDDC = true;
@@ -936,7 +936,7 @@ void ULandscapeHeightfieldCollisionComponent::CreateCollisionObject()
 
 			// Release cooked collison data
 			// In cooked builds created collision object will never be deleted while component is alive, so we don't need this data anymore
-			if(FPlatformProperties::RequiresCookedData() || World->IsGameWorld())
+			if(FPlatformProperties::RequiresCookedData() || (World && World->IsGameWorld()))
 			{
 				CookedCollisionData.Empty();
 				bCookedCollisionDataWasDeleted = true;
