Removed deprecated code in landscape and water (anything prior to, and including, 5.3)

diff --git a/Engine/Plugins/Experimental/LandscapePatch/Source/LandscapePatch/Private/LandscapePatchManager.cpp b/Engine/Plugins/Experimental/LandscapePatch/Source/LandscapePatch/Private/LandscapePatchManager.cpp
--- a/Engine/Plugins/Experimental/LandscapePatch/Source/LandscapePatch/Private/LandscapePatchManager.cpp
+++ b/Engine/Plugins/Experimental/LandscapePatch/Source/LandscapePatch/Private/LandscapePatchManager.cpp
@@ -353,13 +353,6 @@ void ALandscapePatchManager::MarkDirtyIfModifiedInConstructionScript()
 	}
 }
 
-PRAGMA_DISABLE_DEPRECATION_WARNINGS
-bool ALandscapePatchManager::IsAffectingWeightmapLayer(const FName& InLayerName) const
-{
-	return AffectsWeightmapLayer(InLayerName);
-}
-PRAGMA_ENABLE_DEPRECATION_WARNINGS
-
 bool ALandscapePatchManager::AffectsWeightmapLayer(const FName& InLayerName) const
 {
 	if (!CanAffectWeightmap())
diff --git a/Engine/Plugins/Experimental/LandscapePatch/Source/LandscapePatch/Private/LandscapeTexturePatch.cpp b/Engine/Plugins/Experimental/LandscapePatch/Source/LandscapePatch/Private/LandscapeTexturePatch.cpp
--- a/Engine/Plugins/Experimental/LandscapePatch/Source/LandscapePatch/Private/LandscapeTexturePatch.cpp
+++ b/Engine/Plugins/Experimental/LandscapePatch/Source/LandscapePatch/Private/LandscapeTexturePatch.cpp
@@ -828,14 +828,6 @@ FMatrix44f ULandscapeTexturePatch::GetPatchToHeightmapUVs(int32 PatchSizeX, int3
 	return (FMatrix44f)PatchToLandscapeUVTransposed.GetTransposed();
 }
 
-
-PRAGMA_DISABLE_DEPRECATION_WARNINGS
-bool ULandscapeTexturePatch::IsAffectingWeightmapLayer(const FName& InLayerName) const
-{
-	return AffectsWeightmapLayer(InLayerName);
-}
-PRAGMA_ENABLE_DEPRECATION_WARNINGS
-
 bool ULandscapeTexturePatch::AffectsWeightmapLayer(const FName& InLayerName) const
 {
 	if (!IsEnabled())
diff --git a/Engine/Plugins/Experimental/LandscapePatch/Source/LandscapePatch/Public/LandscapePatchComponent.h b/Engine/Plugins/Experimental/LandscapePatch/Source/LandscapePatch/Public/LandscapePatchComponent.h
--- a/Engine/Plugins/Experimental/LandscapePatch/Source/LandscapePatch/Public/LandscapePatchComponent.h
+++ b/Engine/Plugins/Experimental/LandscapePatch/Source/LandscapePatch/Public/LandscapePatchComponent.h
@@ -36,8 +36,6 @@ public:
 
 	virtual UTextureRenderTarget2D* RenderLayer_Native(const FLandscapeBrushParameters& InParameters) { return InParameters.CombinedResult; }
 
-	UE_DEPRECATED(5.3, "Use AffectsWeightmapLayer")
-	virtual bool IsAffectingWeightmapLayer(const FName& InLayerName) const { return false; }
 	virtual bool AffectsWeightmapLayer(const FName& InLayerName) const { return false; }
 	virtual bool AffectsVisibilityLayer() const { return false; }
 
diff --git a/Engine/Plugins/Experimental/LandscapePatch/Source/LandscapePatch/Public/LandscapePatchManager.h b/Engine/Plugins/Experimental/LandscapePatch/Source/LandscapePatch/Public/LandscapePatchManager.h
--- a/Engine/Plugins/Experimental/LandscapePatch/Source/LandscapePatch/Public/LandscapePatchManager.h
+++ b/Engine/Plugins/Experimental/LandscapePatch/Source/LandscapePatch/Public/LandscapePatchManager.h
@@ -90,8 +90,6 @@ public:
 	void MarkDirtyIfModifiedInConstructionScript();
 
 	// ALandscapeBlueprintBrushBase
-	UE_DEPRECATED(5.3, "Use AffectsWeightmapLayer")
-	virtual bool IsAffectingWeightmapLayer(const FName& InLayerName) const override;
 	virtual bool AffectsWeightmapLayer(const FName& InLayerName) const override;
 	virtual bool AffectsVisibilityLayer() const override;
 	virtual void SetOwningLandscape(class ALandscape* InOwningLandscape) override;
diff --git a/Engine/Plugins/Experimental/LandscapePatch/Source/LandscapePatch/Public/LandscapeTexturePatch.h b/Engine/Plugins/Experimental/LandscapePatch/Source/LandscapePatch/Public/LandscapeTexturePatch.h
--- a/Engine/Plugins/Experimental/LandscapePatch/Source/LandscapePatch/Public/LandscapeTexturePatch.h
+++ b/Engine/Plugins/Experimental/LandscapePatch/Source/LandscapePatch/Public/LandscapeTexturePatch.h
@@ -223,8 +223,6 @@ public:
 	UTextureRenderTarget2D* RenderLayer_Native(const FLandscapeBrushParameters& InParameters);
 
 	// ULandscapePatchComponent
-	UE_DEPRECATED(5.3, "Use AffectsWeightmapLayer")
-	virtual bool IsAffectingWeightmapLayer(const FName& InLayerName) const override;
 	virtual bool AffectsWeightmapLayer(const FName& InLayerName) const override;
 	virtual bool AffectsVisibilityLayer() const override;
 	virtual bool IsEnabled() const override;
diff --git a/Engine/Plugins/Experimental/Water/Source/Editor/Private/WaterBrushManager.h b/Engine/Plugins/Experimental/Water/Source/Editor/Private/WaterBrushManager.h
--- a/Engine/Plugins/Experimental/Water/Source/Editor/Private/WaterBrushManager.h
+++ b/Engine/Plugins/Experimental/Water/Source/Editor/Private/WaterBrushManager.h
@@ -185,11 +185,11 @@ public:
 	bool bNeedsForceUpdate = false;
 
 #if WITH_EDITORONLY_DATA
-	UE_DEPRECATED(5.1, "This material is now useless, the WaterVelocityTexture is now regenerated at runtime (WaterInfoTexture in AWaterZone).")
+	UE_DEPRECATED(all, "This material is now useless, the WaterVelocityTexture is now regenerated at runtime (WaterInfoTexture in AWaterZone).")
 	UPROPERTY()
 	TObjectPtr<UMaterialInterface> FinalizeVelocityHeightMaterial_DEPRECATED = nullptr;
 
-	UE_DEPRECATED(5.1, "This material instance is now useless, the WaterVelocityTexture is now regenerated at runtime (WaterInfoTexture in AWaterZone).")
+	UE_DEPRECATED(all, "This material instance is now useless, the WaterVelocityTexture is now regenerated at runtime (WaterInfoTexture in AWaterZone).")
 	UPROPERTY()
 	TObjectPtr<UMaterialInstanceDynamic> FinalizeVelocityHeightMID_DEPRECATED = nullptr;
 #endif // WITH_EDITORONLY_DATA
diff --git a/Engine/Plugins/Experimental/Water/Source/Editor/Private/WaterLandscapeBrush.cpp b/Engine/Plugins/Experimental/Water/Source/Editor/Private/WaterLandscapeBrush.cpp
--- a/Engine/Plugins/Experimental/Water/Source/Editor/Private/WaterLandscapeBrush.cpp
+++ b/Engine/Plugins/Experimental/Water/Source/Editor/Private/WaterLandscapeBrush.cpp
@@ -532,11 +532,6 @@ void AWaterLandscapeBrush::BlueprintOnRenderTargetTexturesUpdated_Implementation
 	// Deprecated
 }
 
-void AWaterLandscapeBrush::ForceWaterTextureUpdate()
-{
-	// Deprecated
-}
-
 #if WITH_EDITOR
 
 AWaterLandscapeBrush::EWaterBrushStatus AWaterLandscapeBrush::CheckWaterBrushStatus()
diff --git a/Engine/Plugins/Experimental/Water/Source/Editor/Private/WaterLandscapeBrush.h b/Engine/Plugins/Experimental/Water/Source/Editor/Private/WaterLandscapeBrush.h
--- a/Engine/Plugins/Experimental/Water/Source/Editor/Private/WaterLandscapeBrush.h
+++ b/Engine/Plugins/Experimental/Water/Source/Editor/Private/WaterLandscapeBrush.h
@@ -48,15 +48,15 @@ public:
 	void BlueprintWaterBodyChanged(AActor* Actor);
 	virtual void BlueprintWaterBodyChanged_Native(AActor* Actor) {}
 
-	UE_DEPRECATED(4.27, "Use SetActorCache instead")
+	UE_DEPRECATED(all, "Use SetActorCache instead")
 	UFUNCTION(BlueprintCallable, Category = "Cache", meta = (DeprecatedFunction, DeprecationMessage="Use SetActorCache instead"))
 	void SetWaterBodyCache(AWaterBody* WaterBody, UObject* InCache) {}
 	
-	UE_DEPRECATED(4.27, "Use GetActorCache instead")
+	UE_DEPRECATED(all, "Use GetActorCache instead")
 	UFUNCTION(BlueprintCallable, Category = "Cache", meta = (DeterminesOutputType = "CacheClass", DeprecatedFunction, DeprecationMessage = "Use GetActorCache instead"))
 	UObject* GetWaterBodyCache(AWaterBody* WaterBody, TSubclassOf<UObject> CacheClass) const { return nullptr; }
 
-	UE_DEPRECATED(4.27, "Use ClearActorCache instead")
+	UE_DEPRECATED(all, "Use ClearActorCache instead")
 	UFUNCTION(BlueprintCallable, Category = "Cache", meta = (DeprecatedFunction, DeprecationMessage = "Use ClearActorCache instead"))
 	void ClearWaterBodyCache(AWaterBody* WaterBody) {}
 
@@ -71,17 +71,17 @@ public:
 
 	UFUNCTION(BlueprintNativeEvent, Category = "Cache", meta = (CallInEditor = "true", DeprecatedFunction, DeprecationMessage = "This event isn't called anymore, the WaterVelocityTexture is now regenerated at runtime (WaterInfoTexture in AWaterZone)."))
 	void BlueprintGetRenderTargets(UTextureRenderTarget2D* InHeightRenderTarget, UTextureRenderTarget2D*& OutVelocityRenderTarget);
-	UE_DEPRECATED(5.1, "This function isn't called anymore, the WaterVelocityTexture is now regenerated at runtime (WaterInfoTexture in AWaterZone)")
+	UE_DEPRECATED(all, "This function isn't called anymore, the WaterVelocityTexture is now regenerated at runtime (WaterInfoTexture in AWaterZone)")
 	virtual void BlueprintGetRenderTargets_Native(UTextureRenderTarget2D* InHeightRenderTarget, UTextureRenderTarget2D*& OutVelocityRenderTarget) {}
 
 	UFUNCTION(BlueprintNativeEvent, Category = "Cache", meta = (CallInEditor = "true", DeprecatedFunction, DeprecationMessage = "This event isn't called anymore, the WaterVelocityTexture is now regenerated at runtime (WaterInfoTexture in AWaterZone)."))
 	void BlueprintOnRenderTargetTexturesUpdated(UTexture2D* VelocityTexture);
-	UE_DEPRECATED(5.1, "This function isn't called anymore, the WaterVelocityTexture is now regenerated at runtime (WaterInfoTexture in AWaterZone)")
+	UE_DEPRECATED(all, "This function isn't called anymore, the WaterVelocityTexture is now regenerated at runtime (WaterInfoTexture in AWaterZone)")
 	virtual void BlueprintOnRenderTargetTexturesUpdated_Native(UTexture2D* VelocityTexture) {}
 
-	UE_DEPRECATED(5.1, "This function is now useless, the WaterVelocityTexture is now regenerated at runtime (WaterInfoTexture in AWaterZone).")
+	UE_DEPRECATED(all, "This function is now useless, the WaterVelocityTexture is now regenerated at runtime (WaterInfoTexture in AWaterZone).")
 	UFUNCTION(BlueprintCallable, Category = "Cache", meta = (CallInEditor = "true", DeprecatedFunction, DeprecationMessage = "This function is now useless, the WaterVelocityTexture is now regenerated at runtime (WaterInfoTexture in AWaterZone)."))
-	void ForceWaterTextureUpdate();
+	void ForceWaterTextureUpdate() {}
 
 	void SetTargetLandscape(ALandscape* InTargetLandscape);
 		
diff --git a/Engine/Plugins/Experimental/Water/Source/Editor/Public/WaterEditorSettings.h b/Engine/Plugins/Experimental/Water/Source/Editor/Public/WaterEditorSettings.h
--- a/Engine/Plugins/Experimental/Water/Source/Editor/Public/WaterEditorSettings.h
+++ b/Engine/Plugins/Experimental/Water/Source/Editor/Public/WaterEditorSettings.h
@@ -217,11 +217,6 @@ public:
 	UMaterialInterface* GetDefaultCompositeWaterBodyTextureMaterial() const;
 	FSoftObjectPath GetDefaultCompositeWaterBodyTextureMaterialPath() const { return DefaultCompositeWaterBodyTextureMaterial.ToSoftObjectPath(); }
 
-	UE_DEPRECATED(5.1, "This material is not needed anymore : the WaterVelocityTexture is now regenerated at runtime (WaterInfoTexture in AWaterZone)")
-	UMaterialInterface* GetDefaultFinalizeVelocityHeightMaterial() const { return nullptr; }
-	UE_DEPRECATED(5.1, "This material is not needed anymore : the WaterVelocityTexture is now regenerated at runtime (WaterInfoTexture in AWaterZone)")
-	FSoftObjectPath GetDefaultFinalizeVelocityHeightMaterialPath() const { return FSoftObjectPath(); }
-
 	UMaterialInterface* GetDefaultJumpFloodStepMaterial() const;
 	FSoftObjectPath GetDefaultJumpFloodStepMaterialPath() const { return DefaultJumpFloodStepMaterial.ToSoftObjectPath(); }
 
diff --git a/Engine/Plugins/Experimental/Water/Source/Runtime/Private/WaterBodyComponent.cpp b/Engine/Plugins/Experimental/Water/Source/Runtime/Private/WaterBodyComponent.cpp
--- a/Engine/Plugins/Experimental/Water/Source/Runtime/Private/WaterBodyComponent.cpp
+++ b/Engine/Plugins/Experimental/Water/Source/Runtime/Private/WaterBodyComponent.cpp
@@ -1529,15 +1529,6 @@ void UWaterBodyComponent::UpdateSplineComponent()
 	}
 }
 
-void UWaterBodyComponent::OnWaterBodyChanged(bool bShapeOrPositionChanged, bool bWeightmapSettingsChanged, bool bUserTriggeredChange)
-{
-	FOnWaterBodyChangedParams Params;
-	Params.bShapeOrPositionChanged = bShapeOrPositionChanged;
-	Params.bWeightmapSettingsChanged = bWeightmapSettingsChanged;
-	Params.bUserTriggered = bUserTriggeredChange;
-	OnWaterBodyChanged(Params);
-}
-
 void UWaterBodyComponent::OnWaterBodyChanged(const FOnWaterBodyChangedParams& InParams)
 {
 	TRACE_CPUPROFILER_EVENT_SCOPE(UWaterBodyComponent::OnWaterBodyChanged)
diff --git a/Engine/Plugins/Experimental/Water/Source/Runtime/Public/NiagaraWaterFunctionLibrary.h b/Engine/Plugins/Experimental/Water/Source/Runtime/Public/NiagaraWaterFunctionLibrary.h
--- a/Engine/Plugins/Experimental/Water/Source/Runtime/Public/NiagaraWaterFunctionLibrary.h
+++ b/Engine/Plugins/Experimental/Water/Source/Runtime/Public/NiagaraWaterFunctionLibrary.h
@@ -15,7 +15,7 @@ class WATER_API UNiagaraWaterFunctionLibrary : public UBlueprintFunctionLibrary
 	GENERATED_BODY()
 
 	/** Sets the water body on the Niagra water data interface on a Niagara particle system*/
-	UE_DEPRECATED(4.27, "Use SetWaterBodyComponent instead.")
+	UE_DEPRECATED(all, "Use SetWaterBodyComponent instead")
 	UFUNCTION(BlueprintCallable, Category = Water, meta = (DisplayName = "Set Water Body"), meta = (DeprecatedFunction, DeprecationMessage="Use SetWaterBodyComponent instead"))
 	static void SetWaterBody(UNiagaraComponent* NiagaraSystem, const FString& OverrideName, AWaterBody* WaterBody) {}
 
diff --git a/Engine/Plugins/Experimental/Water/Source/Runtime/Public/WaterBodyActor.h b/Engine/Plugins/Experimental/Water/Source/Runtime/Public/WaterBodyActor.h
--- a/Engine/Plugins/Experimental/Water/Source/Runtime/Public/WaterBodyActor.h
+++ b/Engine/Plugins/Experimental/Water/Source/Runtime/Public/WaterBodyActor.h
@@ -173,56 +173,46 @@ protected:
 
 #pragma region Deprecated
 public:
-	UE_DEPRECATED(4.27, "Moved to WaterBodyComponent")
+	UE_DEPRECATED(all, "Moved to WaterBodyComponent")
 	UFUNCTION(BlueprintCallable, Category = Rendering, meta = (DeprecatedFunction))
 	virtual UMaterialInstanceDynamic* GetRiverToLakeTransitionMaterialInstance() final { return WaterBodyComponent->GetRiverToLakeTransitionMaterialInstance(); }
 
-	UE_DEPRECATED(4.27, "Moved to WaterBodyComponent")
+	UE_DEPRECATED(all, "Moved to WaterBodyComponent")
 	UFUNCTION(BlueprintCallable, Category = Rendering, meta = (DeprecatedFunction))
 	virtual UMaterialInstanceDynamic* GetRiverToOceanTransitionMaterialInstance() final { return WaterBodyComponent->GetRiverToOceanTransitionMaterialInstance(); }
 
-	UE_DEPRECATED(4.27, "Moved to WaterBodyComponent")
+	UE_DEPRECATED(all, "Moved to WaterBodyComponent")
 	UFUNCTION(BlueprintCallable, Category = Rendering, meta = (DeprecatedFunction))
 	void SetWaterMaterial(UMaterialInterface* InMaterial) { WaterBodyComponent->SetWaterMaterial(InMaterial); }
 
-	UE_DEPRECATED(4.27, "Moved to WaterBodyComponent")
+	UE_DEPRECATED(all, "Moved to WaterBodyComponent")
 	UFUNCTION(BlueprintCallable, Category = Rendering, meta = (DeprecatedFunction))
 	UMaterialInstanceDynamic* GetWaterMaterialInstance() { return WaterBodyComponent->GetWaterMaterialInstance(); }
 
-	UE_DEPRECATED(4.27, "Moved to WaterBodyComponent")
+	UE_DEPRECATED(all, "Moved to WaterBodyComponent")
 	UFUNCTION(BlueprintCallable, Category = WaterBody, meta = (DeprecatedFunction))
 	virtual float GetWaterVelocityAtSplineInputKey(float InKey) const { return WaterBodyComponent->GetWaterVelocityAtSplineInputKey(InKey); }
 
-	UE_DEPRECATED(4.27, "Moved to WaterBodyComponent")
+	UE_DEPRECATED(all, "Moved to WaterBodyComponent")
 	UFUNCTION(BlueprintCallable, Category = WaterBody, meta = (DeprecatedFunction))
 	virtual FVector GetWaterVelocityVectorAtSplineInputKey(float InKey) const { return WaterBodyComponent->GetWaterVelocityVectorAtSplineInputKey(InKey); }
 
-	UE_DEPRECATED(4.27, "Moved to WaterBodyComponent")
+	UE_DEPRECATED(all, "Moved to WaterBodyComponent")
 	UFUNCTION(BlueprintCallable, Category = WaterBody, meta = (DeprecatedFunction))
 	virtual float GetAudioIntensityAtSplineInputKey(float InKey) const { return WaterBodyComponent->GetAudioIntensityAtSplineInputKey(InKey); }
 
-	UE_DEPRECATED(4.27, "Moved to WaterBodyComponent")
+	UE_DEPRECATED(all, "Moved to WaterBodyComponent")
 	UFUNCTION(BlueprintCallable, Category = Water, meta = (DeprecatedFunction))
 	TArray<AWaterBodyIsland*> GetIslands() const { return WaterBodyComponent->GetIslands(); }
 
-	UE_DEPRECATED(4.27, "Moved to WaterBodyComponent")
+	UE_DEPRECATED(all, "Moved to WaterBodyComponent")
 	UFUNCTION(BlueprintCallable, Category = Water, meta = (DeprecatedFunction))
 	TArray<AWaterBodyExclusionVolume*> GetExclusionVolumes() const { return WaterBodyComponent->GetExclusionVolumes(); }
 
-	UE_DEPRECATED(4.27, "Moved to WaterBodyComponent")
+	UE_DEPRECATED(all, "Moved to WaterBodyComponent")
 	UFUNCTION(BlueprintCallable, Category=Water, meta = (DeprecatedFunction))
-	void OnWaterBodyChanged(bool bShapeOrPositionChanged, bool bWeightmapSettingsChanged = false)
-	{ 
-		FOnWaterBodyChangedParams Params;
-		Params.bShapeOrPositionChanged = bShapeOrPositionChanged;
-		Params.bWeightmapSettingsChanged = bWeightmapSettingsChanged;
-		return WaterBodyComponent->OnWaterBodyChanged(Params);
-	}
+	void OnWaterBodyChanged(bool bShapeOrPositionChanged, bool bWeightmapSettingsChanged = false) {}
 
-#if WITH_EDITOR
-	UE_DEPRECATED(5.1, "Moved to WaterBodyComponent")
-	virtual bool IsIconVisible() const { return WaterBodyComponent->IsIconVisible(); }
-#endif // WITH_EDITOR
 protected:
 	friend class UWaterBodyComponent;
 #if WITH_EDITORONLY_DATA
@@ -299,7 +289,7 @@ protected:
 	float ShapeDilation_DEPRECATED = 4096.0f;
 #endif // WITH_EDITORONLY_DATA
 
-#pragma endregion // deprecated
+#pragma endregion // Deprecated
 };
 
 #if UE_ENABLE_INCLUDE_ORDER_DEPRECATED_IN_5_2
diff --git a/Engine/Plugins/Experimental/Water/Source/Runtime/Public/WaterBodyComponent.h b/Engine/Plugins/Experimental/Water/Source/Runtime/Public/WaterBodyComponent.h
--- a/Engine/Plugins/Experimental/Water/Source/Runtime/Public/WaterBodyComponent.h
+++ b/Engine/Plugins/Experimental/Water/Source/Runtime/Public/WaterBodyComponent.h
@@ -247,7 +247,7 @@ public:
 	UMaterialInstanceDynamic* GetWaterStaticMeshMaterialInstance();
 
 	/** Returns water LOD MID */
-	UE_DEPRECATED(5.3, "GetWaterLODMaterialInstance has been renamed to GetWaterStaticMeshMaterialInstance.")
+	UE_DEPRECATED(all, "GetWaterLODMaterialInstance has been renamed to GetWaterStaticMeshMaterialInstance.")
 	UFUNCTION(BlueprintCallable, Category = Rendering, meta=(DeprecationMessage="GetWaterLODMaterialInstance has been renamed to GetWaterStaticMeshMaterialInstance"))
 	UMaterialInstanceDynamic* GetWaterLODMaterialInstance() { return GetWaterStaticMeshMaterialInstance(); };
 
@@ -420,21 +420,9 @@ public:
 
 	const FWaterBodyStaticMeshSettings& GetWaterBodyStaticMeshSettings() const { return StaticMeshSettings; }
 	
-	UE_DEPRECATED(5.1, "Renamed to CanEverAffectWaterMesh")
-	virtual bool CanAffectWaterMesh() const { return true; }
-
-	UE_DEPRECATED(5.1, "Use the version of this function taking FOnWaterBodyChangedParams in parameter")
-	void UpdateAll(bool bShapeOrPositionChanged) {}
-
-	UE_DEPRECATED(5.2, "Use version which takes FOnWaterBodyChangedParams")
+	UE_DEPRECATED(all, "Use version which takes FOnWaterBodyChangedParams")
 	UFUNCTION(BlueprintCallable, Category=Water, meta=(Deprecated = "5.2"))
-	void OnWaterBodyChanged(bool bShapeOrPositionChanged, bool bWeightmapSettingsChanged = false, bool bUserTriggeredChanged = false);
-
-	UE_DEPRECATED(5.3, "Renamed to GetWaterStaticMeshMaterial")
-	UMaterialInterface* GetWaterLODMaterial() const { return GetWaterStaticMeshMaterial(); }
-
-	UE_DEPRECATED(5.3, "Renamed to SetWaterStaticMeshMaterial")
-	void SetWaterLODMaterial(UMaterialInterface* InMaterial) { SetWaterStaticMeshMaterial(InMaterial); }
+	void OnWaterBodyChanged(bool bShapeOrPositionChanged, bool bWeightmapSettingsChanged = false, bool bUserTriggeredChanged = false) {}
 
 protected:
 	//~ Begin UActorComponent interface.
@@ -455,9 +443,6 @@ protected:
 	/** Returns whether the body support a height offset */
 	virtual bool IsHeightOffsetSupported() const;
 
-	UE_DEPRECATED(5.1, "Please use CanEverAffectNavigation() instead.")
-	virtual bool CanAffectNavigation() const { return false; }
-
 	/** Called every time UpdateAll is called on WaterBody (prior to UpdateWaterBody) */
 	virtual void BeginUpdateWaterBody();
 
@@ -498,8 +483,6 @@ protected:
 
 #if WITH_EDITOR
 	/** Called by UWaterBodyComponent::PostEditChangeProperty. */
-	UE_DEPRECATED(5.1, "Use the version of the function taking FOnWaterBodyChangedParams")
-	virtual void OnPostEditChangeProperty(FPropertyChangedEvent& PropertyChangedEvent, bool& bShapeOrPositionChanged, bool& bWeightmapSettingsChanged) {}
 	virtual void OnPostEditChangeProperty(FOnWaterBodyChangedParams& InOutOnWaterBodyChangedParams);
 
 	/** Validates this component's data */
@@ -546,8 +529,6 @@ protected:
 	virtual void PostEditUndo() override;
 	virtual void PostEditImport() override;
 	virtual void PostEditChangeProperty(FPropertyChangedEvent& PropertyChangedEvent) override;
-	UE_DEPRECATED(5.1, "Use OnWaterSplineDataChanged")
-	void OnSplineDataChanged() {}
 	void OnWaterSplineDataChanged(const FOnWaterSplineDataChangedParams& InParams);
 	void RegisterOnUpdateWavesData(UWaterWavesBase* InWaterWaves, bool bRegister);
 	void OnWavesDataUpdated(UWaterWavesBase* InWaterWaves, EPropertyChangeType::Type InChangeType);
@@ -562,8 +543,6 @@ protected:
 	virtual TSubclassOf<class UHLODBuilder> GetCustomHLODBuilderClass() const override;
 #endif // WITH_EDITOR
 
-	UE_DEPRECATED(5.2, "This function is no longer called. Instead uses UWaterSubsystem::FindWaterZone to determine the right water zone when no override is provided.")
-	AWaterZone* FindWaterZone() const { return nullptr; }
 public:
 	// INavRelevantInterface start
 	virtual void GetNavigationData(struct FNavigationRelevantData& Data) const override;
diff --git a/Engine/Plugins/Experimental/Water/Source/Runtime/Public/WaterBodyIslandActor.h b/Engine/Plugins/Experimental/Water/Source/Runtime/Public/WaterBodyIslandActor.h
--- a/Engine/Plugins/Experimental/Water/Source/Runtime/Public/WaterBodyIslandActor.h
+++ b/Engine/Plugins/Experimental/Water/Source/Runtime/Public/WaterBodyIslandActor.h
@@ -95,19 +95,11 @@ protected:
 	virtual void PostEditUndo() override;
 	virtual void PostEditChangeProperty(FPropertyChangedEvent& PropertyChangedEvent) override;
 	virtual void PostEditImport() override;
-	UE_DEPRECATED(5.1, "Use the version of this function taking FOnWaterBodyIslandChangedParams in parameter")
-	void UpdateAll() {}
 	void UpdateAll(const FOnWaterBodyIslandChangedParams& InParams);
 	
 	void OnWaterSplineDataChanged(const FOnWaterSplineDataChangedParams& InParams);
 	void OnWaterBodyIslandChanged(const FOnWaterBodyIslandChangedParams& InParams);
 
-	UE_DEPRECATED(5.1, "Use OnWaterSplineDataChanged")
-	void OnSplineDataChanged() {}
-
-	UE_DEPRECATED(5.1, "Use the version of this function taking a FOnWaterBodyIslandChangedParams in parameter")
-	void OnWaterBodyIslandChanged(bool bShapeOrPositionChanged, bool bWeightmapSettingsChanged) {}
-
 #endif
 
 protected:
diff --git a/Engine/Plugins/Experimental/Water/Source/Runtime/Public/WaterBodyOceanComponent.h b/Engine/Plugins/Experimental/Water/Source/Runtime/Public/WaterBodyOceanComponent.h
--- a/Engine/Plugins/Experimental/Water/Source/Runtime/Public/WaterBodyOceanComponent.h
+++ b/Engine/Plugins/Experimental/Water/Source/Runtime/Public/WaterBodyOceanComponent.h
@@ -34,10 +34,6 @@ public:
 	void FillWaterZoneWithOcean();
 #endif // WITH_EDITOR
 
-	UE_DEPRECATED(5.1, "Oceans no longer rely on the visual extent parameter making this obsolete. Instead they will be guaranteed to fill the entire water zone to which they belong.")
-	void SetVisualExtents(FVector2D) {}
-	UE_DEPRECATED(5.1, "Oceans no longer rely on the visual extent parameter making this obsolete. Instead they will be guaranteed to fill the entire water zone to which they belong.")
-	FVector2D GetVisualExtents() const { return FVector2D(); }
 protected:
 	/** UWaterBodyComponent Interface */
 	virtual bool IsBodyDynamic() const override { return true; }
diff --git a/Engine/Plugins/Experimental/Water/Source/Runtime/Public/WaterBrushActorInterface.h b/Engine/Plugins/Experimental/Water/Source/Runtime/Public/WaterBrushActorInterface.h
--- a/Engine/Plugins/Experimental/Water/Source/Runtime/Public/WaterBrushActorInterface.h
+++ b/Engine/Plugins/Experimental/Water/Source/Runtime/Public/WaterBrushActorInterface.h
@@ -44,9 +44,6 @@ class WATER_API IWaterBrushActorInterface
 	 */
 	virtual bool CanEverAffectWaterMesh() const = 0;
 
-	UE_DEPRECATED(5.1, "Renamed to CanEverAffectWaterMesh")
-	virtual bool CanAffectWaterMesh() const { return CanEverAffectWaterMesh(); };
-
 #if WITH_EDITOR
 	/** 
 	 * Returns the curve settings for this water actor
diff --git a/Engine/Plugins/Experimental/Water/Source/Runtime/Public/WaterSplineComponent.h b/Engine/Plugins/Experimental/Water/Source/Runtime/Public/WaterSplineComponent.h
--- a/Engine/Plugins/Experimental/Water/Source/Runtime/Public/WaterSplineComponent.h
+++ b/Engine/Plugins/Experimental/Water/Source/Runtime/Public/WaterSplineComponent.h
@@ -58,12 +58,6 @@ public:
 	void K2_SynchronizeAndBroadcastDataChange();
 
 #if WITH_EDITOR
-	DECLARE_EVENT(UWaterSplineComponent, UE_DEPRECATED(5.1, "Use FOnWaterSplineDataChanged") FOnSplineDataChanged);
-	UE_DEPRECATED(5.1, "Use OnWaterSplineDataChanged")
-PRAGMA_DISABLE_DEPRECATION_WARNINGS
-	FOnSplineDataChanged& OnSplineDataChanged() { static FOnSplineDataChanged DeprecatedEvent; return DeprecatedEvent; }
-PRAGMA_ENABLE_DEPRECATION_WARNINGS
-
 	FOnWaterSplineDataChanged& OnWaterSplineDataChanged() { return WaterSplineDataChangedEvent; }
 
 	virtual bool CanEditChange(const FProperty* InProperty) const override;
diff --git a/Engine/Plugins/Experimental/Water/Source/Runtime/Public/WaterSplineMetadata.h b/Engine/Plugins/Experimental/Water/Source/Runtime/Public/WaterSplineMetadata.h
--- a/Engine/Plugins/Experimental/Water/Source/Runtime/Public/WaterSplineMetadata.h
+++ b/Engine/Plugins/Experimental/Water/Source/Runtime/Public/WaterSplineMetadata.h
@@ -109,10 +109,6 @@ public:
 	DECLARE_MULTICAST_DELEGATE_OneParam(FOnWaterSplineMetadataChanged, const FOnWaterSplineMetadataChangedParams&);
 	FOnWaterSplineMetadataChanged OnChangeMetadata;
 	
-	DECLARE_MULTICAST_DELEGATE_TwoParams(FOnChangeData, UWaterSplineMetadata* /*WaterSplineMetadata*/, FPropertyChangedEvent& /*PropertyChangedEvent*/);
-	UE_DEPRECATED(5.2, "Use OnChangeMetadata")
-    FOnChangeData OnChangeData;
-    	
 #endif // WITH_EDITORONLY_DATA
 
 private:
diff --git a/Engine/Plugins/Experimental/Water/Source/Runtime/Public/WaterZoneActor.h b/Engine/Plugins/Experimental/Water/Source/Runtime/Public/WaterZoneActor.h
--- a/Engine/Plugins/Experimental/Water/Source/Runtime/Public/WaterZoneActor.h
+++ b/Engine/Plugins/Experimental/Water/Source/Runtime/Public/WaterZoneActor.h
@@ -97,18 +97,6 @@ public:
 	UFUNCTION(BlueprintCallable, Category=Rendering)
 	void SetFarMeshMaterial(UMaterialInterface* InFarMaterial);
 
-	UE_DEPRECATED(5.3, "Function renamed to GetDynamicWaterInfoCenter")
-	FVector GetTessellatedWaterMeshCenter() const { return GetDynamicWaterInfoCenter(); }
-
-	UE_DEPRECATED(5.3, "Function renamed to SetLocalTessellationCenter")
-	void SetTessellatedWaterMeshCenter(FVector NewCenter) { SetLocalTessellationCenter(NewCenter); }
-
-	UE_DEPRECATED(5.3, "Function renamed to GetDynamicWaterInfoExtent")
-	FVector GetTessellatedWaterMeshExtent() const { return GetDynamicWaterInfoExtent(); }
-
-	UE_DEPRECATED(5.3, "Function renamed to IsLocalTessellationEnabled.")
-	bool IsNonTessellatedLODMeshEnabled() const { return IsLocalOnlyTessellationEnabled(); }
-
 private:
 
 	/**
diff --git a/Engine/Source/Runtime/Landscape/Classes/Landscape.h b/Engine/Source/Runtime/Landscape/Classes/Landscape.h
--- a/Engine/Source/Runtime/Landscape/Classes/Landscape.h
+++ b/Engine/Source/Runtime/Landscape/Classes/Landscape.h
@@ -40,19 +40,6 @@ namespace EditLayersWeightmapLocalMerge_RenderThread
 extern LANDSCAPE_API TAutoConsoleVariable<int32> CVarLandscapeSplineFalloffModulation;
 #endif
 
-UENUM()
-enum UE_DEPRECATED(5.3, "ELandscapeSetupErrors is now unused and deprecated") ELandscapeSetupErrors : int
-{
-	LSE_None,
-	/** No Landscape Info available. */
-	LSE_NoLandscapeInfo,
-	/** There was already component with same X,Y. */
-	LSE_CollsionXY,
-	/** No Layer Info, need to add proper layers. */
-	LSE_NoLayerInfo,
-	LSE_MAX,
-};
-
 UENUM()
 enum class ERTDrawingType : uint8
 {
diff --git a/Engine/Source/Runtime/Landscape/Classes/LandscapeComponent.h b/Engine/Source/Runtime/Landscape/Classes/LandscapeComponent.h
--- a/Engine/Source/Runtime/Landscape/Classes/LandscapeComponent.h
+++ b/Engine/Source/Runtime/Landscape/Classes/LandscapeComponent.h
@@ -204,9 +204,11 @@ struct FLandscapeComponentGrassData
 
 	// Guid per material instance in the hierarchy between the assigned landscape material (instance) and the root UMaterial
 	// used to detect changes to material instance parameters or the root material that could affect the grass maps
+	UE_DEPRECATED(5.4, "GenerationHash is now used") 
 	TArray<FGuid, TInlineAllocator<2>> MaterialStateIds_DEPRECATED;
 	// cached component rotation when material world-position-offset is used,
 	// as this will affect the direction of world-position-offset deformation (included in the HeightData below)
+	UE_DEPRECATED(5.4, "GenerationHash is now used") 
 	FQuat RotationForWPO_DEPRECATED;
 
 	// Variable used to detect when grass data needs to be regenerated:
@@ -230,7 +232,6 @@ struct FLandscapeComponentGrassData
 	TArray<uint8> HeightWeightData;
 
 	FLandscapeComponentGrassData() = default;
-
 	FLandscapeComponentGrassData(ULandscapeComponent* Component);
 
 	// Returns whether grass data has been computed (or serialized) yet. Returns true even if the data is completely empty (e.g. all-zero weightmap data)
@@ -256,7 +257,7 @@ struct FLandscapeComponentGrassData
 };
 
 USTRUCT(NotBlueprintable, meta = (Deprecated = "5.1"))
-struct UE_DEPRECATED(5.1, "FLandscapeComponentMaterialOverride is deprecated; please use FLandscapePerLODMaterialOverride instead") FLandscapeComponentMaterialOverride
+struct UE_DEPRECATED(all, "FLandscapeComponentMaterialOverride is deprecated; please use FLandscapePerLODMaterialOverride instead") FLandscapeComponentMaterialOverride
 {
 	GENERATED_USTRUCT_BODY()
 
@@ -427,11 +428,12 @@ class ULandscapeComponent : public UPrimitiveComponent
 
 #if WITH_EDITORONLY_DATA
 	PRAGMA_DISABLE_DEPRECATION_WARNINGS
-	UE_DEPRECATED(5.1, "OverrideMaterials has been deprecated, use PerLODOverrideMaterials instead.")
+	UE_DEPRECATED(all, "OverrideMaterials has been deprecated, use PerLODOverrideMaterials instead.")
 	UPROPERTY()
 	TArray<FLandscapeComponentMaterialOverride> OverrideMaterials_DEPRECATED;
 	PRAGMA_ENABLE_DEPRECATION_WARNINGS
 
+	UE_DEPRECATED(5.4, "MaterialInstance has been deprecated, use MaterialInstances instead.")
 	UPROPERTY()
 	TObjectPtr<UMaterialInstanceConstant> MaterialInstance_DEPRECATED;
 #endif // WITH_EDITORONLY_DATA
@@ -478,6 +480,7 @@ class ULandscapeComponent : public UPrimitiveComponent
 	TArray<double> MipToMipMaxDeltas;
 
 #if WITH_EDITORONLY_DATA
+	UE_DEPRECATED(5.4, "CollisionComponent has been deprecated and will be removed in a future version")
 	UPROPERTY()
 	TLazyObjectPtr<ULandscapeHeightfieldCollisionComponent> CollisionComponent_DEPRECATED;
 #endif // !WITH_EDITORONLY_DATA
@@ -487,7 +490,6 @@ private:
 	UPROPERTY()
 	TObjectPtr<ULandscapeHeightfieldCollisionComponent> CollisionComponentRef;
 
-
 	/** Store  */ 
 	UPROPERTY(Transient)
 	bool bUserTriggeredChangeRequested = false;
@@ -596,21 +598,13 @@ public:
 	UPROPERTY(EditAnywhere, BlueprintReadOnly, Category=LandscapeComponent)
 	int32 LODBias;
 
+	UE_DEPRECATED(5.4, "StateId is unused and will be removed in a future version")
 	UPROPERTY()
-	// TODO [jonathan.bard] : remove unused : 
 	FGuid StateId;
 
-	UE_DEPRECATED(5.3, "BakedTextureMaterialGuid is officially deprecated now and nothing updates it anymore")
-	FGuid BakedTextureMaterialGuid;
-
-	UE_DEPRECATED(5.3, "LastBakedTextureMaterialGuid is officially deprecated now and nothing updates it anymore")
-	FGuid LastBakedTextureMaterialGuid;
-
 #if WITH_EDITORONLY_DATA
-	UE_DEPRECATED(5.3, "GIBakedBaseColorTexture is officially deprecated now and nothing updates it anymore")
-	TObjectPtr<UTexture2D> GIBakedBaseColorTexture;
-
 	/**	Legacy irrelevant lights */
+	UE_DEPRECATED(5.4, "IrrelevantLights is officially deprecated now and will be removed in a future version")
 	UPROPERTY()
 	TArray<FGuid> IrrelevantLights_DEPRECATED;
 
@@ -647,11 +641,16 @@ public:
 	/** Represents last saved hash for PhysicalMaterialTask */
 	UPROPERTY(Transient)
 	uint32 LastSavedPhysicalMaterialHash;
-#endif // WITH_EDITORONLY_DATA
 
+	UE_DEPRECATED(5.4, "MobileMaterialInterface has been deprecated and will be removed in a future version")
 	UPROPERTY(NonPIEDuplicateTransient)
 	TObjectPtr<UMaterialInterface> MobileMaterialInterface_DEPRECATED;
 
+    UE_DEPRECATED(5.4, "MobileCombinationMaterialInstance has been deprecated and will be removed in a future version")
+	UPROPERTY(NonPIEDuplicateTransient)
+	TObjectPtr<UMaterialInstanceConstant> MobileCombinationMaterialInstance_DEPRECATED;
+#endif // WITH_EDITORONLY_DATA
+
 	/** Material interfaces used for mobile */
 	UPROPERTY(NonPIEDuplicateTransient)
 	TArray<TObjectPtr<UMaterialInterface>> MobileMaterialInterfaces;
@@ -673,9 +672,6 @@ public:
 	  because we cannot generate it at runtime for standalone PIE games */
 	UPROPERTY(NonPIEDuplicateTransient)
 	TArray<TObjectPtr<UMaterialInstanceConstant>> MobileCombinationMaterialInstances;
-
-	UPROPERTY(NonPIEDuplicateTransient)
-	TObjectPtr<UMaterialInstanceConstant> MobileCombinationMaterialInstance_DEPRECATED;
 #endif // WITH_EDITORONLY_DATA
 
 public:
@@ -715,7 +711,6 @@ public:
 	virtual void Serialize(FArchive& Ar) override;
 	virtual void GetResourceSizeEx(FResourceSizeEx& CumulativeResourceSize) override;
 	virtual void BeginDestroy() override;
-	virtual void PostDuplicate(bool bDuplicateForPIE) override;
 	virtual void PostLoad() override;
 #if WITH_EDITORONLY_DATA
 	static void DeclareConstructClasses(TArray<FTopLevelAssetPath>& OutConstructClasses, const UClass* SpecificSubclass);
diff --git a/Engine/Source/Runtime/Landscape/Classes/LandscapeInfo.h b/Engine/Source/Runtime/Landscape/Classes/LandscapeInfo.h
--- a/Engine/Source/Runtime/Landscape/Classes/LandscapeInfo.h
+++ b/Engine/Source/Runtime/Landscape/Classes/LandscapeInfo.h
@@ -161,9 +161,6 @@ public:
 	TMap<FIntPoint, FLandscapeAddCollision> XYtoAddCollisionMap;
 #endif // WITH_EDITORONLY_DATA
 
-	UE_DEPRECATED(5.1, "This property has been deprecated, please use the StreamingProxies property instead")
-	TArray<TObjectPtr<ALandscapeStreamingProxy>> Proxies;
-
 	UPROPERTY()
 	TArray<TWeakObjectPtr<ALandscapeStreamingProxy>> StreamingProxies;
 
@@ -257,9 +254,6 @@ public:
 	 */
 	LANDSCAPE_API ALandscapeProxy* GetLandscapeProxyForLevel(ULevel* Level) const;
 
-	UE_DEPRECATED(5.3, "Use ULandscapeSubsystem::GetDirtyOnlyInMode() instead")
-	static bool IsDirtyOnlyInModeEnabled() { return false; }
-
 	LANDSCAPE_API bool GetDirtyOnlyInMode() const;
 
 	LANDSCAPE_API void OnModifiedPackageSaved(UPackage* InPackage);
@@ -328,9 +322,6 @@ public:
 
 	LANDSCAPE_API void RemoveXYOffsets();
 
-	UE_DEPRECATED(5.3, "Texture Baking is officially deprecated now and nothing updates it anymore")
-	void PostponeTextureBaking() {}
-
 	/** Will tell if the landscape actor can have some content related to the layer system */
 	LANDSCAPE_API bool CanHaveLayersContent() const;
 
@@ -376,17 +367,6 @@ public:
 	/** Called after creating object so that it can initialize its state */
 	void Initialize(UWorld* InWorld, const FGuid& InLandscapeGuid);
 
-	/**
-	 * Runs the given function on the root landscape actor and all streaming proxies
-	 * Most easily used with a lambda as follows:
-	 * ForAllLandscapeProxies([](ALandscapeProxy* Proxy)
-	 * {
-	 *     // Code
-	 * });
-	 */
-	UE_DEPRECATED(5.3, "This function has been deprecated, please use the ForEachLandscapeProxy property instead")
-	LANDSCAPE_API void ForAllLandscapeProxies(TFunctionRef<void(ALandscapeProxy*)> Fn) const;
-
 	/**
 	 * Runs the given function on the root landscape actor and all streaming proxies, with the posibility of early exit
 	 * Most easily used with a lambda as follows:
diff --git a/Engine/Source/Runtime/Landscape/Classes/LandscapeProxy.h b/Engine/Source/Runtime/Landscape/Classes/LandscapeProxy.h
--- a/Engine/Source/Runtime/Landscape/Classes/LandscapeProxy.h
+++ b/Engine/Source/Runtime/Landscape/Classes/LandscapeProxy.h
@@ -332,14 +332,14 @@ public:
 };
 
 USTRUCT(meta = (Deprecated = "5.1"))
-struct UE_DEPRECATED(5.1, "FLandscapeProxyMaterialOverride is deprecated; please use FLandscapePerLODMaterialOverride instead") FLandscapeProxyMaterialOverride
+struct UE_DEPRECATED(all, "FLandscapeProxyMaterialOverride is deprecated; please use FLandscapePerLODMaterialOverride instead") FLandscapeProxyMaterialOverride
 {
 	GENERATED_USTRUCT_BODY()
 
-	UPROPERTY(EditAnywhere, Category = Landscape, meta = (UIMin = 0, UIMax = 8, ClampMin = 0, ClampMax = 8))
+	UPROPERTY()
 	FPerPlatformInt LODIndex;
 
-	UPROPERTY(EditAnywhere, Category = Landscape)
+	UPROPERTY()
 	TObjectPtr<UMaterialInterface> Material = nullptr;
 };
 
@@ -451,14 +451,16 @@ public:
 	int32 MaxLODLevel;
 
 #if WITH_EDITORONLY_DATA
+	UE_DEPRECATED(5.4, "Unused property : will be removed in a future version")
 	UPROPERTY()
 	float LODDistanceFactor_DEPRECATED;
 
+	UE_DEPRECATED(5.4, "Unused property : will be removed in a future version")
 	UPROPERTY()
 	TEnumAsByte<ELandscapeLODFalloff::Type> LODFalloff_DEPRECATED;
 #endif // WITH_EDITORONLY_DATA
 
-	UE_DEPRECATED(5.4, "Unused property")
+	UE_DEPRECATED(5.4, "Unused property : will be removed in a future version")
 	UPROPERTY()
 	float ComponentScreenSizeToUseSubSections;
 
@@ -540,7 +542,7 @@ public:
 #if WITH_EDITORONLY_DATA
 
 	PRAGMA_DISABLE_DEPRECATION_WARNINGS
-	UE_DEPRECATED(5.1, "LandscapeComponentMaterialOverride has been deprecated, use PerLODOverrideMaterials instead.")
+	UE_DEPRECATED(all, "LandscapeComponentMaterialOverride has been deprecated, use PerLODOverrideMaterials instead.")
 	UPROPERTY()
 	TArray<FLandscapeProxyMaterialOverride> LandscapeMaterialsOverride_DEPRECATED;
 	PRAGMA_ENABLE_DEPRECATION_WARNINGS
@@ -553,7 +555,8 @@ public:
 
 	UPROPERTY(Transient)
 	bool bIsPerformingInteractiveActionOnLandscapeMaterialOverride;
-#endif 
+#endif // WITH_EDITORONLY_DATA
+
 	/**
 	 * Array of runtime virtual textures into which we draw this landscape.
 	 * The material also needs to be set up to output to a virtual texture.
@@ -624,7 +627,7 @@ public:
 	UPROPERTY(transient, duplicatetransient)
 	TArray<TObjectPtr<UHierarchicalInstancedStaticMeshComponent>> FoliageComponents;
 
-	UE_DEPRECATED(5.3, "NaniteComponent has been deprecated, use NaniteComponents instead.")
+	UE_DEPRECATED(all, "NaniteComponent has been deprecated, use NaniteComponents instead.")
 	UPROPERTY()
 	TObjectPtr<ULandscapeNaniteComponent> NaniteComponent_DEPRECATED;
 
@@ -790,6 +793,7 @@ public:
 	uint32 bBakeMaterialPositionOffsetIntoCollision:1;
 
 #if WITH_EDITORONLY_DATA
+	UE_DEPRECATED(5.4, "Unused property : will be removed in a future version")
 	UPROPERTY()
 	TArray<TObjectPtr<ULandscapeLayerInfoObject>> EditorCachedLayerInfos_DEPRECATED;
 
@@ -1117,19 +1121,10 @@ public:
 	UE_DEPRECATED(5.4, "This version of RenderGrassMaps is deprecated.  Use BuildGrassMaps() instead.")
 	void RenderGrassMaps(const TArray<ULandscapeComponent*>& InLandscapeComponents, const TArray<ULandscapeGrassType*>& InGrassTypes) {}
 
-	struct UE_DEPRECATED(5.3, "FGIBakedTextureState is officially deprecated now and nothing updates it anymore") FGIBakedTextureState
-	{
-		FGuid CombinedStateId;
-		TArray<ULandscapeComponent*> Components;
-	};
-		
 	/** Update the landscape physical material render tasks */
 	void UpdatePhysicalMaterialTasks(bool bInShouldMarkDirty = false);
 	void UpdatePhysicalMaterialTasksStatus(TSet<ULandscapeComponent*>* OutdatedComponents, int32* OutdatedComponentsCount) const;
 
-	UE_DEPRECATED(5.3, "UpdateBakedTexturesCountdown is officially deprecated now and nothing updates it anymore")
-	int32 UpdateBakedTexturesCountdown;
-
 	/** Editor notification when changing feature level */
 	void OnFeatureLevelChanged(ERHIFeatureLevel::Type NewFeatureLevel);
 
@@ -1172,10 +1167,6 @@ public:
 	// Retrieve the screen size at which each LOD should be rendered
 	LANDSCAPE_API TArray<float> GetLODScreenSizeArray() const;
 
-	// Copy properties from parent Landscape actor
-	UE_DEPRECATED(5.3, "GetSharedProperties is being deprecated, please use CopySharedProperties or SynchronizeSharedProperties instead.")
-	LANDSCAPE_API void GetSharedProperties(ALandscapeProxy* Landscape);
-
 #if WITH_EDITOR
 	// Copy properties from parent Landscape actor
 	LANDSCAPE_API void CopySharedProperties(ALandscapeProxy* InLandscape);
@@ -1222,10 +1213,6 @@ public:
 
 	LANDSCAPE_API int32 GetOutdatedGrassMapCount() const;
 	LANDSCAPE_API void BuildGrassMaps(struct FScopedSlowTask* InSlowTask = nullptr);
-	UE_DEPRECATED(5.3, "BuildGIBakedTextures is officially deprecated now")
-	void BuildGIBakedTextures(struct FScopedSlowTask* InSlowTask = nullptr) {}
-	UE_DEPRECATED(5.3, "GetOutdatedGIBakedTextureComponentsCount is officially deprecated now returns 0")
-	int32 GetOutdatedGIBakedTextureComponentsCount() const { return 0; }
 	LANDSCAPE_API void BuildPhysicalMaterial(struct FScopedSlowTask* InSlowTask = nullptr);
 	LANDSCAPE_API void InvalidatePhysicalMaterial();
 	LANDSCAPE_API int32 GetOudatedPhysicalMaterialComponentsCount() const;
@@ -1435,9 +1422,6 @@ public:
 	/** Creates a LandscapeWeightMapUsage object outered to this proxy. */
 	LANDSCAPE_API ULandscapeWeightmapUsage* CreateWeightmapUsage();
 
-	UE_DEPRECATED(5.3, "NumComponentsNeedingTextureBaking is officially deprecated now and nothing updates it anymore")
-	int32 NumComponentsNeedingTextureBaking;
-
 	/** remove an overlapping component. Called from MapCheck. */
 	LANDSCAPE_API void RemoveOverlappingComponent(ULandscapeComponent* Component);
 
@@ -1572,25 +1556,6 @@ private:
 
 
 #if WITH_EDITOR
-/**
- * Helper class used to Build or monitor Landscape GI Textures
- */
-class LANDSCAPE_API UE_DEPRECATED(5.3, "FLandscapeGIBakedTextureBuilder is officially deprecated now and nothing updates it anymore") FLandscapeGIBakedTextureBuilder
-{
-public:
-	FLandscapeGIBakedTextureBuilder(UWorld* InWorld)
-		:World(InWorld)
-	{}
-	void Build(){}
-	int32 GetOutdatedGIBakedTextureComponentsCount(bool bInForceUpdate = true) const { return 0; }
-
-private:
-	UWorld* World = nullptr;
-	mutable int32 OutdatedGIBakedTextureComponentsCount = 0;
-	mutable double GIBakedTexturesLastCheckTime = 0;
-};
-
-
 /**
  * Helper class used to Build or monitor Landscape Physical Material
  */
diff --git a/Engine/Source/Runtime/Landscape/Private/Landscape.cpp b/Engine/Source/Runtime/Landscape/Private/Landscape.cpp
--- a/Engine/Source/Runtime/Landscape/Private/Landscape.cpp
+++ b/Engine/Source/Runtime/Landscape/Private/Landscape.cpp
@@ -884,7 +884,9 @@ void ULandscapeComponent::Serialize(FArchive& Ar)
 		Ar << LegacyMapBuildData->ShadowMap;
 
 #if WITH_EDITORONLY_DATA
+PRAGMA_DISABLE_DEPRECATION_WARNINGS
 		LegacyMapBuildData->IrrelevantLights = IrrelevantLights_DEPRECATED;
+PRAGMA_ENABLE_DEPRECATION_WARNINGS
 #endif // WITH_EDITORONLY_DATA
 
 		FMeshMapBuildLegacyData LegacyComponentData;
@@ -895,6 +897,7 @@ void ULandscapeComponent::Serialize(FArchive& Ar)
 #if WITH_EDITORONLY_DATA
 	if (Ar.IsLoading() && Ar.CustomVer(FFortniteMainBranchObjectVersion::GUID) < FFortniteMainBranchObjectVersion::NewLandscapeMaterialPerLOD)
 	{
+PRAGMA_DISABLE_DEPRECATION_WARNINGS
 		if (MobileMaterialInterface_DEPRECATED != nullptr)
 		{
 			MobileMaterialInterfaces.AddUnique(MobileMaterialInterface_DEPRECATED);
@@ -904,6 +907,7 @@ void ULandscapeComponent::Serialize(FArchive& Ar)
 		{
 			MobileCombinationMaterialInstances.AddUnique(MobileCombinationMaterialInstance_DEPRECATED);
 		}
+PRAGMA_ENABLE_DEPRECATION_WARNINGS
 	}
 #endif // WITH_EDITORONLY_DATA
 
@@ -1269,6 +1273,8 @@ void ULandscapeComponent::PostLoad()
 
 #if WITH_EDITORONLY_DATA
 	// Handle old MaterialInstance
+
+PRAGMA_DISABLE_DEPRECATION_WARNINGS
 	if (MaterialInstance_DEPRECATED)
 	{
 		MaterialInstances.Empty(1);
@@ -1281,6 +1287,7 @@ void ULandscapeComponent::PostLoad()
 			UpdateMaterialInstances();
 		}
 	}
+PRAGMA_ENABLE_DEPRECATION_WARNINGS
 
 	if (CVarStripLayerTextureMipsOnLoad->GetBool())
 	{
@@ -1485,12 +1492,14 @@ void ULandscapeComponent::PostLoad()
 	GrassData->ConditionalDiscardDataOnLoad();
 
 #if WITH_EDITORONLY_DATA
+PRAGMA_DISABLE_DEPRECATION_WARNINGS
 	// If the Collision Component is not set yet and we're transferring the property from the lazy object pointer it was previously stored as to the soft object ptr it is now stored as :
 	if (!CollisionComponentRef && CollisionComponent_DEPRECATED.IsValid())
 	{
 		CollisionComponentRef = CollisionComponent_DEPRECATED.Get();
 		CollisionComponent_DEPRECATED = nullptr;
 	}
+PRAGMA_ENABLE_DEPRECATION_WARNINGS
 
 	// If mip-to-mip info is missing, recompute them (they were introduced later) :
 	if (MipToMipMaxDeltas.IsEmpty())
@@ -3697,6 +3706,7 @@ void ALandscapeProxy::Serialize(FArchive& Ar)
 #if WITH_EDITORONLY_DATA
 	if (Ar.IsLoading() && Ar.CustomVer(FLandscapeCustomVersion::GUID) < FLandscapeCustomVersion::MigrateOldPropertiesToNewRenderingProperties)
 	{
+PRAGMA_DISABLE_DEPRECATION_WARNINGS
 		if (LODDistanceFactor_DEPRECATED > 0)
 		{
 			const float LOD0LinearDistributionSettingMigrationTable[11] = { 1.75f, 1.75f, 1.75f, 1.75f, 1.75f, 1.68f, 1.55f, 1.4f, 1.25f, 1.25f, 1.25f };
@@ -3715,6 +3725,7 @@ void ALandscapeProxy::Serialize(FArchive& Ar)
 				LODDistributionSetting = LODDSquareRootDistributionSettingMigrationTable[FMath::RoundToInt(LODDistanceFactor_DEPRECATED)];
 			}
 		}
+PRAGMA_ENABLE_DEPRECATION_WARNINGS
 	}
 
 #endif
@@ -4130,7 +4141,7 @@ void ALandscapeProxy::PostLoad()
 	}
 
 #if WITH_EDITOR
-	PRAGMA_DISABLE_DEPRECATION_WARNINGS;
+PRAGMA_DISABLE_DEPRECATION_WARNINGS
 	if (!LandscapeMaterialsOverride_DEPRECATED.IsEmpty())
 	{
 		PerLODOverrideMaterials.Reserve(LandscapeMaterialsOverride_DEPRECATED.Num());
@@ -4140,8 +4151,7 @@ void ALandscapeProxy::PostLoad()
 		}
 		LandscapeMaterialsOverride_DEPRECATED.Reset();
 	}
-	
-	PRAGMA_ENABLE_DEPRECATION_WARNINGS;
+PRAGMA_ENABLE_DEPRECATION_WARNINGS
 
 	if (GIsEditor)
 	{
@@ -4175,6 +4185,7 @@ void ALandscapeProxy::PostLoad()
 
 	EditorLayerSettings.RemoveAll([](const FLandscapeEditorLayerSettings& Settings) { return Settings.LayerInfoObj == nullptr; });
 
+PRAGMA_DISABLE_DEPRECATION_WARNINGS
 	if (EditorCachedLayerInfos_DEPRECATED.Num() > 0)
 	{
 		for (int32 i = 0; i < EditorCachedLayerInfos_DEPRECATED.Num(); i++)
@@ -4183,6 +4194,7 @@ void ALandscapeProxy::PostLoad()
 		}
 		EditorCachedLayerInfos_DEPRECATED.Empty();
 	}
+PRAGMA_ENABLE_DEPRECATION_WARNINGS
 
 	bool bFixedUpInvalidMaterialInstances = false;
 
@@ -4381,98 +4393,6 @@ void ALandscapeProxy::Destroyed()
 		FeatureLevelChangedDelegateHandle.Reset();
 	}
 }
-#endif // WITH_EDITOR
-
-void ALandscapeProxy::GetSharedProperties(ALandscapeProxy* Landscape)
-{
-	if (Landscape)
-	{
-		Modify();
-
-		LandscapeGuid = Landscape->LandscapeGuid;
-		OriginalLandscapeGuid = Landscape->OriginalLandscapeGuid;
-
-		//@todo UE merge, landscape, this needs work
-		RootComponent->SetRelativeScale3D(Landscape->GetRootComponent()->GetComponentToWorld().GetScale3D());
-
-		StaticLightingResolution = Landscape->StaticLightingResolution;
-		CastShadow = Landscape->CastShadow;
-		bCastDynamicShadow = Landscape->bCastDynamicShadow;
-		bCastStaticShadow = Landscape->bCastStaticShadow;
-		bCastContactShadow = Landscape->bCastContactShadow;
-		bCastFarShadow = Landscape->bCastFarShadow;
-		bCastHiddenShadow = Landscape->bCastHiddenShadow;
-		bCastShadowAsTwoSided = Landscape->bCastShadowAsTwoSided;
-		bAffectDistanceFieldLighting = Landscape->bAffectDistanceFieldLighting;
-		LightingChannels = Landscape->LightingChannels;
-		bRenderCustomDepth = Landscape->bRenderCustomDepth;
-		CustomDepthStencilWriteMask = Landscape->CustomDepthStencilWriteMask;
-		CustomDepthStencilValue = Landscape->CustomDepthStencilValue;
-		LDMaxDrawDistance = Landscape->LDMaxDrawDistance;
-		ComponentSizeQuads = Landscape->ComponentSizeQuads;
-		NumSubsections = Landscape->NumSubsections;
-		SubsectionSizeQuads = Landscape->SubsectionSizeQuads;
-		MaxLODLevel = Landscape->MaxLODLevel;
-		ScalableLODDistributionSetting = Landscape->ScalableLODDistributionSetting;
-		ScalableLOD0DistributionSetting = Landscape->ScalableLOD0DistributionSetting;
-		ScalableLOD0ScreenSize = Landscape->ScalableLOD0ScreenSize;
-		LODDistributionSetting = Landscape->LODDistributionSetting;
-		LOD0DistributionSetting = Landscape->LOD0DistributionSetting;
-		LOD0ScreenSize = Landscape->LOD0ScreenSize;
-		LODGroupKey = Landscape->LODGroupKey;
-		NegativeZBoundsExtension = Landscape->NegativeZBoundsExtension;
-		PositiveZBoundsExtension = Landscape->PositiveZBoundsExtension;
-		CollisionMipLevel = Landscape->CollisionMipLevel;
-		bBakeMaterialPositionOffsetIntoCollision = Landscape->bBakeMaterialPositionOffsetIntoCollision;
-		RuntimeVirtualTextures = Landscape->RuntimeVirtualTextures;
-		VirtualTextureLodBias = Landscape->VirtualTextureLodBias;
-		bVirtualTextureRenderWithQuad = Landscape->bVirtualTextureRenderWithQuad;
-		bVirtualTextureRenderWithQuadHQ = Landscape->bVirtualTextureRenderWithQuadHQ;
-		VirtualTextureNumLods = Landscape->VirtualTextureNumLods;
-		VirtualTextureRenderPassType = Landscape->VirtualTextureRenderPassType;
-		bEnableNanite = Landscape->bEnableNanite;
-		ShadowCacheInvalidationBehavior = Landscape->ShadowCacheInvalidationBehavior;
-		NonNaniteVirtualShadowMapConstantDepthBias = Landscape->NonNaniteVirtualShadowMapConstantDepthBias;
-		NonNaniteVirtualShadowMapInvalidationHeightErrorThreshold = Landscape->NonNaniteVirtualShadowMapInvalidationHeightErrorThreshold;
-		NonNaniteVirtualShadowMapInvalidationScreenSizeLimit = Landscape->NonNaniteVirtualShadowMapInvalidationScreenSizeLimit;
-
-		bUseCompressedHeightmapStorage = Landscape->bUseCompressedHeightmapStorage;
-#if WITH_EDITORONLY_DATA
-		bNaniteSkirtEnabled = Landscape->bNaniteSkirtEnabled;
-		NaniteSkirtDepth = Landscape->NaniteSkirtDepth;
-		NaniteLODIndex = Landscape->NaniteLODIndex;
-#endif // WITH_EDITORONLY_DATA
-
-		if (!LandscapeMaterial)
-		{
-			LandscapeMaterial = Landscape->LandscapeMaterial;
-			PerLODOverrideMaterials = Landscape->PerLODOverrideMaterials;
-		}
-		if (!LandscapeHoleMaterial)
-		{
-			LandscapeHoleMaterial = Landscape->LandscapeHoleMaterial;
-		}
-		if (!DefaultPhysMaterial)
-		{
-			DefaultPhysMaterial = Landscape->DefaultPhysMaterial;
-		}
-		LightmassSettings = Landscape->LightmassSettings;
-	}
-
-#if WITH_EDITOR
-	if (GIsEditor && Landscape)
-	{
-		LODDistanceFactor_DEPRECATED = Landscape->LODDistanceFactor_DEPRECATED;
-		LODFalloff_DEPRECATED = Landscape->LODFalloff_DEPRECATED;
-		if (LandscapeMaterial == Landscape->LandscapeMaterial)
-		{
-			EditorLayerSettings = Landscape->EditorLayerSettings;
-		}
-	}
-#endif // WITH_EDITOR
-}
-
-#if WITH_EDITOR
 
 namespace UE::Landscape::Private
 {
@@ -5411,22 +5331,6 @@ void ULandscapeInfo::Initialize(UWorld* InWorld, const FGuid& InLandscapeGuid)
 	LandscapeGuid = InLandscapeGuid;
 }
 
-void ULandscapeInfo::ForAllLandscapeProxies(TFunctionRef<void(ALandscapeProxy*)> Fn) const
-{
-	if (ALandscape* Landscape = LandscapeActor.Get())
-	{
-		Fn(Landscape);
-	}
-
-	for (TWeakObjectPtr<ALandscapeStreamingProxy> StreamingProxyPtr : StreamingProxies)
-	{
-		if (ALandscapeProxy* LandscapeProxy = StreamingProxyPtr.Get())
-		{
-			Fn(LandscapeProxy);
-		}
-	}
-}
-
 void ULandscapeInfo::ForEachLandscapeProxy(TFunctionRef<bool(ALandscapeProxy*)> Fn) const
 {
 	if (ALandscape* Landscape = LandscapeActor.Get())
@@ -6120,24 +6024,10 @@ void ULandscapeComponent::PostInitProperties()
 {
 	Super::PostInitProperties();
 
-	// Create a new guid in case this is a newly created component
-	// If not, this guid will be overwritten when serialized
-	FPlatformMisc::CreateGuid(StateId);
-
 	// Initialize MapBuildDataId to something unique, in case this is a new ULandscapeComponent
 	MapBuildDataId = FGuid::NewGuid();
 }
 
-void ULandscapeComponent::PostDuplicate(bool bDuplicateForPIE)
-{
-	if (!bDuplicateForPIE)
-	{
-		// Reset the StateId on duplication since it needs to be unique for each capture.
-		// PostDuplicate covers direct calls to StaticDuplicateObject, but not actor duplication (see PostEditImport)
-		FPlatformMisc::CreateGuid(StateId);
-	}
-}
-
 ULandscapeWeightmapUsage::ULandscapeWeightmapUsage(const FObjectInitializer& ObjectInitializer)
 	: Super(ObjectInitializer)
 {
diff --git a/Engine/Source/Runtime/Landscape/Private/LandscapeBlueprintBrushBase.cpp b/Engine/Source/Runtime/Landscape/Private/LandscapeBlueprintBrushBase.cpp
--- a/Engine/Source/Runtime/Landscape/Private/LandscapeBlueprintBrushBase.cpp
+++ b/Engine/Source/Runtime/Landscape/Private/LandscapeBlueprintBrushBase.cpp
@@ -55,9 +55,10 @@ ALandscapeBlueprintBrushBase::ALandscapeBlueprintBrushBase(const FObjectInitiali
 
 PRAGMA_DISABLE_DEPRECATION_WARNINGS
 
+// Deprecated
 UTextureRenderTarget2D* ALandscapeBlueprintBrushBase::Render_Implementation(bool InIsHeightmap, UTextureRenderTarget2D* InCombinedResult, const FName& InWeightmapLayerName)
 {
-	return Render_Native(InIsHeightmap, InCombinedResult, InWeightmapLayerName);
+	return nullptr;
 }
 
 UTextureRenderTarget2D* ALandscapeBlueprintBrushBase::RenderLayer_Implementation(const FLandscapeBrushParameters& InParameters)
@@ -167,29 +168,6 @@ void ALandscapeBlueprintBrushBase::SetIsVisible(bool bInIsVisible)
 #endif
 }
 
-
-PRAGMA_DISABLE_DEPRECATION_WARNINGS
-void ALandscapeBlueprintBrushBase::SetAffectsHeightmap(bool bAffectsHeightmap)
-{
-	SetCanAffectHeightmap(bAffectsHeightmap);
-}
-
-void ALandscapeBlueprintBrushBase::SetAffectsWeightmap(bool bAffectsWeightmap)
-{
-	SetCanAffectWeightmap(bAffectsWeightmap);
-}
-
-void ALandscapeBlueprintBrushBase::SetAffectsVisibilityLayer(bool bInAffectsVisibilityLayer)
-{
-	SetCanAffectVisibilityLayer(bInAffectsVisibilityLayer);
-}
-
-bool ALandscapeBlueprintBrushBase::IsAffectingWeightmapLayer(const FName& InLayerName) const
-{
-	return AffectsWeightmapLayer(InLayerName);
-}
-PRAGMA_ENABLE_DEPRECATION_WARNINGS
-
 void ALandscapeBlueprintBrushBase::SetCanAffectHeightmap(bool bInCanAffectHeightmap)
 {
 #if WITH_EDITORONLY_DATA
diff --git a/Engine/Source/Runtime/Landscape/Private/LandscapeEditInterface.cpp b/Engine/Source/Runtime/Landscape/Private/LandscapeEditInterface.cpp
--- a/Engine/Source/Runtime/Landscape/Private/LandscapeEditInterface.cpp
+++ b/Engine/Source/Runtime/Landscape/Private/LandscapeEditInterface.cpp
@@ -508,9 +508,6 @@ void FLandscapeEditDataInterface::SetHeightData(int32 X1, int32 Y1, int32 X2, in
 				// In Layer, cumulate dirty collision region (will be used next time UpdateCollisionHeightData is called)
 				Component->UpdateDirtyCollisionHeightData(FIntRect(ComponentX1, ComponentY1, ComponentX2, ComponentY2));
 			}
-			
-			// Update GUID for Platform Data
-			FPlatformMisc::CreateGuid(Component->StateId);
 		}
 	}
 
diff --git a/Engine/Source/Runtime/Landscape/Private/LandscapeGrass.cpp b/Engine/Source/Runtime/Landscape/Private/LandscapeGrass.cpp
--- a/Engine/Source/Runtime/Landscape/Private/LandscapeGrass.cpp
+++ b/Engine/Source/Runtime/Landscape/Private/LandscapeGrass.cpp
@@ -1558,13 +1558,14 @@ FArchive& operator<<(FArchive& Ar, FLandscapeComponentGrassData& Data)
 	{
 		if (Ar.CustomVer(FFortniteMainBranchObjectVersion::GUID) < FFortniteMainBranchObjectVersion::LandscapeSupportPerComponentGrassTypes)
 		{
+PRAGMA_DISABLE_DEPRECATION_WARNINGS
 			if (Ar.CustomVer(FLandscapeCustomVersion::GUID) >= FLandscapeCustomVersion::GrassMaterialInstanceFix)
 			{
-					Ar << Data.MaterialStateIds_DEPRECATED;
+				Ar << Data.MaterialStateIds_DEPRECATED;
 			}
 			else
 			{
-					Data.MaterialStateIds_DEPRECATED.Empty(1);
+				Data.MaterialStateIds_DEPRECATED.Empty(1);
 				if (Ar.UEVer() >= VER_UE4_SERIALIZE_LANDSCAPE_GRASS_DATA_MATERIAL_GUID)
 				{
 					FGuid MaterialStateId;
@@ -1577,6 +1578,7 @@ FArchive& operator<<(FArchive& Ar, FLandscapeComponentGrassData& Data)
 			{
 				Ar << Data.RotationForWPO_DEPRECATED;
 			}
+PRAGMA_ENABLE_DEPRECATION_WARNINGS
 		}
 		else
 		{
diff --git a/Engine/Source/Runtime/Landscape/Private/LandscapeSubsystem.cpp b/Engine/Source/Runtime/Landscape/Private/LandscapeSubsystem.cpp
--- a/Engine/Source/Runtime/Landscape/Private/LandscapeSubsystem.cpp
+++ b/Engine/Source/Runtime/Landscape/Private/LandscapeSubsystem.cpp
@@ -578,16 +578,6 @@ TArray<ALandscapeProxy*> ULandscapeSubsystem::GetOutdatedProxies(UE::Landscape::
 	return FinalProxiesToBuild;
 }
 
-int32 ULandscapeSubsystem::GetOutdatedGrassMapCount()
-{
-	return GrassMapsBuilder->GetOutdatedGrassMapCount(/*bInForceUpdate*/false);
-}
-
-int32 ULandscapeSubsystem::GetOudatedPhysicalMaterialComponentsCount()
-{
-	return PhysicalMaterialBuilder->GetOudatedPhysicalMaterialComponentsCount();
-}
-
 void ULandscapeSubsystem::BuildNanite(TArrayView<ALandscapeProxy*> InProxiesToBuild, bool bForceRebuild)
 {
 	TRACE_BOOKMARK(TEXT("ULandscapeSubsystem::BuildNanite"));
@@ -675,13 +665,6 @@ void ULandscapeSubsystem::BuildNanite(TArrayView<ALandscapeProxy*> InProxiesToBu
 	}
 }
 
-PRAGMA_DISABLE_DEPRECATION_WARNINGS
-bool ULandscapeSubsystem::IsDirtyOnlyInModeEnabled()
-{
-	return ULandscapeInfo::IsDirtyOnlyInModeEnabled();
-}
-PRAGMA_ENABLE_DEPRECATION_WARNINGS
-
 bool ULandscapeSubsystem::GetDirtyOnlyInMode() const
 {
 	const ULandscapeSettings* Settings = GetDefault<ULandscapeSettings>();
diff --git a/Engine/Source/Runtime/Landscape/Public/LandscapeBlueprintBrushBase.h b/Engine/Source/Runtime/Landscape/Public/LandscapeBlueprintBrushBase.h
--- a/Engine/Source/Runtime/Landscape/Public/LandscapeBlueprintBrushBase.h
+++ b/Engine/Source/Runtime/Landscape/Public/LandscapeBlueprintBrushBase.h
@@ -73,9 +73,6 @@ protected:
 #endif
 
 public:
-
-	UE_DEPRECATED(5.3, "Please use RenderLayer_Native instead.")
-	virtual UTextureRenderTarget2D* Render_Native(bool InIsHeightmap, UTextureRenderTarget2D* InCombinedResult, const FName& InWeightmapLayerName) {return nullptr;}
 	virtual void Initialize_Native(const FTransform& InLandscapeTransform, const FIntPoint& InLandscapeSize, const FIntPoint& InLandscapeRenderTargetSize) {}
 
 	UFUNCTION(BlueprintNativeEvent, meta = (DeprecatedFunction, DeprecationMessage = "Please use RenderLayer instead."))
@@ -107,15 +104,6 @@ public:
 	LANDSCAPE_API virtual void SetOwningLandscape(class ALandscape* InOwningLandscape);
 	LANDSCAPE_API class ALandscape* GetOwningLandscape() const;
 
-	UE_DEPRECATED(5.3, "Renamed CanAffectHeightmap")
-	bool IsAffectingHeightmap() const { return AffectHeightmap; }
-	UE_DEPRECATED(5.3, "Renamed CanAffectWeightmap")
-	bool IsAffectingWeightmap() const { return AffectWeightmap; }
-	UE_DEPRECATED(5.3, "Renamed AffectsVisibilityLayer")
-	LANDSCAPE_API virtual bool IsAffectingWeightmapLayer(const FName& InLayerName) const;
-	UE_DEPRECATED(5.3, "Renamed CanAffectVisibilityLayer")
-	bool IsAffectingVisibilityLayer() const { return AffectVisibilityLayer; }
-
 	bool CanAffectHeightmap() const { return AffectHeightmap; }
 	bool CanAffectWeightmap() const { return AffectWeightmap; }
 	bool CanAffectVisibilityLayer() const { return AffectVisibilityLayer; }
@@ -130,13 +118,6 @@ public:
 
 	LANDSCAPE_API void SetIsVisible(bool bInIsVisible);
 
-	UE_DEPRECATED(5.3, "Renamed SetCanAffectHeightmap")
-	LANDSCAPE_API void SetAffectsHeightmap(bool bInAffectsHeightmap);
-	UE_DEPRECATED(5.3, "Renamed SetCanAffectWeightmap")
-	LANDSCAPE_API void SetAffectsWeightmap(bool bInAffectsWeightmap);
-	UE_DEPRECATED(5.3, "Renamed SetCanAffectVisibilityLayer")
-	LANDSCAPE_API void SetAffectsVisibilityLayer(bool bInAffectsVisibilityLayer);
-
 	LANDSCAPE_API virtual bool ShouldTickIfViewportsOnly() const override;
 	LANDSCAPE_API virtual void Tick(float DeltaSeconds) override;
 	LANDSCAPE_API virtual void PostEditMove(bool bFinished) override;
diff --git a/Engine/Source/Runtime/Landscape/Public/LandscapeSubsystem.h b/Engine/Source/Runtime/Landscape/Public/LandscapeSubsystem.h
--- a/Engine/Source/Runtime/Landscape/Public/LandscapeSubsystem.h
+++ b/Engine/Source/Runtime/Landscape/Public/LandscapeSubsystem.h
@@ -91,11 +91,6 @@ public:
 	LANDSCAPE_API void BuildGrassMaps();
 	LANDSCAPE_API void BuildPhysicalMaterial();
 
-	UE_DEPRECATED(5.3, "BuildGIBakedTextures is officially deprecated nowe")
-	void BuildGIBakedTextures() {}
-	UE_DEPRECATED(5.3, "GetOutdatedGIBakedTextureComponentsCount is officially deprecated now")
-	int32 GetOutdatedGIBakedTextureComponentsCount() { return 0; }
-
 	/**
 	 * Updates the Nanite mesh on all landscape actors whose mesh is not up to date.
 	 * @param InProxiesToBuild - If specified, only the Nanite meshes of the specified landscape actors (recursively for all streaming proxies, in the case of a 1 ALandscape / N ALandscapeStreamingProxy setup) will be built
@@ -105,11 +100,6 @@ public:
 	
 	LANDSCAPE_API TArray<ALandscapeProxy*> GetOutdatedProxies(UE::Landscape::EOutdatedDataFlags InMatchingOutdatedDataFlags, bool bInMustMatchAllFlags) const;
 	
-	UE_DEPRECATED(5.3, "GetOutdatedGrassMapCount is now deprecated, use GetOutdatedProxies")
-	LANDSCAPE_API int32 GetOutdatedGrassMapCount();
-	UE_DEPRECATED(5.3, "GetOudatedPhysicalMaterialComponentsCount is now deprecated, use GetOutdatedProxies")
-	LANDSCAPE_API int32 GetOudatedPhysicalMaterialComponentsCount();
-
 	LANDSCAPE_API bool IsGridBased() const;
 	LANDSCAPE_API void ChangeGridSize(ULandscapeInfo* LandscapeInfo, uint32 NewGridSizeInComponents);
 	LANDSCAPE_API ALandscapeProxy* FindOrAddLandscapeProxy(ULandscapeInfo* LandscapeInfo, const FIntPoint& SectionBase);
@@ -117,8 +107,6 @@ public:
 	LANDSCAPE_API void MarkModifiedLandscapesAsDirty();
 	LANDSCAPE_API void SaveModifiedLandscapes();
 	LANDSCAPE_API bool HasModifiedLandscapes() const;
-	UE_DEPRECATED(5.3, "Use GetDirtyOnlyInMode instead")
-	LANDSCAPE_API static bool IsDirtyOnlyInModeEnabled();
 	LANDSCAPE_API bool GetDirtyOnlyInMode() const;
 	FLandscapeNotificationManager* GetNotificationManager() { return NotificationManager; }
 	FOnHeightmapStreamedDelegate& GetOnHeightmapStreamedDelegate() { return OnHeightmapStreamed; }
