GCC compilation fix

diff --git a/Engine/Source/Runtime/Landscape/Classes/LandscapeComponent.h b/Engine/Source/Runtime/Landscape/Classes/LandscapeComponent.h
--- a/Engine/Source/Runtime/Landscape/Classes/LandscapeComponent.h
+++ b/Engine/Source/Runtime/Landscape/Classes/LandscapeComponent.h
@@ -204,11 +204,9 @@ struct FLandscapeComponentGrassData
 
 	// Guid per material instance in the hierarchy between the assigned landscape material (instance) and the root UMaterial
 	// used to detect changes to material instance parameters or the root material that could affect the grass maps
-	UE_DEPRECATED(5.4, "GenerationHash is now used") 
 	TArray<FGuid, TInlineAllocator<2>> MaterialStateIds_DEPRECATED;
 	// cached component rotation when material world-position-offset is used,
 	// as this will affect the direction of world-position-offset deformation (included in the HeightData below)
-	UE_DEPRECATED(5.4, "GenerationHash is now used") 
 	FQuat RotationForWPO_DEPRECATED;
 
 	// Variable used to detect when grass data needs to be regenerated:
