Fixed crash when rendering shadows out of a viewstate-less view.

Full story : when rendering shadows, we create copies of FViewInfo (via FViewInfo::CreateSnapshot()) which are *not* in the ViewFamily's views of the FSceneRenderer, which prevents the landscape render systems from computing LOD section values. In this case, we want to reuse the originating view's values, which works for views with a view state (as the snapshot view simply copies the origin view's view state) but does not for those without view state.
The fix consists in storing a back pointer to the originating view when creating a snapshot view, so that the landscape render system can fallback to this view's LOD section values when it can't find the view in its cached values

diff --git a/Engine/Source/Runtime/Engine/Public/SceneView.h b/Engine/Source/Runtime/Engine/Public/SceneView.h
--- a/Engine/Source/Runtime/Engine/Public/SceneView.h
+++ b/Engine/Source/Runtime/Engine/Public/SceneView.h
@@ -1627,6 +1627,9 @@ public:
 protected:
 	friend class FSceneRenderer;
 
+	/** Some views get cloned for certain renders, like shadows (see FViewInfo::CreateSnapshot()). If that's the case, this will point to the view this one originates from */
+	const FSceneView* SnapshotOriginView = nullptr;
+
 public:
 
 	/** Initialization constructor. */
@@ -1870,6 +1873,8 @@ public:
 		return bHMDHiddenAreaMaskActive ? ERenderTargetLoadAction::EClear : ERenderTargetLoadAction::ENoAction;
 	}
 
+	const FSceneView* GetSnapshotOriginView() const { return SnapshotOriginView; }
+
 protected:
 	FSceneViewStateInterface* EyeAdaptationViewState = nullptr;
 
diff --git a/Engine/Source/Runtime/Landscape/Private/LandscapeRender.cpp b/Engine/Source/Runtime/Landscape/Private/LandscapeRender.cpp
--- a/Engine/Source/Runtime/Landscape/Private/LandscapeRender.cpp
+++ b/Engine/Source/Runtime/Landscape/Private/LandscapeRender.cpp
@@ -947,6 +947,13 @@ const TResourceArray<float>& FLandscapeRenderSystem::GetCachedSectionLODValues(c
 {
 	const uint32 ViewStateKey = InView.GetViewKey();
 	const TResourceArray<float>* CachedSectionLODValues = (ViewStateKey != 0) ? PerViewStateCachedSectionLODValues.Find(ViewStateKey) : PerViewCachedSectionLODValues.Find(&InView);
+	const FSceneView* SnapshotOriginView = InView.GetSnapshotOriginView();
+	// If we couldn't find this view or view state in our maps, it's possible the view originates from another one, which we must have recorded in our maps : 
+	if ((CachedSectionLODValues == nullptr) && (SnapshotOriginView != nullptr))
+	{
+		checkf(SnapshotOriginView != &InView, TEXT("Infinite loop will happen if the view originates from itself!"));
+		CachedSectionLODValues = &GetCachedSectionLODValues(*SnapshotOriginView);
+	}
 	checkf(CachedSectionLODValues != nullptr, TEXT("No section LOD value cached for this view. Make sure FLandscapeRenderSystem::ComputeSectionsLODForView (FLandscapeSceneViewExtension::PreRenderView_RenderThread) was called"));
 	return *CachedSectionLODValues;
 }
diff --git a/Engine/Source/Runtime/Renderer/Private/SceneRendering.cpp b/Engine/Source/Runtime/Renderer/Private/SceneRendering.cpp
--- a/Engine/Source/Runtime/Renderer/Private/SceneRendering.cpp
+++ b/Engine/Source/Runtime/Renderer/Private/SceneRendering.cpp
@@ -905,7 +905,6 @@ void FViewInfo::Init()
 	ShaderMap = GetGlobalShaderMap(FeatureLevel);
 
 	ViewState = (FSceneViewState*)State;
-	bIsSnapshot = false;
 	bHMDHiddenAreaMaskActive = IsHMDHiddenAreaMaskActive();
 	bUseComputePasses = IsPostProcessingWithComputeEnabled(FeatureLevel);
 	bHasCustomDepthPrimitives = false;
@@ -2228,7 +2227,7 @@ FViewInfo* FViewInfo::CreateSnapshot() const
 	static_assert(std::is_trivially_destructible_v<FGPUScenePrimitiveCollector> != 0, "The destructor is not invoked properly because of FMemory::Memcpy(*Result, *this) above");
 	Result->DynamicPrimitiveCollector = FGPUScenePrimitiveCollector(DynamicPrimitiveCollector);
 
-	Result->bIsSnapshot = true;
+	Result->SnapshotOriginView = this;
 	return Result;
 }
 
@@ -2262,6 +2261,8 @@ void FViewInfo::DestroyAllSnapshots(FParallelMeshDrawCommandPass::EWaitThread Wa
 			Snapshot->ParallelMeshDrawCommandPasses[i].FreeCreateSnapshot();
 		}
 
+		Snapshot->SnapshotOriginView = nullptr;
+
 		ViewInfoSnapshotCache.FreeSnapshots.Push(Snapshot);
 	}
 
diff --git a/Engine/Source/Runtime/Renderer/Private/SceneRendering.h b/Engine/Source/Runtime/Renderer/Private/SceneRendering.h
--- a/Engine/Source/Runtime/Renderer/Private/SceneRendering.h
+++ b/Engine/Source/Runtime/Renderer/Private/SceneRendering.h
@@ -1532,8 +1532,6 @@ public:
 
 	FGlobalShaderMap* ShaderMap;
 
-	bool bIsSnapshot;
-
 	// Whether this view should use compute passes where appropriate.
 	bool bUseComputePasses = false;
 
