Fixed a couple of mistakes with my previous deprecation CL (as it was meant to be submitted in 5.4 but ended up missing the cutoff)

diff --git a/Engine/Source/Runtime/Landscape/Classes/LandscapeComponent.h b/Engine/Source/Runtime/Landscape/Classes/LandscapeComponent.h
--- a/Engine/Source/Runtime/Landscape/Classes/LandscapeComponent.h
+++ b/Engine/Source/Runtime/Landscape/Classes/LandscapeComponent.h
@@ -204,9 +204,11 @@ struct FLandscapeComponentGrassData
 
 	// Guid per material instance in the hierarchy between the assigned landscape material (instance) and the root UMaterial
 	// used to detect changes to material instance parameters or the root material that could affect the grass maps
+	UE_DEPRECATED(5.5, "GenerationHash is now used") 
 	TArray<FGuid, TInlineAllocator<2>> MaterialStateIds_DEPRECATED;
 	// cached component rotation when material world-position-offset is used,
 	// as this will affect the direction of world-position-offset deformation (included in the HeightData below)
+	UE_DEPRECATED(5.5, "GenerationHash is now used") 
 	FQuat RotationForWPO_DEPRECATED;
 
 	// Variable used to detect when grass data needs to be regenerated:
@@ -229,8 +231,16 @@ struct FLandscapeComponentGrassData
 	TMap<TObjectPtr<ULandscapeGrassType>, int32> WeightOffsets;
 	TArray<uint8> HeightWeightData;
 
+	// Note: We need to explicitly disable warnings on these constructors/operators for clang to be happy with deprecated variables
+    PRAGMA_DISABLE_DEPRECATION_WARNINGS
 	FLandscapeComponentGrassData() = default;
 	FLandscapeComponentGrassData(ULandscapeComponent* Component);
+    ~FLandscapeComponentGrassData() = default;
+    FLandscapeComponentGrassData(const FLandscapeComponentGrassData&) = default;
+    FLandscapeComponentGrassData(FLandscapeComponentGrassData&&) = default;
+    FLandscapeComponentGrassData& operator=(const FLandscapeComponentGrassData&) = default;
+    FLandscapeComponentGrassData& operator=(FLandscapeComponentGrassData&&) = default;
+	PRAGMA_ENABLE_DEPRECATION_WARNINGS
 
 	// Returns whether grass data has been computed (or serialized) yet. Returns true even if the data is completely empty (e.g. all-zero weightmap data)
 	bool HasValidData() const;
@@ -431,7 +441,7 @@ class ULandscapeComponent : public UPrimitiveComponent
 	TArray<FLandscapeComponentMaterialOverride> OverrideMaterials_DEPRECATED;
 	PRAGMA_ENABLE_DEPRECATION_WARNINGS
 
-	UE_DEPRECATED(5.4, "MaterialInstance has been deprecated, use MaterialInstances instead.")
+	UE_DEPRECATED(5.5, "MaterialInstance has been deprecated, use MaterialInstances instead.")
 	UPROPERTY()
 	TObjectPtr<UMaterialInstanceConstant> MaterialInstance_DEPRECATED;
 #endif // WITH_EDITORONLY_DATA
@@ -478,7 +488,7 @@ class ULandscapeComponent : public UPrimitiveComponent
 	TArray<double> MipToMipMaxDeltas;
 
 #if WITH_EDITORONLY_DATA
-	UE_DEPRECATED(5.4, "CollisionComponent has been deprecated and will be removed in a future version")
+	UE_DEPRECATED(5.5, "CollisionComponent has been deprecated and will be removed in a future version")
 	UPROPERTY()
 	TLazyObjectPtr<ULandscapeHeightfieldCollisionComponent> CollisionComponent_DEPRECATED;
 #endif // !WITH_EDITORONLY_DATA
@@ -596,13 +606,13 @@ public:
 	UPROPERTY(EditAnywhere, BlueprintReadOnly, Category=LandscapeComponent)
 	int32 LODBias;
 
-	UE_DEPRECATED(5.4, "StateId is unused and will be removed in a future version")
+	UE_DEPRECATED(5.5, "StateId is unused and will be removed in a future version")
 	UPROPERTY()
 	FGuid StateId;
 
 #if WITH_EDITORONLY_DATA
 	/**	Legacy irrelevant lights */
-	UE_DEPRECATED(5.4, "IrrelevantLights is officially deprecated now and will be removed in a future version")
+	UE_DEPRECATED(5.5, "IrrelevantLights is officially deprecated now and will be removed in a future version")
 	UPROPERTY()
 	TArray<FGuid> IrrelevantLights_DEPRECATED;
 
@@ -640,11 +650,11 @@ public:
 	UPROPERTY(Transient)
 	uint32 LastSavedPhysicalMaterialHash;
 
-	UE_DEPRECATED(5.4, "MobileMaterialInterface has been deprecated and will be removed in a future version")
+	UE_DEPRECATED(5.5, "MobileMaterialInterface has been deprecated and will be removed in a future version")
 	UPROPERTY(NonPIEDuplicateTransient)
 	TObjectPtr<UMaterialInterface> MobileMaterialInterface_DEPRECATED;
 
-    UE_DEPRECATED(5.4, "MobileCombinationMaterialInstance has been deprecated and will be removed in a future version")
+    UE_DEPRECATED(5.5, "MobileCombinationMaterialInstance has been deprecated and will be removed in a future version")
 	UPROPERTY(NonPIEDuplicateTransient)
 	TObjectPtr<UMaterialInstanceConstant> MobileCombinationMaterialInstance_DEPRECATED;
 #endif // WITH_EDITORONLY_DATA
diff --git a/Engine/Source/Runtime/Landscape/Classes/LandscapeProxy.h b/Engine/Source/Runtime/Landscape/Classes/LandscapeProxy.h
--- a/Engine/Source/Runtime/Landscape/Classes/LandscapeProxy.h
+++ b/Engine/Source/Runtime/Landscape/Classes/LandscapeProxy.h
@@ -451,11 +451,11 @@ public:
 	int32 MaxLODLevel;
 
 #if WITH_EDITORONLY_DATA
-	UE_DEPRECATED(5.4, "Unused property : will be removed in a future version")
+	UE_DEPRECATED(5.5, "Unused property : will be removed in a future version")
 	UPROPERTY()
 	float LODDistanceFactor_DEPRECATED;
 
-	UE_DEPRECATED(5.4, "Unused property : will be removed in a future version")
+	UE_DEPRECATED(5.5, "Unused property : will be removed in a future version")
 	UPROPERTY()
 	TEnumAsByte<ELandscapeLODFalloff::Type> LODFalloff_DEPRECATED;
 #endif // WITH_EDITORONLY_DATA
@@ -793,7 +793,7 @@ public:
 	uint32 bBakeMaterialPositionOffsetIntoCollision:1;
 
 #if WITH_EDITORONLY_DATA
-	UE_DEPRECATED(5.4, "Unused property : will be removed in a future version")
+	UE_DEPRECATED(5.5, "Unused property : will be removed in a future version")
 	UPROPERTY()
 	TArray<TObjectPtr<ULandscapeLayerInfoObject>> EditorCachedLayerInfos_DEPRECATED;
 
