Fix for landscape taking a long time when creating a project from a template while using DirectX 11.

The GPU Readback will never return IsReady() within a loop in a single frame, so we must force it to complete anyways.

I chose to just have it force a completion after 5 ticks, which is just a simple metric that should keep it from being hit in the amortized case (readbacks come back in 3 ticks at most), but will quickly get hit in the BuildNow() case.

diff --git a/Engine/Source/Runtime/Landscape/Private/LandscapeAsyncTextureReadback.cpp b/Engine/Source/Runtime/Landscape/Private/LandscapeAsyncTextureReadback.cpp
--- a/Engine/Source/Runtime/Landscape/Private/LandscapeAsyncTextureReadback.cpp
+++ b/Engine/Source/Runtime/Landscape/Private/LandscapeAsyncTextureReadback.cpp
@@ -83,7 +83,7 @@ void FLandscapeAsyncTextureReadback::FinishReadback_RenderThread()
 	bFinishedOnRenderThread = true;
 }
 
-bool FLandscapeAsyncTextureReadback::CheckAndUpdate(bool& bOutFinishCommandQueued)
+bool FLandscapeAsyncTextureReadback::CheckAndUpdate(bool& bOutFinishCommandQueued, const bool bInForceFinish)
 {
 	// if we already queued the finish commands to render thread, then we're just waiting on it signaling readback complete	
 	if (bFinishQueuedFromGameThread)
@@ -92,7 +92,7 @@ bool FLandscapeAsyncTextureReadback::CheckAndUpdate(bool& bOutFinishCommandQueue
 	}
 	
 	// if we haven't started, or if the readback is not yet ready, then we have nothing to do but wait
-	if (!bStartedOnRenderThread || !AsyncReadback->IsReady())
+	if (!bStartedOnRenderThread || (!bInForceFinish && !AsyncReadback->IsReady()))
 	{
 		return false;
 	}
@@ -101,11 +101,11 @@ bool FLandscapeAsyncTextureReadback::CheckAndUpdate(bool& bOutFinishCommandQueue
 	// queue it to make the data available to the game thread
 	FLandscapeAsyncTextureReadback* Readback = this;
 	ENQUEUE_RENDER_COMMAND(FLandscapeAsyncTextureReadback_FinishReadback)(
-		[Readback](FRHICommandListImmediate& RHICmdList)
+		[Readback, bInForceFinish](FRHICommandListImmediate& RHICmdList)
 		{
 			// sanity check the state 
 			check(Readback->bStartedOnRenderThread && !Readback->bFinishedOnRenderThread);
-			check(Readback->AsyncReadback.IsValid() && Readback->AsyncReadback->IsReady());
+			check(Readback->AsyncReadback.IsValid() && (bInForceFinish || Readback->AsyncReadback->IsReady()));
 			Readback->FinishReadback_RenderThread();
 		});
 
diff --git a/Engine/Source/Runtime/Landscape/Private/LandscapeAsyncTextureReadback.h b/Engine/Source/Runtime/Landscape/Private/LandscapeAsyncTextureReadback.h
--- a/Engine/Source/Runtime/Landscape/Private/LandscapeAsyncTextureReadback.h
+++ b/Engine/Source/Runtime/Landscape/Private/LandscapeAsyncTextureReadback.h
@@ -50,7 +50,9 @@ public:
 	// Returns true when the AsyncReadbackResults are available to the game thread, at which point you can call TakeResults() to access them.
 	// bOutFinishCommandQueued is set to true if the finish render command was queued in this call, othwerwise unchanged.
 	// You must call this function occasionally or the readback will never complete, as the finish command is required.
-	bool CheckAndUpdate(bool& bOutFinishCommandQueued);
+	// bInForceFinish will force the finish process to be queued on the render thread (potentially stalling render thread, but forcing it to finish the readback)
+	// bInForceFinish may not make the results immediately available, but will ensure they are available after the render thread executes the command.
+	bool CheckAndUpdate(bool& bOutFinishCommandQueued, const bool bInForceFinish);
 
 	// Call from game thread to terminate any readback in flight and queue deletion of this (FLandscapeAsyncTextureReadback) on the render thread
 	void CancelAndSelfDestruct();
diff --git a/Engine/Source/Runtime/Landscape/Private/LandscapeGrassMapsBuilder.cpp b/Engine/Source/Runtime/Landscape/Private/LandscapeGrassMapsBuilder.cpp
--- a/Engine/Source/Runtime/Landscape/Private/LandscapeGrassMapsBuilder.cpp
+++ b/Engine/Source/Runtime/Landscape/Private/LandscapeGrassMapsBuilder.cpp
@@ -447,17 +447,25 @@ bool FLandscapeGrassMapsBuilder::UpdateTrackedComponents(const TArray<FVector>&
 
 			case EComponentStage::Rendering:
 				check(State->ActiveRender != nullptr);
-				if (State->ActiveRender->CheckAndUpdateAsyncReadback(bRenderCommandsQueuedByLastUpdate))
 				{
-					if (GGrassMapUseAsyncFetch != 0)
-					{
-						LaunchAsyncFetchTask(*State);
-					}
-					else
+					// NOTE: on RHI platforms that don't support fences (D3D11), the IsReady() async check will never return true within a single frame (i.e. BuildGrassMapsNow)
+					// So it will never signal complete, unless we force it to finish after a number of ticks
+					// in general 3 ticks should be the max latency we should ever see in an amortized use case
+					const bool bForceFinish = (State->TickCount > 4);
+					const bool bComplete = State->ActiveRender->CheckAndUpdateAsyncReadback(bRenderCommandsQueuedByLastUpdate, bForceFinish);
+
+					if (bComplete)
 					{
-						PopulateGrassDataFromReadback(*State);
+						if (GGrassMapUseAsyncFetch != 0)
+						{
+							LaunchAsyncFetchTask(*State);
+						}
+						else
+						{
+							PopulateGrassDataFromReadback(*State);
+						}
+						bChanged = true;
 					}
-					bChanged = true;
 				}
 				continue; // next!
 
@@ -901,7 +909,10 @@ bool FLandscapeGrassMapsBuilder::BuildGrassMapsNowForComponents(
 
 	if (UpToDateCount != LandscapeComponents.Num())
 	{
-		UE_LOG(LogGrass, Warning, TEXT("Failed to build grass maps for %d landscape components, check if you are using a render preview mode, or a non-SM5 capable render device"), LandscapeComponents.Num() - UpToDateCount);
+		UE_LOG(LogGrass, Warning, TEXT("Failed to build grass maps for %d/%d landscape components, check if you are using a render preview mode, or a non-SM5 capable render device.  (%d pending, %d streaming, %d rendering, %d fetching, %d built)"),
+			LandscapeComponents.Num() - UpToDateCount,
+			LandscapeComponents.Num(),
+			PendingCount, StreamingCount, RenderingCount, AsyncFetchCount, PopulatedCount);
 	}
 
 	return (UpToDateCount == LandscapeComponents.Num());
@@ -960,7 +971,7 @@ bool FLandscapeGrassMapsBuilder::CancelAndEvict(FComponentState& State, bool bCa
 					FLandscapeGrassWeightExporter* Render = State.ActiveRender.Release();
 					Render->CancelAndSelfDestruct();
 				}
-				else if (!State.ActiveRender->CheckAndUpdateAsyncReadback(bNewRenderCommands))
+				else if (!State.ActiveRender->CheckAndUpdateAsyncReadback(bNewRenderCommands, State.TickCount > 4))
 				{
 					// we can't cancel yet.. must wait for the readback to complete
 					return false;
diff --git a/Engine/Source/Runtime/Landscape/Private/LandscapeGrassWeightExporter.h b/Engine/Source/Runtime/Landscape/Private/LandscapeGrassWeightExporter.h
--- a/Engine/Source/Runtime/Landscape/Private/LandscapeGrassWeightExporter.h
+++ b/Engine/Source/Runtime/Landscape/Private/LandscapeGrassWeightExporter.h
@@ -91,10 +91,12 @@ public:
 
 	// If using the async readback path, check its status and update if needed. Return true when the AsyncReadbackResults are available.
 	// You must call this periodically, or the async readback may not complete.
-	bool CheckAndUpdateAsyncReadback(bool& bOutRenderCommandsQueued)
+	// bInForceFinish will force the RenderThread to wait until GPU completes the readback, ensuring the readback is completed after the render thread executes the command.
+	// NOTE: you may still see false returned, this just means the render thread hasn't executed the command yet.
+	bool CheckAndUpdateAsyncReadback(bool& bOutRenderCommandsQueued, const bool bInForceFinish = false)
 	{
 		check(AsyncReadbackPtr != nullptr);
-		return AsyncReadbackPtr->CheckAndUpdate(bOutRenderCommandsQueued);
+		return AsyncReadbackPtr->CheckAndUpdate(bOutRenderCommandsQueued, bInForceFinish);
 	}
 
 	// return true if the async readback is complete.  (Does not update the readback state)
