Cleanup some logic in LandscapeRender.cpp

diff --git a/Engine/Source/Runtime/Landscape/Private/LandscapeRender.cpp b/Engine/Source/Runtime/Landscape/Private/LandscapeRender.cpp
--- a/Engine/Source/Runtime/Landscape/Private/LandscapeRender.cpp
+++ b/Engine/Source/Runtime/Landscape/Private/LandscapeRender.cpp
@@ -1859,6 +1859,8 @@ FLandscapeRayTracingState* FLandscapeRayTracingImpl::FindOrCreateRayTracingState
 		{
 			const int8 SubSectionIdx = static_cast<int8>(SubX + SubY * NumSubsections);
 
+			FLandscapeSectionRayTracingState& SectionRayTracingState = RayTracingState->Sections[SubSectionIdx];
+
 			FRayTracingGeometryInitializer Initializer;
 			static const FName DebugName("FLandscapeComponentSceneProxy");
 			static int32 DebugNumber = 0;
@@ -1867,18 +1869,20 @@ FLandscapeRayTracingState* FLandscapeRayTracingImpl::FindOrCreateRayTracingState
 			Initializer.GeometryType = RTGT_Triangles;
 			Initializer.bFastBuild = true;
 			Initializer.bAllowUpdate = true;
+
 			FRayTracingGeometrySegment Segment;
 			Segment.VertexBuffer = nullptr;
 			Segment.VertexBufferStride = sizeof(FVector3f);
 			Segment.VertexBufferElementType = VET_Float3;
 			Segment.MaxVertices = FMath::Square(SubsectionSizeVerts);
 			Initializer.Segments.Add(Segment);
-			RayTracingState->Sections[SubSectionIdx].Geometry.SetInitializer(Initializer);
-			RayTracingState->Sections[SubSectionIdx].Geometry.InitResource(RHICmdList);
+
+			SectionRayTracingState.Geometry.SetInitializer(Initializer);
+			SectionRayTracingState.Geometry.InitResource(RHICmdList);
 
 			FLandscapeVertexFactoryMVFParameters UniformBufferParams;
 			UniformBufferParams.SubXY = FIntPoint(SubX, SubY);
-			RayTracingState->Sections[SubSectionIdx].UniformBuffer = FLandscapeVertexFactoryMVFUniformBufferRef::CreateUniformBufferImmediate(UniformBufferParams, UniformBuffer_MultiFrame);
+			SectionRayTracingState.UniformBuffer = FLandscapeVertexFactoryMVFUniformBufferRef::CreateUniformBufferImmediate(UniformBufferParams, UniformBuffer_MultiFrame);
 		}
 	}
 
@@ -3030,6 +3034,8 @@ void FLandscapeComponentSceneProxy::GetDynamicRayTracingInstances(FRayTracingMat
 			const int8 SubSectionIdx = static_cast<int8>(SubX + SubY * NumSubsections);
 			const int8 CurrentLOD = static_cast<int8>(LODToRender);
 
+			FLandscapeSectionRayTracingState& SectionRayTracingState = RayTracingState->Sections[SubSectionIdx];
+
 			FMeshBatch MeshBatch = BaseMeshBatch;
 
 			FMeshBatchElement BatchElement;
@@ -3060,9 +3066,9 @@ void FLandscapeComponentSceneProxy::GetDynamicRayTracingInstances(FRayTracingMat
 
 			MeshBatch.Elements.Add(BatchElement);
 
-			RayTracingState->Sections[SubSectionIdx].Geometry.Initializer.IndexBuffer = BatchElement.IndexBuffer->IndexBufferRHI;
+			SectionRayTracingState.Geometry.Initializer.IndexBuffer = BatchElement.IndexBuffer->IndexBufferRHI;
 
-			BatchElementParams.LandscapeVertexFactoryMVFUniformBuffer = RayTracingState->Sections[SubSectionIdx].UniformBuffer;
+			BatchElementParams.LandscapeVertexFactoryMVFUniformBuffer = SectionRayTracingState.UniformBuffer;
 
 			bool bNeedsRayTracingGeometryUpdate = false;
 
@@ -3071,24 +3077,24 @@ void FLandscapeComponentSceneProxy::GetDynamicRayTracingInstances(FRayTracingMat
 
 			// Detect continuous LOD parameter changes. This is for far-away high LODs - they change rarely yet the BLAS refit time is not ideal, even if they contains tiny amount of triangles
 			{
-				if (RayTracingState->Sections[SubSectionIdx].CurrentLOD != CurrentLOD)
+				if (SectionRayTracingState.CurrentLOD != CurrentLOD)
 				{
 					bNeedsRayTracingGeometryUpdate = true;
-					RayTracingState->Sections[SubSectionIdx].CurrentLOD = CurrentLOD;
-					RayTracingState->Sections[SubSectionIdx].RayTracingDynamicVertexBuffer.Release();
+					SectionRayTracingState.CurrentLOD = CurrentLOD;
+					SectionRayTracingState.RayTracingDynamicVertexBuffer.Release();
 				}
-				if (RayTracingState->Sections[SubSectionIdx].HeightmapLODBias != RenderSystem.GetSectionLODBias(RenderCoord))
+				if (SectionRayTracingState.HeightmapLODBias != RenderSystem.GetSectionLODBias(RenderCoord))
 				{
 					bNeedsRayTracingGeometryUpdate = true;
-					RayTracingState->Sections[SubSectionIdx].HeightmapLODBias = RenderSystem.GetSectionLODBias(RenderCoord);
+					SectionRayTracingState.HeightmapLODBias = RenderSystem.GetSectionLODBias(RenderCoord);
 				}
 
 				const float PendingFractionalLOD = RenderSystem.GetSectionLODValue(SceneView, RenderCoord);
-				const float FractionLODAbsoluteDifference = FMath::Abs(RayTracingState->Sections[SubSectionIdx].FractionalLOD - PendingFractionalLOD);
+				const float FractionLODAbsoluteDifference = FMath::Abs(SectionRayTracingState.FractionalLOD - PendingFractionalLOD);
 				if (FractionLODAbsoluteDifference > GLandscapeRayTracingGeometryFractionalLODUpdateThreshold)
 				{
 					bNeedsRayTracingGeometryUpdate = true;
-					RayTracingState->Sections[SubSectionIdx].FractionalLOD = PendingFractionalLOD;
+					SectionRayTracingState.FractionalLOD = PendingFractionalLOD;
 				}
 			}
 
@@ -3106,16 +3112,16 @@ void FLandscapeComponentSceneProxy::GetDynamicRayTracingInstances(FRayTracingMat
 					const FUniformExpressionSet& UniformExpressionSet = Material.GetRenderingThreadShaderMap()->GetUniformExpressionSet();
 					const uint32 Hash = UniformExpressionSet.GetReferencedTexture2DRHIHash(MaterialRenderContext);
 
-					if (RayTracingState->Sections[SubSectionIdx].ReferencedTextureRHIHash != Hash)
+					if (SectionRayTracingState.ReferencedTextureRHIHash != Hash)
 					{
 						bNeedsRayTracingGeometryUpdate = true;
-						RayTracingState->Sections[SubSectionIdx].ReferencedTextureRHIHash = Hash;
+						SectionRayTracingState.ReferencedTextureRHIHash = Hash;
 					}
 				}
 			}
 
 			FRayTracingInstance RayTracingInstance;
-			RayTracingInstance.Geometry = &RayTracingState->Sections[SubSectionIdx].Geometry;
+			RayTracingInstance.Geometry = &SectionRayTracingState.Geometry;
 			RayTracingInstance.InstanceTransforms.Add(GetLocalToWorld());
 			RayTracingInstance.Materials.Add(MeshBatch);
 			OutRayTracingInstances.Add(RayTracingInstance);
@@ -3133,8 +3139,8 @@ void FLandscapeComponentSceneProxy::GetDynamicRayTracingInstances(FRayTracingMat
 						(uint32)FMath::Square(LodSubsectionSizeVerts),
 						FMath::Square(LodSubsectionSizeVerts) * (uint32)sizeof(FVector3f),
 						(uint32)FMath::Square(LodSubsectionSizeVerts - 1) * 2,
-						&RayTracingState->Sections[SubSectionIdx].Geometry,
-						&RayTracingState->Sections[SubSectionIdx].RayTracingDynamicVertexBuffer,
+						&SectionRayTracingState.Geometry,
+						&SectionRayTracingState.RayTracingDynamicVertexBuffer,
 						true
 					}
 				);
