Remove unnecessary warning in Landscape, and fix the behavior of the DrawScale cached value in LandscapeInfo.

Previously, landscape info would grab DrawScale from the first registered Proxy that had a valid transform, and then issue warnings on any discrepancy with later registered Proxies.  This is not really necessary as the transforms all get fixed up to match the main landscape Actor when both are registered, so the warning is removed.

However, the DrawScale is grabbed somewhat incorrectly.  If we first register a streaming proxy that has incorrect scale because it was never updated, then it will lock onto the incorrect value of DrawScale.  The code was changed to prefer  the Actor DrawScale over any streaming proxy DrawScale.

Also fixed an issue with the PropertyChanged code pertaining to the scale, that would spuriously issue warnings if you pasted a scale with mismatched signs on X and Y.

I think ideally we would get rid of the DrawScale cached value, as I don't see the value in caching it over computing it.  But it is used in many places, so is a large change to deprecate.

diff --git a/Engine/Source/Runtime/Landscape/Classes/LandscapeInfo.h b/Engine/Source/Runtime/Landscape/Classes/LandscapeInfo.h
--- a/Engine/Source/Runtime/Landscape/Classes/LandscapeInfo.h
+++ b/Engine/Source/Runtime/Landscape/Classes/LandscapeInfo.h
@@ -138,10 +138,14 @@ class ULandscapeInfo : public UObject
 	UPROPERTY()
 	int32 ComponentNumSubsections;
 	
-	UPROPERTY()
-	FVector DrawScale;
-	
 #if WITH_EDITORONLY_DATA
+	/** True when the draw scale has been set by the landscape Actor (false if default or set by a streaming proxy) */
+	UPROPERTY()
+	bool bDrawScaleSetByActor = false;
+
+	UPROPERTY()
+	FVector DrawScale = FVector(100.0);
+
 	UPROPERTY()
 	TArray<FLandscapeInfoLayerSettings> Layers;
 
diff --git a/Engine/Source/Runtime/Landscape/Private/Landscape.cpp b/Engine/Source/Runtime/Landscape/Private/Landscape.cpp
--- a/Engine/Source/Runtime/Landscape/Private/Landscape.cpp
+++ b/Engine/Source/Runtime/Landscape/Private/Landscape.cpp
@@ -5453,7 +5453,6 @@ void ULandscapeInfo::RegisterActor(ALandscapeProxy* Proxy, bool bMapCheck, bool
 		ComponentSizeQuads = Proxy->ComponentSizeQuads;
 		ComponentNumSubsections = Proxy->NumSubsections;
 		SubsectionSizeQuads = Proxy->SubsectionSizeQuads;
-		DrawScale = Proxy->GetRootComponent() != nullptr ? Proxy->GetRootComponent()->GetRelativeScale3D() : FVector(100.0f);
 	}
 
 	// check that passed actor matches all shared parameters
@@ -5462,15 +5461,18 @@ void ULandscapeInfo::RegisterActor(ALandscapeProxy* Proxy, bool bMapCheck, bool
 	check(ComponentNumSubsections == Proxy->NumSubsections);
 	check(SubsectionSizeQuads == Proxy->SubsectionSizeQuads);
 
-	if (Proxy->GetRootComponent() != nullptr && !DrawScale.Equals(Proxy->GetRootComponent()->GetRelativeScale3D()))
-	{
-		UE_LOG(LogLandscape, Warning, TEXT("Landscape proxy (%s) scale (%s) does not match to main actor scale (%s)."),
-			*Proxy->GetPathName(), *Proxy->GetRootComponent()->GetRelativeScale3D().ToCompactString(), *DrawScale.ToCompactString());
-	}
-
 	// register
 	if (ALandscape* Landscape = Cast<ALandscape>(Proxy))
 	{
+#if WITH_EDITORONLY_DATA
+		USceneComponent* Root = Proxy->GetRootComponent();
+		if (Root)
+		{
+			DrawScale = Root->GetRelativeScale3D();
+			bDrawScaleSetByActor = true;
+		}
+#endif // WITH_EDITORONLY_DATA
+
 		if (!LandscapeActor.IsValid())
 		{
 			LandscapeActor = Landscape;
@@ -5508,6 +5510,17 @@ void ULandscapeInfo::RegisterActor(ALandscapeProxy* Proxy, bool bMapCheck, bool
 	}
 	else
 	{
+#if WITH_EDITORONLY_DATA
+		if (!bDrawScaleSetByActor)
+		{
+			USceneComponent* Root = Proxy->GetRootComponent();
+			if (Root)
+			{
+				DrawScale = Root->GetRelativeScale3D();
+			}
+		}
+#endif // WITH_EDITORONLY_DATA
+
 		auto LamdbdaLowerBound = [](TWeakObjectPtr<ALandscapeProxy> APtr, TWeakObjectPtr<ALandscapeProxy> BPtr)
 		{
 			ALandscapeProxy *A = APtr.Get();
diff --git a/Engine/Source/Runtime/Landscape/Private/LandscapeEdit.cpp b/Engine/Source/Runtime/Landscape/Private/LandscapeEdit.cpp
--- a/Engine/Source/Runtime/Landscape/Private/LandscapeEdit.cpp
+++ b/Engine/Source/Runtime/Landscape/Private/LandscapeEdit.cpp
@@ -5888,24 +5888,29 @@ void ALandscapeProxy::PostEditChangeProperty(FPropertyChangedEvent& PropertyChan
 			const FVector OriginalScale = RootComponent->GetRelativeScale3D();
 			FVector ModifiedScale = OriginalScale;
 
-			// Lock X and Y scaling to the same value
+			// Lock X and Y scaling to the same value (but preserving their original sign)
 			if (SubPropertyName == FName("Y"))
 			{
-				ModifiedScale.X = FMath::Abs(OriginalScale.Y) * FMath::Sign(ModifiedScale.X);
+				ModifiedScale.X = FMath::Abs(OriginalScale.Y) * FMath::Sign(OriginalScale.X);
 			}
 			else if (SubPropertyName == FName("X"))
 			{
-				ModifiedScale.Y = FMath::Abs(OriginalScale.X) * FMath::Sign(ModifiedScale.Y);
+				ModifiedScale.Y = FMath::Abs(OriginalScale.X) * FMath::Sign(OriginalScale.Y);
+			} 
+			else if (SubPropertyName == FName("Z"))
+			{
+				// no adjustment necessary for Z, it is unrestricted
 			}
-			else if (SubPropertyName != FName("Z"))
+			else
 			{
-				// When changing all axis values at once (e.g. when copy/pasting the scale), we receive only one event and the sub-property is not set :
-				check(SubPropertyName == MemberPropertyName); // Any other combination of property / sub-property is invalid				
-				if (!FMath::IsNearlyEqual(ModifiedScale.X, ModifiedScale.Y))
+				// When changing all axis values at once (e.g. when copy/pasting the scale), we receive only one event and the sub-property is not set
+				check(SubPropertyName == MemberPropertyName); // Any other combination of property / sub-property is invalid
+				// ensure that X and Y are locked (except for sign)
+				if (!FMath::IsNearlyEqual(FMath::Abs(OriginalScale.X), FMath::Abs(OriginalScale.Y)))
 				{
-					UE_LOG(LogLandscape, Warning, TEXT("Non-uniform XY scale for landscape (%f, %f) : scale will be forced to (%f, %f)"), ModifiedScale.X, ModifiedScale.Y, ModifiedScale.X, ModifiedScale.X);
-					// Arbitrarily favor the X axis as the uniform scale value (but retain the sign) : 
-					ModifiedScale.Y = FMath::Abs(OriginalScale.X) * FMath::Sign(ModifiedScale.Y);
+					// Arbitrarily favor the X axis as the uniform scale value (but retain the sign of Y) : 
+					ModifiedScale.Y = FMath::Abs(OriginalScale.X) * FMath::Sign(OriginalScale.Y);
+					UE_LOG(LogLandscape, Warning, TEXT("Non-uniform XY scale for landscape (%f, %f) : scale will be forced to (%f, %f)"), OriginalScale.X, OriginalScale.Y, ModifiedScale.X, ModifiedScale.Y);
 				}
 			}
 
