Fix for landscape grass disappearing when switching to a preview mode in editor
Grass maps are not recalculated per-platform, so there's no reason to evict them when the feature level changes.

diff --git a/Engine/Source/Runtime/Landscape/Private/Landscape.cpp b/Engine/Source/Runtime/Landscape/Private/Landscape.cpp
--- a/Engine/Source/Runtime/Landscape/Private/Landscape.cpp
+++ b/Engine/Source/Runtime/Landscape/Private/Landscape.cpp
@@ -3574,7 +3574,7 @@ void ALandscape::SetDisableRuntimeGrassMapGeneration(bool bInDisableRuntimeGrass
 
 void ALandscapeProxy::OnFeatureLevelChanged(ERHIFeatureLevel::Type NewFeatureLevel)
 {
-	FlushGrassComponents();
+	FlushGrassComponents(nullptr, /*bFlushGrassMaps=*/ false); // rebuild grass instances, but keep the grass maps
 
 	UpdateAllComponentMaterialInstances();
 
