Fixed case where landscape creates mesh batches with a null vertex factory

diff --git a/Engine/Source/Runtime/Landscape/Private/LandscapeRender.cpp b/Engine/Source/Runtime/Landscape/Private/LandscapeRender.cpp
--- a/Engine/Source/Runtime/Landscape/Private/LandscapeRender.cpp
+++ b/Engine/Source/Runtime/Landscape/Private/LandscapeRender.cpp
@@ -2475,6 +2475,7 @@ void FLandscapeComponentSceneProxy::DrawStaticElements(FStaticPrimitiveDrawInter
 	}
 
 	// Add fixed grid mesh batch for rendering the water info texture
+	if (FixedGridVertexFactory)
 	{
 		int32 LODIndex = 0;
 		int32 MaterialIndex = LODIndexToMaterialIndex[LODIndex];
@@ -2488,6 +2489,7 @@ void FLandscapeComponentSceneProxy::DrawStaticElements(FStaticPrimitiveDrawInter
 	}
 
 	// add fixed grid for lumen card captures
+	if (FixedGridVertexFactory)
 	{
 		FMeshBatch MeshBatch;
 
