Modifies Landscape to not add packages to ModifiedPackages if the package is already dirty.
Removes package from ModifiedPackages when a package is dirty.

diff --git a/Engine/Source/Runtime/Landscape/Private/Landscape.cpp b/Engine/Source/Runtime/Landscape/Private/Landscape.cpp
--- a/Engine/Source/Runtime/Landscape/Private/Landscape.cpp
+++ b/Engine/Source/Runtime/Landscape/Private/Landscape.cpp
@@ -4978,6 +4978,12 @@ bool ULandscapeInfo::TryAddToModifiedPackages(UPackage* InPackage, const ALandsc
 		return false;
 	}
 
+	// No need to add the package to ModifiedPackages if it's already dirty.
+	if (InPackage->IsDirty())
+	{
+		return false;
+	}
+
 	// Don't consider unsaved packages as modified/not dirty because they will be saved later on anyway. What we're really after are existing packages made dirty on load
 	if (FPackageName::IsTempPackage(InPackage->GetName()))
 	{
@@ -5043,6 +5049,8 @@ bool ULandscapeInfo::ModifyObject(UObject* InObject, bool bAlwaysMarkDirty)
 		if (LocalLandscapeActor->HasLandscapeEdMode())
 		{
 			InObject->Modify(true);
+			// We just marked the package dirty, no need to keep track of it with ModifiedPackages.
+			ModifiedPackages.Remove(InObject->GetPackage());
 		}
 		else 
 		{
