World Partition - HLOD: When dealing with landscape HLODs that are nanite enabled,we add a skirt around the edge the the generated mesh. Make sure to use a custom depth for this skirt, that is the length of a quad at the exported resolution. This should cover all possible gap scenario while also avoiding the skirt clipping through neighborhood LS tiles.

diff --git a/Engine/Source/Runtime/Landscape/Private/LandscapeHLODBuilder.cpp b/Engine/Source/Runtime/Landscape/Private/LandscapeHLODBuilder.cpp
--- a/Engine/Source/Runtime/Landscape/Private/LandscapeHLODBuilder.cpp
+++ b/Engine/Source/Runtime/Landscape/Private/LandscapeHLODBuilder.cpp
@@ -386,7 +386,17 @@ TArray<UActorComponent*> ULandscapeHLODBuilder::Build(const FHLODBuildContext& I
 	
 			ALandscapeProxy::FRawMeshExportParams ExportParams;
 			ExportParams.ExportLOD = LandscapeLOD;
-		    ExportParams.SkirtDepth = LandscapeProxy->IsNaniteEnabled() ? LandscapeProxy->GetNaniteSkirtDepth() : TOptional<float>();
+
+			// Always add a skirt when dealing with a Nanite landscape, as we'll not be able to avoid a gap when dealing with a lower landscape LOD in HLOD
+			if (LandscapeProxy->IsNaniteEnabled())
+			{
+				// Use a full tile size (at the ExportLOD LOD) as the skirt depth, this will cover all possible gap scenario
+				// and avoid the skirt clipping through neighborhood tiles/HLODs
+				const int32 ComponentSizeVerts = (LandscapeProxy->ComponentSizeQuads + 1) >> LandscapeLOD;
+				const float ScaleFactor = (float)LandscapeProxy->ComponentSizeQuads / (float)(ComponentSizeVerts - 1);
+				ExportParams.SkirtDepth = ScaleFactor;
+			}
+
 			LandscapeProxy->ExportToRawMesh(ExportParams, *MeshDescription);
 
 			StaticMesh->CommitMeshDescription(0);
