Landscape : Nanite skirts generated on a more general set of edges ( no neighbouring component in set of exported components)

diff --git a/Engine/Source/Runtime/Landscape/Private/LandscapeEdit.cpp b/Engine/Source/Runtime/Landscape/Private/LandscapeEdit.cpp
--- a/Engine/Source/Runtime/Landscape/Private/LandscapeEdit.cpp
+++ b/Engine/Source/Runtime/Landscape/Private/LandscapeEdit.cpp
@@ -4136,19 +4136,28 @@ bool ALandscapeProxy::ExportToRawMeshDataCopy(const FRawMeshExportParams& InExpo
 		1.0f / (ExportedSectionMaxQuadCoord.X - ExportedSectionMinQuadCoord.X + LightmapEdgePadding * 2.0f),
 		1.0f / (ExportedSectionMaxQuadCoord.Y - ExportedSectionMinQuadCoord.Y + LightmapEdgePadding * 2.0f));
 	FVector2f LightmapUVOffset = (-ExportedSectionMinQuadCoord + LightmapEdgePadding) * LightmapUVScale;
-
-
+	
 	const bool bExportSkirt = InExportParams.SkirtDepth.IsSet();
 	for (ULandscapeComponent* Component : ComponentsToExport)
 	{
-		// Only generate a skirt around the edge of the proxy not around each component
-		const int32 MinXPadding = bExportSkirt && Component->SectionBaseX == MinSectionBase.X ? 1 : 0;
-		const int32 MaxXPadding = bExportSkirt && Component->SectionBaseX == MaxSectionBase.X ? 1 : 0;
+		TStaticArray<ULandscapeComponent*, 9> NeighborComponents;
+		Component->GetLandscapeComponentNeighbors3x3(NeighborComponents);
+
+		// Export an edge if there isn't a neighoring component or if the neighbor is not a part of set of components
+		// to export
+		auto ExportEdgePadding = [&NeighborComponents, &ComponentsToExport, bExportSkirt] (int32 Edge)
+		{
+			const ULandscapeComponent* Neighbor = NeighborComponents[Edge];
+			return bExportSkirt && (!Neighbor || !ComponentsToExport.Contains(Neighbor)) ? 1 : 0;
+		};
+		
+		const int32 MinXPadding = ExportEdgePadding(3); // left
+		const int32 MaxXPadding = ExportEdgePadding(5); // right
 
 		const int32 XPadding = MinXPadding + MaxXPadding;
 		
-		const int32 MinYPadding = bExportSkirt && Component->SectionBaseY == MinSectionBase.Y ? 1 : 0;
-		const int32 MaxYPadding = bExportSkirt && Component->SectionBaseY == MaxSectionBase.Y ? 1 : 0;
+		const int32 MinYPadding = ExportEdgePadding(1); // bottom
+		const int32 MaxYPadding = ExportEdgePadding(7); // top
 
 		const int32 YPadding = MinYPadding + MaxYPadding;
 
