PR #11613: Use triangle-precise placement of landscape grass instances to fix issue with Landscape Grass floating over/under landscape surface.

#tests I validated landscape grass instance placements are much more accurate in editor
#11613

diff --git a/Engine/Source/Runtime/Landscape/Private/LandscapeGrass.cpp b/Engine/Source/Runtime/Landscape/Private/LandscapeGrass.cpp
--- a/Engine/Source/Runtime/Landscape/Private/LandscapeGrass.cpp
+++ b/Engine/Source/Runtime/Landscape/Private/LandscapeGrass.cpp
@@ -2299,7 +2299,7 @@ struct FAsyncGrassBuilder : public FGrassBuilderBase
 			FMath::Lerp(SampleWeight12, SampleWeight22, LerpX),
 			LerpY);
 
-		// Bilinear interpolate sampled heights
+		// Sample quad corner heights
 		const float SampleHeight11 = GrassData.GetHeight(IdxX1, IdxY1);
 		const float SampleHeight21 = GrassData.GetHeight(IdxX2, IdxY1);
 		const float SampleHeight12 = GrassData.GetHeight(IdxX1, IdxY2);
@@ -2307,11 +2307,23 @@ struct FAsyncGrassBuilder : public FGrassBuilderBase
 
 		OutLocation.X = InLocation.X - DrawScale.X * float(LandscapeSectionOffset.X);
 		OutLocation.Y = InLocation.Y - DrawScale.Y * float(LandscapeSectionOffset.Y);
-		OutLocation.Z = DrawScale.Z * FMath::Lerp(
-			FMath::Lerp(SampleHeight11, SampleHeight21, LerpX),
-			FMath::Lerp(SampleHeight12, SampleHeight22, LerpX),
-			LerpY);
-		
+
+		// Compute height using triangle barycentric coordinates to account for quad triangulation
+		if (LerpX < LerpY)
+		{
+			double A = (1.0 - LerpY);
+			double B = LerpX;
+			double C = 1.0 - A - B;
+			OutLocation.Z = DrawScale.Z * (SampleHeight11 * A + SampleHeight22 * B + SampleHeight12 * C);
+		}
+		else
+		{
+			double A = (1.0 - LerpX);
+			double B = LerpX - LerpY;
+			double C = 1.0 - A - B;
+			OutLocation.Z = DrawScale.Z * (SampleHeight11 * A + SampleHeight21 * B + SampleHeight22 * C);
+		}
+
 		// Compute normal
 		if (OutNormal)
 		{
