Updates landscape actionable messages.

diff --git a/Engine/Source/Runtime/Landscape/Private/LandscapeSubsystem.cpp b/Engine/Source/Runtime/Landscape/Private/LandscapeSubsystem.cpp
--- a/Engine/Source/Runtime/Landscape/Private/LandscapeSubsystem.cpp
+++ b/Engine/Source/Runtime/Landscape/Private/LandscapeSubsystem.cpp
@@ -868,9 +868,9 @@ bool ULandscapeSubsystem::GetActionableMessage(FActionableMessage& OutActionable
 {
 	TRACE_CPUPROFILER_EVENT_SCOPE(ULandscapeSubsystem::GetActionableMessage);
 
-	const FText DefaultMessage = LOCTEXT("LandscapeModifiedText", "Landscape is not up to date");
-	const FText DefaultTooltip = LOCTEXT("LandscapeModifiedToolTip", "External asset modification lead to the landscape not being up to date. Please rebuild the landscape.");
-	const FText DefaultActionMessage = LOCTEXT("LandscapeRebuildText", "Rebuild");
+	const FText DefaultMessage = LOCTEXT("DefaultLandscapeModified.Message", "Landscape is not up to date");
+	const FText DefaultTooltip = LOCTEXT("DefaultLandscapeModified.ToolTip", "Assets that affect the Landscape have changed. Rebuild the Landscape to see the results.");
+	const FText DefaultActionMessage = LOCTEXT("DefaultLandscapeModified.Action", "Rebuild");
 
 	const TArray<ALandscapeProxy*> OutdatedGrassMapProxies = GetOutdatedProxies(UE::Landscape::EOutdatedDataFlags::GrassMaps, /*bInMustMatchAllFlags = */false);
 	const TArray<ALandscapeProxy*> OutdatedPhysicalMaterialProxies = GetOutdatedProxies(UE::Landscape::EOutdatedDataFlags::PhysicalMaterials, /*bInMustMatchAllFlags = */false);
@@ -884,9 +884,9 @@ bool ULandscapeSubsystem::GetActionableMessage(FActionableMessage& OutActionable
 
 	if (HasModifiedLandscapes())
 	{
-		OutActionableMessage.Message = LOCTEXT("LandscapeModifiedText.Dirty", "Landscape contains unsaved changes");
-		OutActionableMessage.Tooltip = DefaultTooltip;
-		OutActionableMessage.ActionMessage = LOCTEXT("LandscapeModifiedText.MarkDirty", "Mark Dirty");
+		OutActionableMessage.Message = LOCTEXT("LandscapeModified.Message", "Landscape assets are out of date");
+		OutActionableMessage.Tooltip = LOCTEXT("LandscapeModified.Tooltip", "The Landscape actors visible in your level have been modified as a result of changes to other assets.\nThese changes need to be applied to the Landscape assets.");
+		OutActionableMessage.ActionMessage = LOCTEXT("LandscapeModified.Action", "Update");
 		OutActionableMessage.ActionCallback = UE::Landscape::MarkModifiedLandscapesAsDirty;
 
 		return true;
