Adjust DDC stats to be more representative of DDC actions. By default COOK_STAT timers will record a miss unless `AddHit` is called. Here we explicitly only track cycles if we know we won't be using DDC (and thus can't miss)

diff --git a/Engine/Source/Runtime/Landscape/Private/LandscapeCollision.cpp b/Engine/Source/Runtime/Landscape/Private/LandscapeCollision.cpp
--- a/Engine/Source/Runtime/Landscape/Private/LandscapeCollision.cpp
+++ b/Engine/Source/Runtime/Landscape/Private/LandscapeCollision.cpp
@@ -1199,7 +1199,9 @@ bool ULandscapeHeightfieldCollisionComponent::CookCollisionData(const FName& For
 	}
 
 	COOK_STAT(auto Timer = LandscapeCollisionCookStats::HeightfieldUsageStats.TimeSyncWork());
-	 
+	// If we aren't using DDC, only track time spent, so we aren't affecting hit/miss stats
+	COOK_STAT(if (GLandscapeCollisionSkipDDC) { Timer.TrackCyclesOnly(); });
+
 	// we have 2 versions of collision objects
 	const int32 CookedDataIndex = bUseDefaultMaterialOnly ? 0 : 1;
 
