Landscape: Fix mobile texture arrays

* Batching the copy of textures to arrays didn't take into account the target slice
* shader variations cached for landscape mobile texture enabled / disabled.

diff --git a/Engine/Source/Runtime/Landscape/Private/Landscape.cpp b/Engine/Source/Runtime/Landscape/Private/Landscape.cpp
--- a/Engine/Source/Runtime/Landscape/Private/Landscape.cpp
+++ b/Engine/Source/Runtime/Landscape/Private/Landscape.cpp
@@ -79,6 +79,7 @@ Landscape.cpp: Terrain rendering
 #include "Misc/ArchiveMD5.h"
 #include "LandscapeEditLayer.h"
 #include "LandscapeTextureStorageProvider.h"
+#include "LandscapeUtils.h"
 #include "LandscapeVersion.h"
 #include "UObject/FortniteMainBranchObjectVersion.h"
 #include "UObject/FortniteReleaseBranchCustomObjectVersion.h"
@@ -638,6 +639,9 @@ void ULandscapeComponent::CheckGenerateMobilePlatformData(bool bIsCooking, const
 	FString MobileVersion = FDevSystemGuids::GetSystemGuid(FDevSystemGuids::Get().LANDSCAPE_MOBILE_COOK_VERSION).ToString();
 	ComponentStateAr << MobileVersion;
 
+	bool IsTextureArrayEnabled = UE::Landscape::IsMobileWeightmapTextureArrayEnabled();
+	ComponentStateAr << IsTextureArrayEnabled;
+
 	uint32 Hash[5];
 	FSHA1::HashBuffer(ComponentStateAr.GetData(), ComponentStateAr.Num(), (uint8*)Hash);
 	FGuid NewSourceHash = FGuid(Hash[0] ^ Hash[4], Hash[1], Hash[2], Hash[3]);
diff --git a/Engine/Source/Runtime/Landscape/Private/LandscapeUtils.cpp b/Engine/Source/Runtime/Landscape/Private/LandscapeUtils.cpp
--- a/Engine/Source/Runtime/Landscape/Private/LandscapeUtils.cpp
+++ b/Engine/Source/Runtime/Landscape/Private/LandscapeUtils.cpp
@@ -110,13 +110,14 @@ bool IsVisibilityLayer(const ULandscapeLayerInfoObject* InLayerInfoObject)
 
 uint32 GetTypeHash(const FTextureCopyRequest& InKey)
 {
-	uint32 Hash = GetTypeHash(InKey.Source);
-	return HashCombine(Hash, GetTypeHash(InKey.Destination));
+	uint32 Hash = ::GetTypeHash(InKey.Source);
+	uint32 HashSlice = ::GetTypeHash(InKey.DestinationSlice);
+	return HashCombine(Hash, ::GetTypeHash(InKey.Destination));
 }
 
 bool operator==(const FTextureCopyRequest& InEntryA, const FTextureCopyRequest& InEntryB)
 {
-	return (InEntryA.Source == InEntryB.Source) && (InEntryA.Destination == InEntryB.Destination);
+	return (InEntryA.Source == InEntryB.Source) && (InEntryA.Destination == InEntryB.Destination) && (InEntryA.DestinationSlice == InEntryB.DestinationSlice);
 }
 
 bool FBatchTextureCopy::AddWeightmapCopy(UTexture* InDestination, int8 InDestinationSlice, int8 InDestinationChannel, const ULandscapeComponent* InComponent, ULandscapeLayerInfoObject* InLayerInfo)
@@ -270,7 +271,7 @@ static FAutoConsoleVariableRef CVarLandscapeMobileWeightTextureArray(
 	TEXT("landscape.MobileWeightTextureArray"),
 	LandscapeMobileWeightTextureArray,
 	TEXT("Use Texture Arrays for weights on Mobile platforms"),
-	ECVF_ReadOnly);
+	ECVF_ReadOnly | ECVF_MobileShaderChange);
 
 bool IsMobileWeightmapTextureArrayEnabled()
 {
