WorldPartitionRuntimeVirtualTextureBuilder: Fix Spatially loaded landscape proxies not being loaded, resulting in them not being included in the virtual texture.

diff --git a/Engine/Source/Runtime/Landscape/Classes/LandscapeStreamingProxy.h b/Engine/Source/Runtime/Landscape/Classes/LandscapeStreamingProxy.h
--- a/Engine/Source/Runtime/Landscape/Classes/LandscapeStreamingProxy.h
+++ b/Engine/Source/Runtime/Landscape/Classes/LandscapeStreamingProxy.h
@@ -47,6 +47,7 @@ public:
 	virtual bool GetReferencedContentObjects(TArray<UObject*>& Objects) const override;
 	virtual bool CanChangeIsSpatiallyLoadedFlag() const override { return AActor::CanChangeIsSpatiallyLoadedFlag(); }
 	virtual bool ShouldIncludeGridSizeInName(UWorld* InWorld, const FActorPartitionIdentifier& InIdentifier) const override;
+	virtual void GetActorDescProperties(FPropertyPairsMap& PropertyPairsMap) const override;
 #endif
 	//~ End UObject Interface
 
diff --git a/Engine/Source/Runtime/Landscape/Private/LandscapeEdit.cpp b/Engine/Source/Runtime/Landscape/Private/LandscapeEdit.cpp
--- a/Engine/Source/Runtime/Landscape/Private/LandscapeEdit.cpp
+++ b/Engine/Source/Runtime/Landscape/Private/LandscapeEdit.cpp
@@ -6320,6 +6320,20 @@ bool ALandscapeStreamingProxy::GetReferencedContentObjects(TArray<UObject*>& Obj
 	return true;
 }
 
+void ALandscapeStreamingProxy::GetActorDescProperties(FPropertyPairsMap& PropertyPairsMap) const
+{
+	Super::GetActorDescProperties(PropertyPairsMap);
+
+	for (URuntimeVirtualTexture* RuntimeVirtualTexture : RuntimeVirtualTextures)
+	{
+		if (RuntimeVirtualTexture)
+		{
+			PropertyPairsMap.AddProperty(UPrimitiveComponent::RVTActorDescProperty);
+			return;
+		}
+	}
+}
+
 bool ALandscapeStreamingProxy::ShouldIncludeGridSizeInName(UWorld* InWorld, const FActorPartitionIdentifier& InIdentifier) const
 {
 	// Always return true if this world setting flag is true
