Added AllowedDensityRange to landscape Grass Variety. This can be used to prevent a variety from ever spawning on a given range of grass density (specified as an interval between 0 and 1)

diff --git a/Engine/Source/Runtime/Landscape/Classes/LandscapeGrassType.h b/Engine/Source/Runtime/Landscape/Classes/LandscapeGrassType.h
--- a/Engine/Source/Runtime/Landscape/Classes/LandscapeGrassType.h
+++ b/Engine/Source/Runtime/Landscape/Classes/LandscapeGrassType.h
@@ -68,6 +68,7 @@ struct LANDSCAPE_API FGrassVariety
 
 	UPROPERTY(BlueprintReadOnly, EditAnywhere, Category = Grass, meta = (UIMin = 0, ClampMin = 0, UIMax = 1000000, ClampMax = 1000000))
 	FPerQualityLevelInt EndCullDistanceQuality;
+
 	/** 
 	 * Specifies the smallest LOD that will be used for this component.
 	 * If -1 (default), the MinLOD of the static mesh asset will be used instead.
@@ -75,6 +76,10 @@ struct LANDSCAPE_API FGrassVariety
 	UPROPERTY(BlueprintReadOnly, EditAnywhere, Category = Grass, meta = (UIMin = -1, ClampMin = -1, UIMax = 8, ClampMax = 8))
 	int32 MinLOD;
 
+	/** Specifies the density range where the grass variety is allowed to be spawned ([0,1] represents the entire range). */
+	UPROPERTY(BlueprintReadOnly, EditAnywhere, Category=Grass, meta = (UIMin = 0.0, ClampMin = 0.0, UIMax = 1.0, ClampMax = 1.0))
+	FFloatInterval AllowedDensityRange;
+
 	/** Specifies grass instance scaling type */
 	UPROPERTY(BlueprintReadOnly, EditAnywhere, Category=Grass)
 	EGrassScaling Scaling;
diff --git a/Engine/Source/Runtime/Landscape/Private/LandscapeGrass.cpp b/Engine/Source/Runtime/Landscape/Private/LandscapeGrass.cpp
--- a/Engine/Source/Runtime/Landscape/Private/LandscapeGrass.cpp
+++ b/Engine/Source/Runtime/Landscape/Private/LandscapeGrass.cpp
@@ -1242,6 +1242,7 @@ FGrassVariety::FGrassVariety()
 	, EndCullDistance(10000)
 	, EndCullDistanceQuality(10000)
 	, MinLOD(-1)
+	, AllowedDensityRange(0.0f, 1.0f)
 	, Scaling(EGrassScaling::Uniform)
 	, ScaleX(1.0f, 1.0f)
 	, ScaleY(1.0f, 1.0f)
@@ -1832,6 +1833,7 @@ struct FAsyncGrassBuilder : public FGrassBuilderBase
 	FFloatInterval ScaleX;
 	FFloatInterval ScaleY;
 	FFloatInterval ScaleZ;
+	FFloatInterval AllowedDensityRange;
 	bool bWeightAttenuatesMaxScale;
 	float MaxScaleWeightAttenuation ;
 	bool bRandomRotation;
@@ -1870,6 +1872,7 @@ struct FAsyncGrassBuilder : public FGrassBuilderBase
 		, ScaleX(GrassVariety.ScaleX)
 		, ScaleY(GrassVariety.ScaleY)
 		, ScaleZ(GrassVariety.ScaleZ)
+		, AllowedDensityRange(GrassVariety.AllowedDensityRange)
 		, bWeightAttenuatesMaxScale(GrassVariety.bWeightAttenuatesMaxScale)
 		, MaxScaleWeightAttenuation(GrassVariety.MaxScaleWeightAttenuation)
 		, bRandomRotation(GrassVariety.RandomRotation)
@@ -2103,7 +2106,9 @@ struct FAsyncGrassBuilder : public FGrassBuilderBase
 				FVector ComputedNormal;
 				float Weight = 0.f;
 				SampleLandscapeAtLocationLocal(Location, LocationWithHeight, Weight, bAlignToSurface ? &ComputedNormal : nullptr);
-				bool bKeep = Weight > 0.0f && Weight >= RandomStream.GetFraction() && !IsExcluded(LocationWithHeight);
+				bool bKeep = (Weight > AllowedDensityRange.Min) && (Weight <= AllowedDensityRange.Max) 
+					&& (Weight >= RandomStream.GetFraction()) 
+					&& !IsExcluded(LocationWithHeight);
 				if (bKeep)
 				{
 					const FVector Scale = bRandomScale ? GetRandomScale(Weight) : DefaultScale;
@@ -2169,7 +2174,9 @@ struct FAsyncGrassBuilder : public FGrassBuilderBase
 						FInstanceLocal& Instance = Instances[InstanceIndex];
 						float Weight = 0.f;
 						SampleLandscapeAtLocationLocal(Location, Instance.Pos, Weight);
-						Instance.bKeep = Weight > 0.0f && Weight >= RandomStream.GetFraction() && !IsExcluded(Instance.Pos);
+						Instance.bKeep = (Weight > AllowedDensityRange.Min) && (Weight <= AllowedDensityRange.Max) 
+							&& (Weight >= RandomStream.GetFraction()) 
+							&& !IsExcluded(Instance.Pos);
 						Instance.Weight = Weight;
 						if (Instance.bKeep)
 						{
