Fix an issue introduced by 32833605 with landscape materials with displacement.

diff --git a/Engine/Source/Runtime/Landscape/Private/LandscapeRender.cpp b/Engine/Source/Runtime/Landscape/Private/LandscapeRender.cpp
--- a/Engine/Source/Runtime/Landscape/Private/LandscapeRender.cpp
+++ b/Engine/Source/Runtime/Landscape/Private/LandscapeRender.cpp
@@ -4796,6 +4796,7 @@ public:
 		}
 
 		bool bAnySectionMasked = false;
+		bHasVertexProgrammableRaster = false;
 		bHasPixelProgrammableRaster = false;
 
 		for (::Nanite::FSceneProxyBase::FMaterialSection& MaterialSection : MaterialSections)
@@ -4803,6 +4804,11 @@ public:
 			const bool bWasMasked = MaterialSection.MaterialRelevance.bMasked;
 			MaterialSection.MaterialRelevance.bMasked = false;
 
+			if (MaterialSection.IsVertexProgrammableRaster(bEvaluateWorldPositionOffset))
+			{
+				bHasVertexProgrammableRaster = true;
+			}
+			
 			if (MaterialSection.IsPixelProgrammableRaster())
 			{
 				// Don't change bMasked if it is not the sole factor that makes the material section programmable
@@ -4810,7 +4816,8 @@ public:
 				bAnySectionMasked |= bWasMasked;
 				bHasPixelProgrammableRaster = true;
 			}
-			else
+			
+			if (!bHasVertexProgrammableRaster && !bHasPixelProgrammableRaster)
 			{
 				MaterialSection.ResetToDefaultMaterial(false, true);
 			}
