Fix another subtle issue I created with Nanite landscape rendering.

diff --git a/Engine/Source/Runtime/Landscape/Private/LandscapeRender.cpp b/Engine/Source/Runtime/Landscape/Private/LandscapeRender.cpp
--- a/Engine/Source/Runtime/Landscape/Private/LandscapeRender.cpp
+++ b/Engine/Source/Runtime/Landscape/Private/LandscapeRender.cpp
@@ -4804,20 +4804,22 @@ public:
 			const bool bWasMasked = MaterialSection.MaterialRelevance.bMasked;
 			MaterialSection.MaterialRelevance.bMasked = false;
 
+			bool bProgrammableRasterMaterial = false;
 			if (MaterialSection.IsVertexProgrammableRaster(bEvaluateWorldPositionOffset))
 			{
 				bHasVertexProgrammableRaster = true;
+				bProgrammableRasterMaterial = true;
 			}
-			
 			if (MaterialSection.IsPixelProgrammableRaster())
 			{
 				// Don't change bMasked if it is not the sole factor that makes the material section programmable
 				MaterialSection.MaterialRelevance.bMasked = bWasMasked;
 				bAnySectionMasked |= bWasMasked;
 				bHasPixelProgrammableRaster = true;
+				bProgrammableRasterMaterial = true;
 			}
 			
-			if (!bHasVertexProgrammableRaster && !bHasPixelProgrammableRaster)
+			if (!bProgrammableRasterMaterial)
 			{
 				MaterialSection.ResetToDefaultMaterial(false, true);
 			}
