Fix for crash when migrating levels containing landscape.
When landscapes are saved without having registered with a world, we can't build grass maps, as they aren't renderable.

diff --git a/Engine/Source/Runtime/Landscape/Private/LandscapeGrassMapsBuilder.cpp b/Engine/Source/Runtime/Landscape/Private/LandscapeGrassMapsBuilder.cpp
--- a/Engine/Source/Runtime/Landscape/Private/LandscapeGrassMapsBuilder.cpp
+++ b/Engine/Source/Runtime/Landscape/Private/LandscapeGrassMapsBuilder.cpp
@@ -795,12 +795,12 @@ bool FLandscapeGrassMapsBuilder::BuildGrassMapsNowForComponents(
 	int32 UpToDateCount = 0;
 
 	// track any components that have failed to build
-	TSet<FComponentState*> FailedStates;
+	TSet<ULandscapeComponent*> FailedComponents;
 
 	const double StartTime = FPlatformTime::Seconds();
 	double LastFlush = StartTime;
 	double LastChangeTime = StartTime;
-	while (UpToDateCount + FailedStates.Num() != LandscapeComponents.Num())
+	while (UpToDateCount + FailedComponents.Num() != LandscapeComponents.Num())
 	{
 		// ensure we are making progress within a reasonable amount of time TODO [chris.tchou] there should be a better way to detect non-progress here
 		const double CurTime = FPlatformTime::Seconds();
@@ -819,7 +819,11 @@ bool FLandscapeGrassMapsBuilder::BuildGrassMapsNowForComponents(
 		for (ULandscapeComponent* Component : LandscapeComponents)
 		{
 			FComponentState* State = ComponentStates.FindRef(Component);
-			check(State); // components must be registered
+			if (State == nullptr)
+			{
+				FailedComponents.Add(Component);
+				continue;
+			}
 
 			if (State->Stage == EComponentStage::Pending)
 			{
@@ -848,16 +852,16 @@ bool FLandscapeGrassMapsBuilder::BuildGrassMapsNowForComponents(
 				// guaranteed by UpdateTrackedComponents(), as long as we update all of the components
 				check(Component->ComputeGrassMapGenerationHash() == Component->GrassData->GenerationHash);
 #endif // WITH_EDITOR
-				if (FailedStates.Contains(State))
+				if (FailedComponents.Contains(Component))
 				{
-					FailedStates.Remove(State);
+					FailedComponents.Remove(Component);
 				}
 				UpToDateCount++;
 			}
 			if (State->Stage == EComponentStage::NotReady)
 			{
 				// if it's not ready because of shader reasons
-				if (!FailedStates.Contains(State) && !UE::Landscape::CanRenderGrassMap(Component))
+				if (!FailedComponents.Contains(Component) && !UE::Landscape::CanRenderGrassMap(Component))
 				{
 #if WITH_EDITOR
 					// in editor, try to force compilation to complete
@@ -866,7 +870,7 @@ bool FLandscapeGrassMapsBuilder::BuildGrassMapsNowForComponents(
 #endif // WITH_EDITOR
 					{
 						// failed to compile shaders... we won't be able to build the grass map for this component
-						FailedStates.Add(State);
+						FailedComponents.Add(Component);
 					}
 				}
 			}
@@ -906,8 +910,9 @@ bool FLandscapeGrassMapsBuilder::BuildGrassMapsNowForComponents(
 	}
 
 	UE_LOG(LogGrass, Verbose, TEXT("BuildGrassMapsNowForComponents() updated %d/%d components in %f seconds"), UpToDateCount, LandscapeComponents.Num(), FPlatformTime::Seconds() - StartTime);
-
-	if (UpToDateCount != LandscapeComponents.Num())
+	
+	// warn if we failed to build grass maps, except when there are no registered states (which happens when we migrate levels from project to project - because it doesn't register before saving)
+	if ((UpToDateCount != LandscapeComponents.Num()) && (ComponentStates.Num() > 0))
 	{
 		UE_LOG(LogGrass, Warning, TEXT("Failed to build grass maps for %d/%d landscape components, check if you are using a render preview mode, or a non-SM5 capable render device.  (%d pending, %d streaming, %d rendering, %d fetching, %d built)"),
 			LandscapeComponents.Num() - UpToDateCount,
