Adjustments to actionable (landscape) messages UX:
* Added tooltip when the widget is collapsed to help the user understand what it's about without having to expand it
* Show the number of affected landscape actors even in the "dirty actors" case to help the user assess how much "out-of-date" the overall landscape is

Misc:
* Dirty actors are now handled part of out-of-date landscape proxy flags
* Optimized update by gathering all out-of-date flags at once (~33% of what it used to cost)
#tests editor

diff --git a/Engine/Source/Editor/ActionableMessage/Private/SActionableMessageViewportWidget.cpp b/Engine/Source/Editor/ActionableMessage/Private/SActionableMessageViewportWidget.cpp
--- a/Engine/Source/Editor/ActionableMessage/Private/SActionableMessageViewportWidget.cpp
+++ b/Engine/Source/Editor/ActionableMessage/Private/SActionableMessageViewportWidget.cpp
@@ -122,6 +122,15 @@ void SActionableMessageViewportWidget::Construct(const FArguments& InArgs)
 			SNew(SButton)
 			.ButtonStyle(&FCoreStyle::Get().GetWidgetStyle<FButtonStyle>("HoverHintOnly"))
 			.IsEnabled(true)
+			.ToolTipText_Lambda([this]()
+			{
+				if (!bExpanded)
+				{
+					return LOCTEXT("ActionableMessages.CollapsedWarning", "Your project has warnings that need to be resolved. Click for details.");
+				}
+
+				return FText();
+			})
 			.OnClicked_Lambda([this]()
 			{
 				bExpanded = !bExpanded;
diff --git a/Engine/Source/Runtime/Landscape/Classes/LandscapeInfo.h b/Engine/Source/Runtime/Landscape/Classes/LandscapeInfo.h
--- a/Engine/Source/Runtime/Landscape/Classes/LandscapeInfo.h
+++ b/Engine/Source/Runtime/Landscape/Classes/LandscapeInfo.h
@@ -263,6 +263,7 @@ public:
 	LANDSCAPE_API void OnModifiedPackageSaved(UPackage* InPackage);
 	LANDSCAPE_API int32 GetModifiedPackageCount() const;
 	LANDSCAPE_API TArray<UPackage*> GetModifiedPackages() const;
+	LANDSCAPE_API bool IsPackageModified(UPackage* InPackage) const;
 	LANDSCAPE_API void MarkModifiedPackagesAsDirty();
 
 	/** Landscapes are a bit special in that they contain derived data within their source data, which can be updated on-the-fly depending on many external factors (e.g. a procedural edit layer can invalidate some 
diff --git a/Engine/Source/Runtime/Landscape/Private/Landscape.cpp b/Engine/Source/Runtime/Landscape/Private/Landscape.cpp
--- a/Engine/Source/Runtime/Landscape/Private/Landscape.cpp
+++ b/Engine/Source/Runtime/Landscape/Private/Landscape.cpp
@@ -4981,6 +4981,11 @@ TArray<UPackage*> ULandscapeInfo::GetModifiedPackages() const
 	return LocalModifiedPackages;
 }
 
+bool ULandscapeInfo::IsPackageModified(UPackage* InPackage) const
+{
+	return ModifiedPackages.Contains(InPackage);
+}
+
 void ULandscapeInfo::MarkModifiedPackagesAsDirty()
 {
 	for (TWeakObjectPtr<UPackage> WeakPackagePtr : ModifiedPackages)
@@ -6340,6 +6345,14 @@ UE::Landscape::EOutdatedDataFlags ALandscapeProxy::GetOutdatedDataFlags() const
 		OutdatedDataFlags |= UE::Landscape::EOutdatedDataFlags::NaniteMeshes;
 	}
 
+	if (ULandscapeInfo* Info = GetLandscapeInfo())
+	{
+		if (Info->IsPackageModified(GetPackage()))
+		{
+			OutdatedDataFlags |= UE::Landscape::EOutdatedDataFlags::PackageModified;
+		}
+	}
+
 	return OutdatedDataFlags;
 }
 
diff --git a/Engine/Source/Runtime/Landscape/Private/LandscapeSubsystem.cpp b/Engine/Source/Runtime/Landscape/Private/LandscapeSubsystem.cpp
--- a/Engine/Source/Runtime/Landscape/Private/LandscapeSubsystem.cpp
+++ b/Engine/Source/Runtime/Landscape/Private/LandscapeSubsystem.cpp
@@ -29,6 +29,7 @@
 #include "Engine/Canvas.h"
 #include "EngineUtils.h"
 #include "Misc/ScopedSlowTask.h"
+#include "Algo/ForEach.h"
 #include "Algo/Transform.h"
 #include "Algo/RemoveIf.h"
 #include "Algo/Unique.h"
@@ -634,8 +635,11 @@ void ULandscapeSubsystem::BuildPhysicalMaterial()
 	PhysicalMaterialBuilder->Rebuild();
 }
 
+// Deprecated :
 TArray<ALandscapeProxy*> ULandscapeSubsystem::GetOutdatedProxies(UE::Landscape::EOutdatedDataFlags InMatchingOutdatedDataFlags, bool bInMustMatchAllFlags) const
 {
+	TRACE_CPUPROFILER_EVENT_SCOPE(ULandscapeSubsystem::GetOutdatedProxies);
+
 	UWorld* World = GetWorld();
 	if (!World || World->IsGameWorld())
 	{
@@ -656,6 +660,33 @@ TArray<ALandscapeProxy*> ULandscapeSubsystem::GetOutdatedProxies(UE::Landscape::
 	return FinalProxiesToBuild;
 }
 
+TArray<TTuple<ALandscapeProxy*, UE::Landscape::EOutdatedDataFlags>> ULandscapeSubsystem::GetOutdatedProxyDetails(UE::Landscape::EOutdatedDataFlags InMatchingOutdatedDataFlags, bool bInMustMatchAllFlags) const
+{
+	TRACE_CPUPROFILER_EVENT_SCOPE(ULandscapeSubsystem::GetOutdatedProxyDetails);
+
+	UWorld* World = GetWorld();
+	if (!World || World->IsGameWorld())
+	{
+		return {};
+	}
+
+	TArray<TTuple<ALandscapeProxy*, UE::Landscape::EOutdatedDataFlags>> OutdatedProxies;
+	OutdatedProxies.Reserve(Proxies.Num());
+	for (TWeakObjectPtr<ALandscapeProxy> Proxy : Proxies)
+	{
+		if (ALandscapeProxy* ValidProxy = Proxy.Get())
+		{
+			const UE::Landscape::EOutdatedDataFlags ProxyOutdatedDataFlags = ValidProxy->GetOutdatedDataFlags();
+			if ((bInMustMatchAllFlags && EnumHasAllFlags(ProxyOutdatedDataFlags, InMatchingOutdatedDataFlags))
+				|| (!bInMustMatchAllFlags && EnumHasAnyFlags(ProxyOutdatedDataFlags, InMatchingOutdatedDataFlags)))
+			{
+				OutdatedProxies.Add({ValidProxy, ProxyOutdatedDataFlags});
+			}
+		}
+	}
+	return OutdatedProxies;
+}
+
 void ULandscapeSubsystem::BuildNanite(TArrayView<ALandscapeProxy*> InProxiesToBuild, bool bForceRebuild)
 {
 	TRACE_BOOKMARK(TEXT("ULandscapeSubsystem::BuildNanite"));
@@ -866,45 +897,45 @@ void ULandscapeSubsystem::DisplayMessages(FCanvas* Canvas, float& XPos, float& Y
 
 bool ULandscapeSubsystem::GetActionableMessage(FActionableMessage& OutActionableMessage)
 {
+	using namespace UE::Landscape;
+
 	TRACE_CPUPROFILER_EVENT_SCOPE(ULandscapeSubsystem::GetActionableMessage);
 
-	const FText DefaultMessage = LOCTEXT("DefaultLandscapeModified.Message", "Landscape is not up to date");
-	const FText DefaultTooltip = LOCTEXT("DefaultLandscapeModified.ToolTip", "Assets that affect the Landscape have changed. Rebuild the Landscape to see the results.");
+	const FText DefaultTooltip = LOCTEXT("DefaultLandscapeModified.ToolTip", "Assets that affect the Landscape may have changed. Rebuild the Landscape to see the results.");
 	const FText DefaultActionMessage = LOCTEXT("DefaultLandscapeModified.Action", "Rebuild");
 
-	const TArray<ALandscapeProxy*> OutdatedGrassMapProxies = GetOutdatedProxies(UE::Landscape::EOutdatedDataFlags::GrassMaps, /*bInMustMatchAllFlags = */false);
-	const TArray<ALandscapeProxy*> OutdatedPhysicalMaterialProxies = GetOutdatedProxies(UE::Landscape::EOutdatedDataFlags::PhysicalMaterials, /*bInMustMatchAllFlags = */false);
-	const TArray<ALandscapeProxy*> OutdatedNaniteProxies = GetOutdatedProxies(UE::Landscape::EOutdatedDataFlags::NaniteMeshes, /*bInMustMatchAllFlags = */false);
-	
-	const int32 OutdatedGrassmapProxiesCount = OutdatedGrassMapProxies.Num();
-	const int32 OutdatedPhysicalMaterialProxiesCount = OutdatedPhysicalMaterialProxies.Num();
-	const int32 OutdatedNaniteProxiesCount = OutdatedNaniteProxies.Num();
-
-	const int32 OutdatedFlags = ((OutdatedGrassmapProxiesCount > 0) ? 1 : 0) + ((OutdatedPhysicalMaterialProxiesCount > 0) ? 1 : 0) + ((OutdatedNaniteProxiesCount > 0) ? 1 : 0);
-
-	if (HasModifiedLandscapes())
+	const TArray<TTuple<ALandscapeProxy*, EOutdatedDataFlags>> OutdatedProxies = GetOutdatedProxyDetails(EOutdatedDataFlags::All, /*bInMustMatchAllFlags = */false);
+	TArray<int32> NumOutdatedProxyPerFlag;
+	NumOutdatedProxyPerFlag.AddDefaulted(GetOutdatedDataFlagIndex(EOutdatedDataFlags::Last) + 1);
+	int32 NumTotalOutdatedProxy = 0;
+	EOutdatedDataFlags OutdatedFlagsUnion = EOutdatedDataFlags::None;
+	Algo::ForEach(OutdatedProxies, [&OutdatedProxies, &NumOutdatedProxyPerFlag, &NumTotalOutdatedProxy, &OutdatedFlagsUnion](const TTuple<ALandscapeProxy*, EOutdatedDataFlags>& ProxyAndFlag)
 	{
-		OutActionableMessage.Message = LOCTEXT("LandscapeModified.Message", "Landscape assets are out of date");
-		OutActionableMessage.Tooltip = LOCTEXT("LandscapeModified.Tooltip", "The Landscape actors visible in your level have been modified as a result of changes to other assets.\nThese changes need to be applied to the Landscape assets.");
-		OutActionableMessage.ActionMessage = LOCTEXT("LandscapeModified.Action", "Update");
-		OutActionableMessage.ActionCallback = UE::Landscape::MarkModifiedLandscapesAsDirty;
+		++NumTotalOutdatedProxy;
+		OutdatedFlagsUnion |= ProxyAndFlag.Value;
 
-		return true;
-	}
-	
-	if (OutdatedFlags > 1)
+		uint32 RemainingFlags = static_cast<uint32>(ProxyAndFlag.Value);
+		while (RemainingFlags != 0)
+		{
+			uint32 FlagIndex = FBitSet::GetAndClearNextBit(RemainingFlags);
+			++NumOutdatedProxyPerFlag[FlagIndex];
+		}
+	});
+
+	// If more than 1 action is required, go with a BuildAll action
+	if (FMath::CountBits(static_cast<uint64>(OutdatedFlagsUnion)) > 1)
 	{
-		OutActionableMessage.Message = DefaultMessage;
+		OutActionableMessage.Message = FText::Format(LOCTEXT("SeveralLandscapeDataOutdated.Message", "{0} Landscape {0}|plural(one=actor,other=actors) {0}|plural(one=is,other=are) out of date and {0}|plural(one=needs,other=need) to be rebuilt"), NumTotalOutdatedProxy);;
 		OutActionableMessage.Tooltip = DefaultTooltip;
 		OutActionableMessage.ActionMessage = DefaultActionMessage;
 		OutActionableMessage.ActionCallback = UE::Landscape::BuildAll;
 
 		return true;
 	}
-	
-	if (OutdatedGrassmapProxiesCount > 0)
+
+	if (int32 OutdatedProxiesCount = NumOutdatedProxyPerFlag[GetOutdatedDataFlagIndex(EOutdatedDataFlags::GrassMaps)])
 	{
-		OutActionableMessage.Message = FText::Format(LOCTEXT("GRASS_MAPS_NEED_TO_BE_REBUILT_FMT", "{0} Landscape {0}|plural(one=actor,other=actors) with grass maps {0}|plural(one=needs,other=need) to be rebuilt"), OutdatedGrassmapProxiesCount);
+		OutActionableMessage.Message = FText::Format(LOCTEXT("GRASS_MAPS_NEED_TO_BE_REBUILT_FMT", "{0} Landscape {0}|plural(one=actor,other=actors) with grass maps {0}|plural(one=needs,other=need) to be rebuilt"), OutdatedProxiesCount);
 		OutActionableMessage.Tooltip = DefaultTooltip;
 		OutActionableMessage.ActionMessage = DefaultActionMessage;
 		OutActionableMessage.ActionCallback = UE::Landscape::BuildGrassMaps;
@@ -912,9 +943,9 @@ bool ULandscapeSubsystem::GetActionableMessage(FActionableMessage& OutActionable
 		return true;
 	}
 	
-	if (OutdatedPhysicalMaterialProxiesCount > 0)
+	if (int32 OutdatedProxiesCount = NumOutdatedProxyPerFlag[GetOutdatedDataFlagIndex(EOutdatedDataFlags::PhysicalMaterials)])
 	{
-		OutActionableMessage.Message = FText::Format(LOCTEXT("LANDSCAPE_PHYSICALMATERIAL_NEED_TO_BE_REBUILT_FMT", "{0} Landscape {0}|plural(one=actor,other=actors) with physical materials {0}|plural(one=needs,other=need) to be rebuilt"), OutdatedPhysicalMaterialProxiesCount);
+		OutActionableMessage.Message = FText::Format(LOCTEXT("LANDSCAPE_PHYSICALMATERIAL_NEED_TO_BE_REBUILT_FMT", "{0} Landscape {0}|plural(one=actor,other=actors) with physical materials {0}|plural(one=needs,other=need) to be rebuilt"), OutdatedProxiesCount);
 		OutActionableMessage.Tooltip = DefaultTooltip;
 		OutActionableMessage.ActionMessage = DefaultActionMessage;
 		OutActionableMessage.ActionCallback = UE::Landscape::BuildPhysicalMaterial;
@@ -922,9 +953,9 @@ bool ULandscapeSubsystem::GetActionableMessage(FActionableMessage& OutActionable
 		return true;
 	}
 
-	if (OutdatedNaniteProxiesCount > 0)
+	if (int32 OutdatedProxiesCount = NumOutdatedProxyPerFlag[GetOutdatedDataFlagIndex(EOutdatedDataFlags::NaniteMeshes)])
 	{
-		OutActionableMessage.Message = FText::Format(LOCTEXT("LANDSCAPE_NANITE_MESHES_NEED_TO_BE_REBUILT_FMT", "{0} Landscape {0}|plural(one=actor,other=actors) with Nanite meshes {0}|plural(one=needs,other=need) to be rebuilt"), OutdatedNaniteProxiesCount);
+		OutActionableMessage.Message = FText::Format(LOCTEXT("LANDSCAPE_NANITE_MESHES_NEED_TO_BE_REBUILT_FMT", "{0} Landscape {0}|plural(one=actor,other=actors) with Nanite meshes {0}|plural(one=needs,other=need) to be rebuilt"), OutdatedProxiesCount);
 		OutActionableMessage.Tooltip = DefaultTooltip;
 		OutActionableMessage.ActionMessage = DefaultActionMessage;
 		OutActionableMessage.ActionCallback = UE::Landscape::BuildNanite;
@@ -932,6 +963,16 @@ bool ULandscapeSubsystem::GetActionableMessage(FActionableMessage& OutActionable
 		return true;
 	}
 
+	if (int32 OutdatedProxiesCount = NumOutdatedProxyPerFlag[GetOutdatedDataFlagIndex(EOutdatedDataFlags::PackageModified)])
+	{
+		OutActionableMessage.Message = FText::Format(LOCTEXT("LandscapeModified.Message", "{0} Landscape {0}|plural(one=actor,other=actors) {0}|plural(one=is,other=are) out of date and {0}|plural(one=needs,other=need) to be rebuilt"), OutdatedProxiesCount);
+		OutActionableMessage.Tooltip = LOCTEXT("LandscapeModified.Tooltip", "The Landscape actors visible in your level have been modified as a result of changes to other assets.\nThese changes need to be applied to the Landscape assets.");
+		OutActionableMessage.ActionMessage = LOCTEXT("LandscapeModified.Action", "Update");
+		OutActionableMessage.ActionCallback = UE::Landscape::MarkModifiedLandscapesAsDirty;
+
+		return true;
+	}
+
 	return false;
 }
 
diff --git a/Engine/Source/Runtime/Landscape/Public/LandscapeEditTypes.h b/Engine/Source/Runtime/Landscape/Public/LandscapeEditTypes.h
--- a/Engine/Source/Runtime/Landscape/Public/LandscapeEditTypes.h
+++ b/Engine/Source/Runtime/Landscape/Public/LandscapeEditTypes.h
@@ -21,12 +21,27 @@ namespace UE::Landscape
 enum class EOutdatedDataFlags : uint8
 {
 	None = 0,
+
+	// Actual flags : 
 	GrassMaps = (1 << 0),
 	PhysicalMaterials = (1 << 1),
 	NaniteMeshes = (1 << 2),
-	// TODO [jonathan.bard] : DirtyActors = (1 << 3) ?
-	All = (GrassMaps | PhysicalMaterials | NaniteMeshes /*| DirtyActors*/)
+	PackageModified = (1 << 3),
+
+	// Not real flags, only useful to loop through the actual flags : 
+	LastPlusOne, 
+	Last = LastPlusOne - 1,
+
+	// Combined flags :
+	All = (GrassMaps | PhysicalMaterials | NaniteMeshes | PackageModified)
 };
 ENUM_CLASS_FLAGS(EOutdatedDataFlags);
 
+inline uint32 GetOutdatedDataFlagIndex(EOutdatedDataFlags InFlag)
+{
+	const uint64 InFlagAsUInt64 = static_cast<uint64>(InFlag);
+	check((InFlagAsUInt64 > static_cast<uint64>(EOutdatedDataFlags::None)) && (InFlagAsUInt64 < static_cast<uint64>(EOutdatedDataFlags::LastPlusOne)) && (FMath::CountBits(InFlagAsUInt64) == 1u));
+	return FMath::CountTrailingZeros(InFlagAsUInt64);
+}
+
 } // namespace UE::Landscape
\ No newline at end of file
diff --git a/Engine/Source/Runtime/Landscape/Public/LandscapeSubsystem.h b/Engine/Source/Runtime/Landscape/Public/LandscapeSubsystem.h
--- a/Engine/Source/Runtime/Landscape/Public/LandscapeSubsystem.h
+++ b/Engine/Source/Runtime/Landscape/Public/LandscapeSubsystem.h
@@ -107,7 +107,9 @@ public:
 	 */
 	LANDSCAPE_API void BuildNanite(TArrayView<ALandscapeProxy*> InProxiesToBuild = TArrayView<ALandscapeProxy*>(), bool bForceRebuild = false);
 	
+	UE_DEPRECATED(5.5, "Use GetOutdatedProxyDetails")
 	LANDSCAPE_API TArray<ALandscapeProxy*> GetOutdatedProxies(UE::Landscape::EOutdatedDataFlags InMatchingOutdatedDataFlags, bool bInMustMatchAllFlags) const;
+	LANDSCAPE_API TArray<TTuple<ALandscapeProxy*, UE::Landscape::EOutdatedDataFlags>> GetOutdatedProxyDetails(UE::Landscape::EOutdatedDataFlags InMatchingOutdatedDataFlags, bool bInMustMatchAllFlags) const;
 	
 	LANDSCAPE_API bool IsGridBased() const;
 	LANDSCAPE_API void ChangeGridSize(ULandscapeInfo* LandscapeInfo, uint32 NewGridSizeInComponents);
