Fix crash while converting a level to World Partition - Fixed issue with ALandscape::SplitHeightmap() modifying textures outside of a PreEditChange()/PostEditChange() block

diff --git a/Engine/Source/Runtime/Landscape/Private/LandscapeEdit.cpp b/Engine/Source/Runtime/Landscape/Private/LandscapeEdit.cpp
--- a/Engine/Source/Runtime/Landscape/Private/LandscapeEdit.cpp
+++ b/Engine/Source/Runtime/Landscape/Private/LandscapeEdit.cpp
@@ -5706,6 +5706,12 @@ void ALandscape::SplitHeightmap(ULandscapeComponent* Comp, ALandscapeProxy* Targ
 	FVector4 OldHeightmapScaleBias = Comp->HeightmapScaleBias;
 	FVector4 NewHeightmapScaleBias = FVector4(1.0f / (float)HeightmapSizeU, 1.0f / (float)HeightmapSizeV, 0.0f, 0.0f);
 
+	// TODO-LS: Although we are only reading from the existing heightmap texture, the current code in FLandscapeEditDataInterface
+	// is performing a call to LockMip(), which results in the UnlockMip() modifying the TextureSource CompressionFormat.
+	// Modification to textures MUST be wrapped in PreEditChange/PostEditChange() in order to prevent texture compilation issues.
+	// But really, the ideal fix here would be to call LockMipReadOnly() instead of LockMip(), and completely avoid modifying the texture.
+	OldHeightmapTexture->PreEditChange(nullptr);
+
 	{
 		// Read old data and split
 		FLandscapeEditDataInterface LandscapeEdit(Info);
@@ -5728,6 +5734,8 @@ void ALandscape::SplitHeightmap(ULandscapeComponent* Comp, ALandscapeProxy* Targ
 		LandscapeEdit.SetHeightData(Comp->GetSectionBase().X, Comp->GetSectionBase().Y, Comp->GetSectionBase().X + Comp->ComponentSizeQuads, Comp->GetSectionBase().Y + Comp->ComponentSizeQuads, (uint16*)HeightData.GetData(), 0, false, (uint16*)NormalData.GetData());
 	}
 
+	OldHeightmapTexture->PostEditChange();
+
 	// End material update
 	if (InOutUpdateContext != nullptr && InOutRecreateRenderStateContext != nullptr)
 	{
@@ -5808,40 +5816,56 @@ void ALandscape::SplitHeightmap(ULandscapeComponent* Comp, ALandscapeProxy* Targ
 			}
 		}
 
-		// Move layer content to new layer heightmap
-		FLandscapeEditDataInterface LandscapeEdit(Info);
-		ALandscape* Landscape = Info->LandscapeActor.Get();
-		Comp->ForEachLayer([&](const FGuid& LayerGuid, FLandscapeLayerComponentData& LayerData)
+
+		// TODO-LS: Although we are only reading from the existing heightmap texture, the current code in FLandscapeEditDataInterface
+		// is performing a call to LockMip(), which results in the UnlockMip() modifying the TextureSource CompressionFormat.
+		// Modification to textures MUST be wrapped in PreEditChange/PostEditChange() in order to prevent texture compilation issues.
+		// But really, the ideal fix here would be to call LockMipReadOnly() instead of LockMip(), and completely avoid modifying the texture.
+		TArray<UTexture2D*> OldLayerHeightmaps;
+
 		{
-			UTexture2D* OldLayerHeightmap = LayerData.HeightmapData.Texture;
-			if (OldLayerHeightmap != nullptr)
+			// Move layer content to new layer heightmap
+			FLandscapeEditDataInterface LandscapeEdit(Info);
+			ALandscape* Landscape = Info->LandscapeActor.Get();
+
+			Comp->ForEachLayer([&](const FGuid& LayerGuid, FLandscapeLayerComponentData& LayerData)
 			{
-				FScopedSetLandscapeEditingLayer Scope(Landscape, LayerGuid);
-				// Read old data and split
-				TArray<uint8> LayerHeightData;
-				LayerHeightData.AddZeroed((1 + Comp->ComponentSizeQuads) * (1 + Comp->ComponentSizeQuads) * sizeof(uint16));
-				// Because of edge problem, normal would be just copy from old component data
-				TArray<uint8> LayerNormalData;
-				LayerNormalData.AddZeroed((1 + Comp->ComponentSizeQuads) * (1 + Comp->ComponentSizeQuads) * sizeof(uint16));
-
-				// Read using old heightmap scale/bias
-				Comp->HeightmapScaleBias = OldHeightmapScaleBias;
-				LandscapeEdit.GetHeightDataFast(Comp->GetSectionBase().X, Comp->GetSectionBase().Y, Comp->GetSectionBase().X + Comp->ComponentSizeQuads, Comp->GetSectionBase().Y + Comp->ComponentSizeQuads, (uint16*)LayerHeightData.GetData(), 0, (uint16*)LayerNormalData.GetData());
-				// Restore new heightmap scale/bias
-				Comp->HeightmapScaleBias = NewHeightmapScaleBias;
+				UTexture2D* OldLayerHeightmap = LayerData.HeightmapData.Texture;
+				if (OldLayerHeightmap != nullptr)
 				{
-					// no mipchain required as these layer weight maps are used in layer compositing to generate a final set of weight maps to be used for rendering
-					UTexture2D* LayerHeightmapTexture = DstProxy->CreateLandscapeTexture(HeightmapSizeU, HeightmapSizeV, TEXTUREGROUP_Terrain_Heightmap, TSF_BGRA8, /* OptionalOverrideOuter = */ nullptr, /* bCompress = */ false, /* bMipChain = */ false);
-					ULandscapeComponent::CreateEmptyTextureMips(LayerHeightmapTexture, true);
-					LayerHeightmapTexture->PostEditChange();
-					// Set Layer heightmap texture
-					LayerData.HeightmapData.Texture = LayerHeightmapTexture;
-					LandscapeEdit.SetHeightData(Comp->GetSectionBase().X, Comp->GetSectionBase().Y, Comp->GetSectionBase().X + Comp->ComponentSizeQuads, Comp->GetSectionBase().Y + Comp->ComponentSizeQuads, (uint16*)LayerHeightData.GetData(), 0, false, (uint16*)LayerNormalData.GetData());
+					FScopedSetLandscapeEditingLayer Scope(Landscape, LayerGuid);
+
+					OldLayerHeightmaps.Add(OldLayerHeightmap);
+					OldLayerHeightmap->PreEditChange(nullptr);
+
+					// Read old data and split
+					TArray<uint8> LayerHeightData;
+					LayerHeightData.AddZeroed((1 + Comp->ComponentSizeQuads) * (1 + Comp->ComponentSizeQuads) * sizeof(uint16));
+					// Because of edge problem, normal would be just copy from old component data
+					TArray<uint8> LayerNormalData;
+					LayerNormalData.AddZeroed((1 + Comp->ComponentSizeQuads) * (1 + Comp->ComponentSizeQuads) * sizeof(uint16));
+
+					// Read using old heightmap scale/bias
+					Comp->HeightmapScaleBias = OldHeightmapScaleBias;
+					LandscapeEdit.GetHeightDataFast(Comp->GetSectionBase().X, Comp->GetSectionBase().Y, Comp->GetSectionBase().X + Comp->ComponentSizeQuads, Comp->GetSectionBase().Y + Comp->ComponentSizeQuads, (uint16*)LayerHeightData.GetData(), 0, (uint16*)LayerNormalData.GetData());
+					// Restore new heightmap scale/bias
+					Comp->HeightmapScaleBias = NewHeightmapScaleBias;
+					{
+						// no mipchain required as these layer weight maps are used in layer compositing to generate a final set of weight maps to be used for rendering
+						UTexture2D* LayerHeightmapTexture = DstProxy->CreateLandscapeTexture(HeightmapSizeU, HeightmapSizeV, TEXTUREGROUP_Terrain_Heightmap, TSF_BGRA8, /* OptionalOverrideOuter = */ nullptr, /* bCompress = */ false, /* bMipChain = */ false);
+						ULandscapeComponent::CreateEmptyTextureMips(LayerHeightmapTexture, true);
+						LayerHeightmapTexture->PostEditChange();
+						// Set Layer heightmap texture
+						LayerData.HeightmapData.Texture = LayerHeightmapTexture;
+						LandscapeEdit.SetHeightData(Comp->GetSectionBase().X, Comp->GetSectionBase().Y, Comp->GetSectionBase().X + Comp->ComponentSizeQuads, Comp->GetSectionBase().Y + Comp->ComponentSizeQuads, (uint16*)LayerHeightData.GetData(), 0, false, (uint16*)LayerNormalData.GetData());
+					}
 				}
-			}
-		});
+			});
+
+			Landscape->RequestLayersContentUpdateForceAll();
+		}
 
-		Landscape->RequestLayersContentUpdateForceAll();
+		Algo::ForEach(OldLayerHeightmaps, [](UTexture2D* Tex) { Tex->PostEditChange(); });
 	}
 #endif
 
