Fixed crash when landscape fixed grid vertex factory is not created : we would still create mesh batches regardless
#tests ps4, editor, standalone x64

diff --git a/Engine/Source/Runtime/Landscape/Private/LandscapeRender.cpp b/Engine/Source/Runtime/Landscape/Private/LandscapeRender.cpp
--- a/Engine/Source/Runtime/Landscape/Private/LandscapeRender.cpp
+++ b/Engine/Source/Runtime/Landscape/Private/LandscapeRender.cpp
@@ -1330,6 +1330,7 @@ FLandscapeComponentSceneProxy::FLandscapeComponentSceneProxy(ULandscapeComponent
 	, SharedBuffersKey(0)
 	, SharedBuffers(nullptr)
 	, VertexFactory(nullptr)
+	, FixedGridVertexFactory(nullptr)
 	, ComponentLightInfo(nullptr)
 	, NumRelevantMips(InComponent->GetNumRelevantMips())
 #if WITH_EDITORONLY_DATA
@@ -2338,6 +2339,8 @@ void FLandscapeComponentSceneProxy::OnTransformChanged(FRHICommandListBase& RHIC
 /** Creates a mesh batch for virtual texture or water info texture rendering. The caller is responsible for setting the required flags (bRenderToVirtualTexture, bUseForWaterInfoTextureDepth). Will render a simple fixed grid with combined subsections. */
 bool FLandscapeComponentSceneProxy::GetMeshElementForFixedGrid(int32 InLodIndex, FMaterialRenderProxy* InMaterialInterface, FMeshBatch& OutMeshBatch, TArray<FLandscapeBatchElementParams>& OutStaticBatchParamArray) const
 {
+	check(FixedGridVertexFactory != nullptr);
+
 	if (InMaterialInterface == nullptr)
 	{
 		return false;
@@ -2450,61 +2453,69 @@ void FLandscapeComponentSceneProxy::DrawStaticElements(FStaticPrimitiveDrawInter
 	}
 
 	int32 TotalBatchCount = 1 + LastLOD - FirstLOD;
-	TotalBatchCount += (1 + LastVirtualTextureLOD - FirstVirtualTextureLOD) * RuntimeVirtualTextureMaterialTypes.Num();
-	TotalBatchCount += 1; // TODO: Currently we always add a single LOD0 fixed grid landscape mesh batch for rendering the water info texture. Higher LODs might be better and we might not always need to do this.
-	TotalBatchCount += 1; // LOD0 for lumen surface cache capture
+
+	if (FixedGridVertexFactory != nullptr)
+	{
+		TotalBatchCount += (1 + LastVirtualTextureLOD - FirstVirtualTextureLOD) * RuntimeVirtualTextureMaterialTypes.Num();
+		TotalBatchCount += 1; // TODO: Currently we always add a single LOD0 fixed grid landscape mesh batch for rendering the water info texture. Higher LODs might be better and we might not always need to do this.
+		TotalBatchCount += 1; // LOD0 for lumen surface cache capture
+	}
 
 	StaticBatchParamArray.Empty(TotalBatchCount);
 	PDI->ReserveMemoryForMeshes(TotalBatchCount);
 
-	// Add fixed grid mesh batches for runtime virtual texture usage
-	for (ERuntimeVirtualTextureMaterialType MaterialType : RuntimeVirtualTextureMaterialTypes)
+	// Create batches for fixed grid : 
+	if (FixedGridVertexFactory != nullptr)
 	{
-		const int32 MaterialIndex = LODIndexToMaterialIndex[FirstLOD];
-
-		for (int32 LODIndex = FirstVirtualTextureLOD; LODIndex <= LastVirtualTextureLOD; ++LODIndex)
+		// Add fixed grid mesh batches for runtime virtual texture usage
+		for (ERuntimeVirtualTextureMaterialType MaterialType : RuntimeVirtualTextureMaterialTypes)
 		{
-			FMeshBatch RuntimeVirtualTextureMeshBatch;
-			if (GetMeshElementForFixedGrid(LODIndex, AvailableMaterials[MaterialIndex], RuntimeVirtualTextureMeshBatch, StaticBatchParamArray))
+			const int32 MaterialIndex = LODIndexToMaterialIndex[FirstLOD];
+	
+			for (int32 LODIndex = FirstVirtualTextureLOD; LODIndex <= LastVirtualTextureLOD; ++LODIndex)
 			{
-				RuntimeVirtualTextureMeshBatch.bRenderToVirtualTexture = true;
-				RuntimeVirtualTextureMeshBatch.RuntimeVirtualTextureMaterialType = (uint32)MaterialType;
-				PDI->DrawMesh(RuntimeVirtualTextureMeshBatch, FLT_MAX);
+				FMeshBatch RuntimeVirtualTextureMeshBatch;
+				if (GetMeshElementForFixedGrid(LODIndex, AvailableMaterials[MaterialIndex], RuntimeVirtualTextureMeshBatch, StaticBatchParamArray))
+				{
+					RuntimeVirtualTextureMeshBatch.bRenderToVirtualTexture = true;
+					RuntimeVirtualTextureMeshBatch.RuntimeVirtualTextureMaterialType = (uint32)MaterialType;
+					PDI->DrawMesh(RuntimeVirtualTextureMeshBatch, FLT_MAX);
+				}
 			}
 		}
-	}
-
-	// Add fixed grid mesh batch for rendering the water info texture
-	if (FixedGridVertexFactory)
-	{
-		int32 LODIndex = 0;
-		int32 MaterialIndex = LODIndexToMaterialIndex[LODIndex];
-
-		FMeshBatch MeshBatch;
-		if (GetMeshElementForFixedGrid(LODIndex, AvailableMaterials[MaterialIndex], MeshBatch, StaticBatchParamArray))
+	
+		// Add fixed grid mesh batch for rendering the water info texture
+		if (FixedGridVertexFactory)
 		{
-			MeshBatch.bUseForWaterInfoTextureDepth = true;
-			PDI->DrawMesh(MeshBatch, FLT_MAX);
+			int32 LODIndex = 0;
+			int32 MaterialIndex = LODIndexToMaterialIndex[LODIndex];
+	
+			FMeshBatch MeshBatch;
+			if (GetMeshElementForFixedGrid(LODIndex, AvailableMaterials[MaterialIndex], MeshBatch, StaticBatchParamArray))
+			{
+				MeshBatch.bUseForWaterInfoTextureDepth = true;
+				PDI->DrawMesh(MeshBatch, FLT_MAX);
+			}
 		}
-	}
-
-	// add fixed grid for lumen card captures
-	if (FixedGridVertexFactory)
-	{
-		FMeshBatch MeshBatch;
-
-		if (GetStaticMeshElement(0, false, MeshBatch, StaticBatchParamArray))
+	
+		// add fixed grid for lumen card captures
+		if (FixedGridVertexFactory)
 		{
-			MeshBatch.VertexFactory = FixedGridVertexFactory;
-			MeshBatch.CastShadow = false;
-			MeshBatch.bUseForDepthPass = false;
-			MeshBatch.bUseAsOccluder = false;
-			MeshBatch.bUseForMaterial = false;
-			MeshBatch.bDitheredLODTransition = false;
-			MeshBatch.bRenderToVirtualTexture = false;
-			MeshBatch.bUseForLumenSurfaceCacheCapture = true;
-
-			PDI->DrawMesh(MeshBatch, FLT_MAX);
+			FMeshBatch MeshBatch;
+	
+			if (GetStaticMeshElement(0, false, MeshBatch, StaticBatchParamArray))
+			{
+				MeshBatch.VertexFactory = FixedGridVertexFactory;
+				MeshBatch.CastShadow = false;
+				MeshBatch.bUseForDepthPass = false;
+				MeshBatch.bUseAsOccluder = false;
+				MeshBatch.bUseForMaterial = false;
+				MeshBatch.bDitheredLODTransition = false;
+				MeshBatch.bRenderToVirtualTexture = false;
+				MeshBatch.bUseForLumenSurfaceCacheCapture = true;
+	
+				PDI->DrawMesh(MeshBatch, FLT_MAX);
+			}
 		}
 	}
 
