Fixed crash when a landscape edit layer object fails to load : the layer is now deleted entirely from landscape (with an error log) when this happens

diff --git a/Engine/Source/Runtime/Landscape/Private/Landscape.cpp b/Engine/Source/Runtime/Landscape/Private/Landscape.cpp
--- a/Engine/Source/Runtime/Landscape/Private/Landscape.cpp
+++ b/Engine/Source/Runtime/Landscape/Private/Landscape.cpp
@@ -2860,15 +2860,30 @@ PRAGMA_DISABLE_DEPRECATION_WARNINGS
 PRAGMA_ENABLE_DEPRECATION_WARNINGS
 	}
 
-	for (FLandscapeLayer& Layer : LandscapeEditLayers)
+	for (int32 LayerIndex = 0; LayerIndex < LandscapeEditLayers.Num(); ++LayerIndex)
 	{
-		check(Layer.EditLayer != nullptr);
-		// For now, only Layer reserved for Landscape Spline uses AlphaBlend
-		Layer.BlendMode = Layer.EditLayer->IsA<ULandscapeEditLayerSplines>() ? LSBM_AlphaBlend : LSBM_AdditiveBlend;
+		FLandscapeLayer& Layer = LandscapeEditLayers[LayerIndex];
+		if (Layer.EditLayer != nullptr)
+		{
+			// For now, only Layer reserved for Landscape Spline uses AlphaBlend
+			Layer.BlendMode = Layer.EditLayer->IsA<ULandscapeEditLayerSplines>() ? LSBM_AlphaBlend : LSBM_AdditiveBlend;
+
+			for (FLandscapeLayerBrush& Brush : Layer.Brushes)
+			{
+				Brush.SetOwner(this);
+			}
 
-		for (FLandscapeLayerBrush& Brush : Layer.Brushes)
+		}
+		else
 		{
-			Brush.SetOwner(this);
+			UE_LOG(LogLandscape, Error, TEXT("Couldn't load edit layer object associated with layer %s for landscape %s. This may happen when the edit layer class cannot be found (for example, when a plugin is removed from the project). The layer will be deleted."), *Layer.Name.ToString(), *GetFullName());
+			// Don't use DeleteLayer because it relies on the ULandscapeInfo object to be valid, which is not the case on PostLoad. 
+			// Simply remove the layer from list and the proxies will eventually tell the user to remove their data associated with this layer upon registration (in order to let them avoid data loss) :
+			LandscapeEditLayers.RemoveAt(LayerIndex);
+			--LayerIndex;
+
+			// Request Update
+			RequestLayersContentUpdateForceAll();
 		}
 	}
 #endif // WITH_EDITOR
