Added a way for landscape edit layers to register their render dependencies (textures/material interfaces) so that we can wait upon them to be ready (respectively: fully streamed in/compiled) before executing the landscape edit layers merge, in order to ensure a deterministic result

diff --git a/Engine/Source/Runtime/Landscape/Classes/Landscape.h b/Engine/Source/Runtime/Landscape/Classes/Landscape.h
--- a/Engine/Source/Runtime/Landscape/Classes/Landscape.h
+++ b/Engine/Source/Runtime/Landscape/Classes/Landscape.h
@@ -28,6 +28,7 @@ struct FTextureToComponentHelper;
 struct FUpdateLayersContentContext;
 struct FEditLayersHeightmapMergeParams;
 struct FEditLayersWeightmapMergeParams;
+enum class ELandscapeNotificationType;
 
 namespace EditLayersHeightmapLocalMerge_RenderThread
 {
@@ -495,7 +496,7 @@ private:
 
 	bool PrepareLayersTextureResources(bool bInWaitForStreaming);
 	bool PrepareLayersTextureResources(const TArray<FLandscapeLayer>& InLayers, bool bInWaitForStreaming);
-	bool PrepareLayersBrushResources(ERHIFeatureLevel::Type InFeatureLevel, bool bInWaitForStreaming);
+	bool PrepareLayersResources(ERHIFeatureLevel::Type InFeatureLevel, bool bInWaitForStreaming);
 	void InvalidateRVTForTextures(const TSet<UTexture2D*>& InTextures);
 	void PrepareLayersHeightmapsLocalMergeRenderThreadData(const FUpdateLayersContentContext& InUpdateLayersContentContext, const FEditLayersHeightmapMergeParams& InMergeParams, EditLayersHeightmapLocalMerge_RenderThread::FMergeInfo& OutRenderThreadData);
 	void PrepareLayersWeightmapsLocalMergeRenderThreadData(const FUpdateLayersContentContext& InUpdateLayersContentContext, const FEditLayersWeightmapMergeParams& InMergeParams, EditLayersWeightmapLocalMerge_RenderThread::FMergeInfo& OutRenderThreadData);
@@ -654,15 +655,22 @@ private:
 	UPROPERTY(Transient)
 	bool bSplineLayerUpdateRequested;
 
-	/** Time since waiting for landscape resources to be ready (for displaying a notification to the user) */
-	double WaitingForLandscapeTextureResourcesStartTime = -1.0;
+	struct FWaitingForResourcesNotificationHelper
+	{
+		void Notify(ALandscape* InLandscape, class FLandscapeNotificationManager* InNotificationManager, ELandscapeNotificationType InNotificationType, const FText& InNotificationText);
+		void Reset();
+
+		/** Time since waiting for resources to be ready */
+		double WaitingForResourcesStartTime = -1.0;
 
-	/** Time since waiting for brush resources to be ready (for displaying a notification to the user) */
-	double WaitingForLandscapeBrushResourcesStartTime = -1.0;
+		/** Non-stackable user notification for landscape editor */
+		TSharedPtr<FLandscapeNotification> Notification;
+	};
 
 	/** Non-stackable user notifications for landscape editor */
-	TSharedPtr<FLandscapeNotification> WaitingForTexturesNotification;
-	TSharedPtr<FLandscapeNotification> WaitingForBrushesNotification;
+	FWaitingForResourcesNotificationHelper WaitingForTexturesNotificationHelper;
+	FWaitingForResourcesNotificationHelper WaitingForEditLayerResourcesNotificationHelper;
+
 	TSharedPtr<FLandscapeNotification> InvalidShadingModelNotification;
 
 	// Represent all the resolved paint layer, from all layers blended together (size of the landscape x material layer count)
diff --git a/Engine/Source/Runtime/Landscape/Classes/LandscapeEditLayer.h b/Engine/Source/Runtime/Landscape/Classes/LandscapeEditLayer.h
--- a/Engine/Source/Runtime/Landscape/Classes/LandscapeEditLayer.h
+++ b/Engine/Source/Runtime/Landscape/Classes/LandscapeEditLayer.h
@@ -119,6 +119,13 @@ public:
 	*/
 	virtual TArray<FEditLayerAction> GetActions() const { return {}; }
 
+	/**
+	* @return a list of UObjects that this layers needs in order to render properly. This is necessary to avoid trying to render a layer while some of its
+	*  resources are not fully ready. 
+	*  These can be UTexture (requires all mips to be fully loaded) or UMaterialInterface (requires shader maps to be fully compiled)
+	*/
+	virtual void GetRenderDependencies(TSet<UObject*>& OutDependencies) const { }
+
 	// TODO [jonathan.bard] : Remove : temporary method to give the edit layer an opportunity to change some settings on the old struct (FLandscapeLayer) upon creation. 
 	//  The better way would be to move most of the settings to ULandscapeEditLayer(Base?) and expose a "property changed" event that both the UI and runtime code 
 	//  could listen to in order to generically react to any change on the layer settings, including on the derived class (e.g. if there were some splines layer-specific 
diff --git a/Engine/Source/Runtime/Landscape/Private/LandscapeEditLayers.cpp b/Engine/Source/Runtime/Landscape/Private/LandscapeEditLayers.cpp
--- a/Engine/Source/Runtime/Landscape/Private/LandscapeEditLayers.cpp
+++ b/Engine/Source/Runtime/Landscape/Private/LandscapeEditLayers.cpp
@@ -3628,7 +3628,7 @@ bool ALandscape::PrepareLayersTextureResources(bool bInWaitForStreaming)
 
 bool ALandscape::PrepareLayersTextureResources(const TArray<FLandscapeLayer>& InLayers, bool bInWaitForStreaming)
 {
-	TRACE_CPUPROFILER_EVENT_SCOPE(LandscapeLayers_PrepareLayersTextureResources);
+	TRACE_CPUPROFILER_EVENT_SCOPE(ALandscape::PrepareLayersTextureResources);
 
 	ULandscapeInfo* Info = GetLandscapeInfo();
 	if (Info == nullptr)
@@ -3662,12 +3662,20 @@ bool ALandscape::PrepareLayersTextureResources(const TArray<FLandscapeLayer>& In
 	return bIsReady;
 }
 
-bool ALandscape::PrepareLayersBrushResources(ERHIFeatureLevel::Type InFeatureLevel, bool bInWaitForStreaming)
+bool ALandscape::PrepareLayersResources(ERHIFeatureLevel::Type InFeatureLevel, bool bInWaitForStreaming)
 {
-	TRACE_CPUPROFILER_EVENT_SCOPE(LandscapeLayers_PrepareLayersBrushTextureResources);
+	TRACE_CPUPROFILER_EVENT_SCOPE(ALandscape::PrepareLayersResources);
 	TSet<UObject*> Dependencies;
 	for (const FLandscapeLayer& Layer : LandscapeEditLayers)
 	{
+		ULandscapeEditLayerBase* EditLayer = Layer.EditLayer;
+		check(EditLayer != nullptr);
+
+		if (EditLayer->SupportsTargetType(ELandscapeToolTargetType::Heightmap) || EditLayer->SupportsTargetType(ELandscapeToolTargetType::Weightmap) || EditLayer->SupportsTargetType(ELandscapeToolTargetType::Visibility))
+		{
+			EditLayer->GetRenderDependencies(Dependencies);
+		}
+
 		for (const FLandscapeLayerBrush& Brush : Layer.Brushes)
 		{
 			if (ALandscapeBlueprintBrushBase* LandscapeBrush = Brush.GetBrush())
@@ -3686,7 +3694,7 @@ bool ALandscape::PrepareLayersBrushResources(ERHIFeatureLevel::Type InFeatureLev
 	for (UObject* Dependency : Dependencies)
 	{
 		// Streamable textures need to be fully streamed in : 
-		if (UTexture2D* Texture = Cast<UTexture2D>(Dependency))
+		if (UTexture* Texture = Cast<UTexture>(Dependency))
 		{
 			bIsReady &= TextureStreamingManager->RequestTextureFullyStreamedInForever(Texture, bInWaitForStreaming);
 		}
@@ -8232,6 +8240,27 @@ void ULandscapeComponent::GetLandscapeComponentWeightmapsToRender(TSet<ULandscap
 	}
 }
 
+void ALandscape::FWaitingForResourcesNotificationHelper::Notify(ALandscape* InLandscape, FLandscapeNotificationManager* InNotificationManager, ELandscapeNotificationType InNotificationType, const FText& InNotificationText)
+{
+	// We need to wait until layers texture resources are ready to initialize the landscape to avoid taking the sizes and format of the default texture:
+	static constexpr double TimeBeforeDisplayingWaitingForResourcesNotification = 3.0;
+
+	WaitingForResourcesStartTime = FSlateApplicationBase::IsInitialized() ? FSlateApplicationBase::Get().GetCurrentTime() : 0.0f;
+	if (!Notification.IsValid())
+	{
+		Notification = MakeShared<FLandscapeNotification>(InLandscape, InNotificationType);
+		Notification->NotificationText = InNotificationText;
+		Notification->NotificationStartTime = WaitingForResourcesStartTime + TimeBeforeDisplayingWaitingForResourcesNotification;
+	}
+	InNotificationManager->RegisterNotification(Notification);
+}
+
+void ALandscape::FWaitingForResourcesNotificationHelper::Reset()
+{
+	Notification.Reset();
+	WaitingForResourcesStartTime = -1.0;
+}
+
 void ALandscape::UpdateLayersContent(bool bInWaitForStreaming, bool bInSkipMonitorLandscapeEdModeChanges, bool bIntermediateRender, bool bFlushRender)
 {
 	TRACE_CPUPROFILER_EVENT_SCOPE(LandscapeLayers_UpdateLayersContent);
@@ -8242,11 +8271,9 @@ void ALandscape::UpdateLayersContent(bool bInWaitForStreaming, bool bInSkipMonit
 		// Make sure that we don't leave any notification behind when we leave this function without explicitly displaying one :
 		if (bHideNotifications)
 		{
-			WaitingForTexturesNotification.Reset();
-			WaitingForBrushesNotification.Reset();
+			WaitingForTexturesNotificationHelper.Reset();
+			WaitingForEditLayerResourcesNotificationHelper.Reset();
 			InvalidShadingModelNotification.Reset();
-			WaitingForLandscapeTextureResourcesStartTime = -1.0;
-			WaitingForLandscapeBrushResourcesStartTime = -1.0;
 		}
 
 		// If nothing to do, let's do some garbage collecting on async readback tasks so that we slowly get rid of staging textures 
@@ -8357,11 +8384,6 @@ void ALandscape::UpdateLayersContent(bool bInWaitForStreaming, bool bInSkipMonit
 		return;
 	}
 
-	auto GetCurrentTime = []()
-	{
-		return FSlateApplicationBase::IsInitialized() ? FSlateApplicationBase::Get().GetCurrentTime() : 0.0f;
-	};
-	
 	// The Edit layers shaders only work on SM5 : cancel any update that might happen when SM5+ shading model is not active :
 	if (World->GetFeatureLevel() < ERHIFeatureLevel::SM5)
 	{
@@ -8369,7 +8391,7 @@ void ALandscape::UpdateLayersContent(bool bInWaitForStreaming, bool bInSkipMonit
 		{
 			if (!InvalidShadingModelNotification.IsValid())
 			{
-				InvalidShadingModelNotification = MakeShared<FLandscapeNotification>(this, FLandscapeNotification::EType::ShadingModelInvalid);
+				InvalidShadingModelNotification = MakeShared<FLandscapeNotification>(this, ELandscapeNotificationType::ShadingModelInvalid);
 				static const FText NotificationText(LOCTEXT("InvalidShadingModel", "Cannot update landscape with a feature level less than SM5"));
 				InvalidShadingModelNotification->NotificationText = NotificationText;
 			}
@@ -8383,47 +8405,26 @@ void ALandscape::UpdateLayersContent(bool bInWaitForStreaming, bool bInSkipMonit
 		InvalidShadingModelNotification.Reset();
 	}
 
-	// We need to wait until layers texture resources are ready to initialize the landscape to avoid taking the sizes and format of the default texture:
-	static constexpr double TimeBeforeDisplayingWaitingForResourcesNotification = 3.0;
-
 	bResourcesReady &= PrepareLayersTextureResources(bInWaitForStreaming);
 	if (!bResourcesReady && LandscapeNotificationManager)
 	{
-		WaitingForLandscapeTextureResourcesStartTime = GetCurrentTime(); 
-		if (!WaitingForTexturesNotification.IsValid())
-		{
-			WaitingForTexturesNotification = MakeShared<FLandscapeNotification>(this, FLandscapeNotification::EType::LandscapeTextureResourcesNotReady);
-			static const FText NotificationText(LOCTEXT("WaitForLandscapeTextureResources", "Waiting for texture resources to be ready"));
-			WaitingForTexturesNotification->NotificationText = NotificationText;
-			WaitingForTexturesNotification->NotificationStartTime = WaitingForLandscapeTextureResourcesStartTime + TimeBeforeDisplayingWaitingForResourcesNotification;
-		}
-		LandscapeNotificationManager->RegisterNotification(WaitingForTexturesNotification);
+		WaitingForTexturesNotificationHelper.Notify(this, LandscapeNotificationManager, ELandscapeNotificationType::LandscapeTextureResourcesNotReady, LOCTEXT("WaitForLandscapeTextureResources", "Waiting for texture resources to be ready"));
 		bHideNotifications = false;
 	}
 	else
 	{
-		WaitingForTexturesNotification.Reset();
-		WaitingForLandscapeTextureResourcesStartTime = -1.0;
+		WaitingForTexturesNotificationHelper.Reset();
 	}
 
-	bResourcesReady &= PrepareLayersBrushResources(World->GetFeatureLevel(), bInWaitForStreaming);
+	bResourcesReady &= PrepareLayersResources(World->GetFeatureLevel(), bInWaitForStreaming);
 	if (!bResourcesReady && LandscapeNotificationManager)
 	{
-		WaitingForLandscapeBrushResourcesStartTime = GetCurrentTime();
-		if (!WaitingForBrushesNotification.IsValid())
-		{
-			WaitingForBrushesNotification = MakeShared<FLandscapeNotification>(this, FLandscapeNotification::EType::LandscapeBrushResourcesNotReady);
-			static const FText NotificationText(LOCTEXT("WaitForLandscapeBrushResources", "Waiting for brush resources to be ready"));
-			WaitingForBrushesNotification->NotificationText = NotificationText;
-			WaitingForBrushesNotification->NotificationStartTime = WaitingForLandscapeBrushResourcesStartTime + TimeBeforeDisplayingWaitingForResourcesNotification;
-		}
-		LandscapeNotificationManager->RegisterNotification(WaitingForBrushesNotification);
+		WaitingForEditLayerResourcesNotificationHelper.Notify(this, LandscapeNotificationManager, ELandscapeNotificationType::LandscapeEditLayerResourcesNotReady, LOCTEXT("WaitForLandscapeEditLayerResources", "Waiting for edit layer resources to be ready"));
 		bHideNotifications = false;
 	}
 	else
 	{
-		WaitingForBrushesNotification.Reset();
-		WaitingForLandscapeBrushResourcesStartTime = -1.0;
+		WaitingForEditLayerResourcesNotificationHelper.Reset();
 	}
 
 	if (!bResourcesReady)
diff --git a/Engine/Source/Runtime/Landscape/Private/LandscapeNotification.cpp b/Engine/Source/Runtime/Landscape/Private/LandscapeNotification.cpp
--- a/Engine/Source/Runtime/Landscape/Private/LandscapeNotification.cpp
+++ b/Engine/Source/Runtime/Landscape/Private/LandscapeNotification.cpp
@@ -14,7 +14,7 @@
 
 // ----------------------------------------------------------------------------------
 
-FLandscapeNotification::FLandscapeNotification(const TWeakObjectPtr<ALandscape>& InLandscape, EType InNotificationType, FConditionCallback InConditionCallback, FUpdateTextCallback InUpdateTextCallback)
+FLandscapeNotification::FLandscapeNotification(const TWeakObjectPtr<ALandscape>& InLandscape, ELandscapeNotificationType InNotificationType, FConditionCallback InConditionCallback, FUpdateTextCallback InUpdateTextCallback)
 	: Landscape(InLandscape)
 	, NotificationType(InNotificationType)
 	, ConditionCallback(InConditionCallback)
diff --git a/Engine/Source/Runtime/Landscape/Public/LandscapeNotification.h b/Engine/Source/Runtime/Landscape/Public/LandscapeNotification.h
--- a/Engine/Source/Runtime/Landscape/Public/LandscapeNotification.h
+++ b/Engine/Source/Runtime/Landscape/Public/LandscapeNotification.h
@@ -7,10 +7,17 @@
 #include "CoreMinimal.h"
 #include "Landscape.h"
 
+/** The various types of landscape notifications, ordered by priority (a single notification is displayed at a time). */
+enum class ELandscapeNotificationType
+{
+	ShadingModelInvalid = 0,
+	LandscapeTextureResourcesNotReady,
+	LandscapeEditLayerResourcesNotReady,
+};
+
 struct FLandscapeNotification
 {
-	/** The various types of landscape notifications, ordered by priority (a single notification is displayed at a time). */
-	enum class EType : uint8
+	enum class UE_DEPRECATED(5.5, "Use ELandscapeNotificationType") EType : uint8
 	{
 		ShadingModelInvalid = 0,
 		LandscapeTextureResourcesNotReady,
@@ -29,14 +36,14 @@ struct FLandscapeNotification
 	* @param InConditionCallback - Lambda used to conditionally determine whether this notification should be displayed per tick.
 	* @param InUpdateTextCallback - Lambda used to construct notification text when the text can change (e.g. displaying the value of a variable).
 	*/ 
-	FLandscapeNotification(const TWeakObjectPtr<ALandscape>& InLandscape, EType InNotificationType, FConditionCallback InConditionCallback = FConditionCallback([]() { return true; }), FUpdateTextCallback InUpdateTextCallback = [](FText& InText) {});
+	FLandscapeNotification(const TWeakObjectPtr<ALandscape>& InLandscape, ELandscapeNotificationType InNotificationType, FConditionCallback InConditionCallback = FConditionCallback([]() { return true; }), FUpdateTextCallback InUpdateTextCallback = [](FText& InText) {});
 	~FLandscapeNotification() = default;
 
 	bool operator == (const FLandscapeNotification& Other) const;
 	bool operator < (const FLandscapeNotification& Other) const;
 
 	TWeakObjectPtr<ALandscape> GetLandscape() const { return Landscape; }
-	EType GetNotificationType() const { return NotificationType; }
+	ELandscapeNotificationType GetNotificationType() const { return NotificationType; }
 
 	// Public wrappers for the callback functions.
 	bool ShouldShowNotification() const;
@@ -54,7 +61,7 @@ public:
 
 private:
 	TWeakObjectPtr<ALandscape> Landscape;
-	EType NotificationType;
+	ELandscapeNotificationType NotificationType;
 
 	/** Determines whether or not a notification should be shown. Defaults to return true. */
 	FConditionCallback ConditionCallback;
