Adding "no Z scaling" option for Landscape Segment Spline Meshes

diff --git a/Engine/Source/Runtime/Landscape/Classes/LandscapeSplineSegment.h b/Engine/Source/Runtime/Landscape/Classes/LandscapeSplineSegment.h
--- a/Engine/Source/Runtime/Landscape/Classes/LandscapeSplineSegment.h
+++ b/Engine/Source/Runtime/Landscape/Classes/LandscapeSplineSegment.h
@@ -151,6 +151,10 @@ struct FLandscapeSplineMeshEntry
 	UPROPERTY(EditAnywhere, Category=LandscapeSplineMeshEntry)
 	uint32 bScaleToWidth:1;
 
+	/** Disables scale to width on the mesh Z coordinate */
+	UPROPERTY(EditAnywhere, Category = LandscapeSplineMeshEntry, meta = (EditCondition = "bScaleToWidth"))
+	uint32 bNoZScaling:1;
+
 	/** Scale of the spline mesh, (Z=Forwards) */
 	UPROPERTY(EditAnywhere, Category=LandscapeSplineMeshEntry)
 	FVector Scale;
@@ -173,6 +177,7 @@ struct FLandscapeSplineMeshEntry
 		bCenterH(true),
 		CenterAdjust(0, 0),
 		bScaleToWidth(true),
+		bNoZScaling(false),
 		Scale(1,1,1),
 		Orientation_DEPRECATED(LSMO_YUp),
 		ForwardAxis(ESplineMeshAxis::X),
diff --git a/Engine/Source/Runtime/Landscape/Private/LandscapeSplines.cpp b/Engine/Source/Runtime/Landscape/Private/LandscapeSplines.cpp
--- a/Engine/Source/Runtime/Landscape/Private/LandscapeSplines.cpp
+++ b/Engine/Source/Runtime/Landscape/Private/LandscapeSplines.cpp
@@ -3123,7 +3123,16 @@ void ULandscapeSplineSegment::UpdateSplinePoints(bool bUpdateCollision, bool bUp
 			FVector Scale = MeshEntry->Scale;
 			if (MeshEntry->bScaleToWidth)
 			{
-				Scale *= Width / USplineMeshComponent::GetAxisValueRef(MeshBounds.BoxExtent, CrossAxis(MeshEntry->ForwardAxis, MeshEntry->UpAxis));
+				float ScaleFactor = Width / USplineMeshComponent::GetAxisValueRef(MeshBounds.BoxExtent, CrossAxis(MeshEntry->ForwardAxis, MeshEntry->UpAxis));
+				if (MeshEntry->bNoZScaling)
+				{
+					Scale.X *= ScaleFactor;
+					Scale.Y *= ScaleFactor;
+				}
+				else
+				{
+					Scale *= ScaleFactor;
+				}
 			}
 
 			const float MeshLength = static_cast<float>(FMath::Abs(USplineMeshComponent::GetAxisValueRef(MeshBounds.BoxExtent, MeshEntry->ForwardAxis) * 2.0 *
@@ -3259,7 +3268,16 @@ void ULandscapeSplineSegment::UpdateSplinePoints(bool bUpdateCollision, bool bUp
 			FVector Scale = MeshEntry->Scale;
 			if (MeshEntry->bScaleToWidth)
 			{
-				Scale *= Width / USplineMeshComponent::GetAxisValueRef(MeshBounds.BoxExtent, SideAxis);
+				float ScaleFactor = Width / USplineMeshComponent::GetAxisValueRef(MeshBounds.BoxExtent, SideAxis);
+				if (MeshEntry->bNoZScaling)
+				{
+					Scale.X *= ScaleFactor;
+					Scale.Y *= ScaleFactor;
+				}
+				else
+				{
+					Scale *= ScaleFactor;
+				}
 			}
 
 			FVector2D Offset = MeshEntry->CenterAdjust;
@@ -3309,7 +3327,16 @@ void ULandscapeSplineSegment::UpdateSplinePoints(bool bUpdateCollision, bool bUp
 			FVector ScaleEnd = MeshEntry->Scale;
 			if (MeshEntry->bScaleToWidth)
 			{
-				ScaleEnd *= WidthEnd / USplineMeshComponent::GetAxisValueRef(MeshBounds.BoxExtent, SideAxis);
+				float ScaleFactor = WidthEnd / USplineMeshComponent::GetAxisValueRef(MeshBounds.BoxExtent, SideAxis);
+				if (MeshEntry->bNoZScaling)
+				{
+					ScaleEnd.X *= ScaleFactor;
+					ScaleEnd.Y *= ScaleFactor;
+				}
+				else
+				{
+					ScaleEnd *= ScaleFactor;
+				}
 			}
 
 			FVector2D OffsetEnd = MeshEntry->CenterAdjust;
