Fix an issue with landscape culling bounds being possibly incorrect.
Detect when the cached local box is incorrect on serialization and fix it.

diff --git a/Engine/Source/Runtime/Landscape/Classes/LandscapeComponent.h b/Engine/Source/Runtime/Landscape/Classes/LandscapeComponent.h
--- a/Engine/Source/Runtime/Landscape/Classes/LandscapeComponent.h
+++ b/Engine/Source/Runtime/Landscape/Classes/LandscapeComponent.h
@@ -1014,9 +1014,9 @@ public:
 
 #if WITH_EDITOR
 	/**
-	 * Recalculate cached bounds using height values.
+	 * Recalculate cached bounds using height values.  Returns true when the bounds were changed.
 	 */
-	LANDSCAPE_API void UpdateCachedBounds(bool bInApproximateBounds = false);
+	LANDSCAPE_API bool UpdateCachedBounds(bool bInApproximateBounds = false);
 
 	/**
 	 * Update the MaterialInstance parameters to match the layer and weightmaps for this component
diff --git a/Engine/Source/Runtime/Landscape/Private/Landscape.cpp b/Engine/Source/Runtime/Landscape/Private/Landscape.cpp
--- a/Engine/Source/Runtime/Landscape/Private/Landscape.cpp
+++ b/Engine/Source/Runtime/Landscape/Private/Landscape.cpp
@@ -777,6 +777,14 @@ void ULandscapeComponent::Serialize(FArchive& Ar)
 		}
 	}
 
+#if WITH_EDITOR
+	// double check we never save an invalid cached local box to a cooked package (should always be recalculated in ALandscapeProxy::PreSave)
+	if (Ar.IsSaving() && Ar.IsCooking() && !Ar.IsSerializingDefaults())
+	{
+		check(CachedLocalBox.GetVolume() > 0);
+	}
+#endif // WITH_EDITOR
+
 	// Avoid the archiver in the PIE duplicate writer case because we want to share landscape textures & materials
 	if (Ar.GetPortFlags() & PPF_DuplicateForPIE)
 	{
@@ -891,6 +899,16 @@ void ULandscapeComponent::Serialize(FArchive& Ar)
 		Super::Serialize(Ar);
 	}
 
+	// this is a sanity check, as ALandscapeProxy::PreSave() for cook should have ensured that the cached local box has non-zero volume
+	if (Ar.IsLoadingFromCookedPackage() && (CachedLocalBox.GetVolume() <= 0))
+	{
+		// we must set a conservative bounds as a last resort here -- if not we risk strobing flicker of landscape visibility
+		FVector MinBox(0, 0, LandscapeDataAccess::GetLocalHeight(0));
+		FVector MaxBox(ComponentSizeQuads + 1, ComponentSizeQuads + 1, LandscapeDataAccess::GetLocalHeight(UINT16_MAX));
+		CachedLocalBox = FBox(MinBox, MaxBox);
+		UE_LOG(LogLandscape, Error, TEXT("The component %s has an invalid CachedLocalBox. It has been set to a conservative bounds, that may result in reduced visibility culling performance"), *GetName());
+	}
+
 	if (Ar.IsLoading() && Ar.CustomVer(FRenderingObjectVersion::GUID) < FRenderingObjectVersion::MapBuildDataSeparatePackage)
 	{
 		FMeshMapBuildData* LegacyMapBuildData = new FMeshMapBuildData();
@@ -3698,14 +3716,22 @@ void ALandscapeProxy::PreSave(FObjectPreSaveContext ObjectSaveContext)
 		UpdateRenderingMethod();
 	}
 
-	// Warn if there are any active Physical Material Renders on this Proxy
-	// don.boogert todo: We should block or wait for this to complete within some timeout period.
 	for (ULandscapeComponent* LandscapeComponent : LandscapeComponents)
 	{
+		// Warn if there are any active Physical Material Renders on this Proxy
+		// don.boogert todo: We should block or wait for this to complete within some timeout period.
 		if (LandscapeComponent->PhysicalMaterialTask.IsInProgress())
 		{
 			UE_LOG(LogLandscape, Warning, TEXT("Physical material render on component: '%s' in progress."), *LandscapeComponent->GetFullName());
 		}
+
+		// Ensure the component's cached bounds are correct
+		FBox OldCachedLocalBox = LandscapeComponent->CachedLocalBox;
+		if (LandscapeComponent->UpdateCachedBounds(/* bInApproximateBounds= */ false))
+		{
+			UE_LOG(LogLandscape, Warning, TEXT("The component %s had incorrect bounds.  The bounds have been recalculated (old CachedLocalBox: %s, new CachedLocalBox: %s)"), *GetName(), *OldCachedLocalBox.ToString(), *LandscapeComponent->CachedLocalBox.ToString());
+			check(LandscapeComponent->CachedLocalBox.GetVolume() > 0.0);
+		}
 	}
 
 	if (LandscapeGuid.IsValid())
diff --git a/Engine/Source/Runtime/Landscape/Private/LandscapeEdit.cpp b/Engine/Source/Runtime/Landscape/Private/LandscapeEdit.cpp
--- a/Engine/Source/Runtime/Landscape/Private/LandscapeEdit.cpp
+++ b/Engine/Source/Runtime/Landscape/Private/LandscapeEdit.cpp
@@ -262,17 +262,27 @@ namespace UE::Landscape::Private
 	}
 }
 
-void ULandscapeComponent::UpdateCachedBounds(bool bInApproximateBounds)
+bool ULandscapeComponent::UpdateCachedBounds(bool bInApproximateBounds)
 {
 	TRACE_CPUPROFILER_EVENT_SCOPE(ULandscapeComponent::UpdateCachedBounds);
 
+	if ((SubsectionSizeQuads == 0) || (NumSubsections == 0) || (ComponentSizeQuads == 0))
+	{
+		// the landscape component is in an uninitialized/default state, we cannot calculate meaningful bounds
+		return false;
+	}
+
+	bool bChanged = false;
+	FBox NewLocalBox;
+
 	// Update local-space bounding box
-	CachedLocalBox.Init();
+	NewLocalBox.Init();
 	if (bInApproximateBounds && GetLandscapeProxy()->HasLayersContent())
 	{
 		FVector MinBox(0, 0, LandscapeDataAccess::GetLocalHeight(0));
 		FVector MaxBox(ComponentSizeQuads + 1, ComponentSizeQuads + 1, LandscapeDataAccess::GetLocalHeight(UINT16_MAX));
-		CachedLocalBox = FBox(MinBox, MaxBox);
+		NewLocalBox = FBox(MinBox, MaxBox);
+		check(NewLocalBox.GetExtent().Z > 0.0);
 	}
 	else
 	{
@@ -332,7 +342,7 @@ void ULandscapeComponent::UpdateCachedBounds(bool bInApproximateBounds)
 				}
 			}
 
-			CachedLocalBox = FBox(FVector(0.0, 0.0, LocalMin), FVector(ComponentSizeQuads, ComponentSizeQuads, LocalMax));
+			NewLocalBox = FBox(FVector(0.0, 0.0, LocalMin), FVector(ComponentSizeQuads, ComponentSizeQuads, LocalMax));
 		}
 
 		MipToMipMaxDeltas = ComputeMipToMipMaxDeltas(MakeArrayView<const FQuadHeightInfo>(AllMipsQuadInfos.GetData(), AllMipsQuadInfos.Num()), NumTextureMips, NumRelevantMips);
@@ -343,10 +353,16 @@ void ULandscapeComponent::UpdateCachedBounds(bool bInApproximateBounds)
 			MarkRenderStateDirty();
 		}
 	}
-	if (CachedLocalBox.GetExtent().Z == 0)
+	if (NewLocalBox.GetExtent().Z == 0)
 	{
 		// expand bounds to avoid flickering issues with zero-size bounds
-		CachedLocalBox = CachedLocalBox.ExpandBy(FVector(0, 0, 1));
+		NewLocalBox = NewLocalBox.ExpandBy(FVector(0, 0, 1));
+	}
+
+	if (NewLocalBox != CachedLocalBox)
+	{
+		CachedLocalBox = NewLocalBox;
+		bChanged = true;
 	}
 
 	// Update collision component bounds
@@ -361,6 +377,8 @@ void ULandscapeComponent::UpdateCachedBounds(bool bInApproximateBounds)
 		HFCollisionComponent->CachedLocalBox = CachedLocalBox;
 		HFCollisionComponent->UpdateComponentToWorld();
 	}
+
+	return bChanged;
 }
 #endif //WITH_EDITOR
 
