Landscape: Remove assert that Physical Material Task isn't in flight if the hash is indentical

diff --git a/Engine/Source/Runtime/Landscape/Private/LandscapeEdit.cpp b/Engine/Source/Runtime/Landscape/Private/LandscapeEdit.cpp
--- a/Engine/Source/Runtime/Landscape/Private/LandscapeEdit.cpp
+++ b/Engine/Source/Runtime/Landscape/Private/LandscapeEdit.cpp
@@ -2080,7 +2080,10 @@ void ULandscapeComponent::UpdatePhysicalMaterialTasks()
 
 	if (Hash == PhysicalMaterialHash)
 	{
-		check(!PhysicalMaterialTask.IsValid());
+		if (PhysicalMaterialTask.IsValid())
+		{
+			PhysicalMaterialTask.Release();
+		}
 		return;
 	}
 
