Fixed crash when landscape fixed grid vertex factory is not created : we would still create mesh batches regardless
#tests editor

diff --git a/Engine/Source/Runtime/Landscape/Private/LandscapeRender.cpp b/Engine/Source/Runtime/Landscape/Private/LandscapeRender.cpp
--- a/Engine/Source/Runtime/Landscape/Private/LandscapeRender.cpp
+++ b/Engine/Source/Runtime/Landscape/Private/LandscapeRender.cpp
@@ -2494,7 +2494,6 @@ void FLandscapeComponentSceneProxy::DrawStaticElements(FStaticPrimitiveDrawInter
 		}
 	
 		// Add fixed grid mesh batch for rendering the water info texture
-		if (FixedGridVertexFactory)
 		{
 			int32 LODIndex = 0;
 			int32 MaterialIndex = LODIndexToMaterialIndex[LODIndex];
@@ -2508,7 +2507,6 @@ void FLandscapeComponentSceneProxy::DrawStaticElements(FStaticPrimitiveDrawInter
 		}
 	
 		// add fixed grid for lumen card captures
-		if (FixedGridVertexFactory)
 		{
 			FMeshBatch MeshBatch;
 	
