#UE Add more info to a warning

diff --git a/Engine/Source/Runtime/Landscape/Private/Landscape.cpp b/Engine/Source/Runtime/Landscape/Private/Landscape.cpp
--- a/Engine/Source/Runtime/Landscape/Private/Landscape.cpp
+++ b/Engine/Source/Runtime/Landscape/Private/Landscape.cpp
@@ -3729,7 +3729,7 @@ void ALandscapeProxy::PreSave(FObjectPreSaveContext ObjectSaveContext)
 		FBox OldCachedLocalBox = LandscapeComponent->CachedLocalBox;
 		if (LandscapeComponent->UpdateCachedBounds(/* bInApproximateBounds= */ false))
 		{
-			UE_LOG(LogLandscape, Warning, TEXT("The component %s had incorrect bounds.  The bounds have been recalculated (old CachedLocalBox: %s, new CachedLocalBox: %s)"), *GetName(), *OldCachedLocalBox.ToString(), *LandscapeComponent->CachedLocalBox.ToString());
+			UE_LOG(LogLandscape, Warning, TEXT("The component %s had incorrect bounds.  The bounds have been recalculated (old CachedLocalBox: %s, new CachedLocalBox: %s)"), *GetPathName(), *OldCachedLocalBox.ToString(), *LandscapeComponent->CachedLocalBox.ToString());
 			check(LandscapeComponent->CachedLocalBox.GetVolume() > 0.0);
 		}
 	}
