Merging CL#33161903

#UE Lower warning verbosity to Display, since this is showing for existing content during cook

diff --git a/Engine/Source/Runtime/Landscape/Private/Landscape.cpp b/Engine/Source/Runtime/Landscape/Private/Landscape.cpp
--- a/Engine/Source/Runtime/Landscape/Private/Landscape.cpp
+++ b/Engine/Source/Runtime/Landscape/Private/Landscape.cpp
@@ -3729,7 +3729,7 @@ void ALandscapeProxy::PreSave(FObjectPreSaveContext ObjectSaveContext)
 		FBox OldCachedLocalBox = LandscapeComponent->CachedLocalBox;
 		if (LandscapeComponent->UpdateCachedBounds(/* bInApproximateBounds= */ false))
 		{
-			UE_LOG(LogLandscape, Warning, TEXT("The component %s had incorrect bounds.  The bounds have been recalculated (old CachedLocalBox: %s, new CachedLocalBox: %s)"), *GetPathName(), *OldCachedLocalBox.ToString(), *LandscapeComponent->CachedLocalBox.ToString());
+			UE_LOG(LogLandscape, Display, TEXT("The component %s had incorrect bounds.  The bounds have been recalculated (old CachedLocalBox: %s, new CachedLocalBox: %s)"), *GetPathName(), *OldCachedLocalBox.ToString(), *LandscapeComponent->CachedLocalBox.ToString());
 			check(LandscapeComponent->CachedLocalBox.GetVolume() > 0.0);
 		}
 	}
