Fix for rare crash in GetCurrentLevelLandscapeProxy()

diff --git a/Engine/Source/Runtime/Landscape/Private/Landscape.cpp b/Engine/Source/Runtime/Landscape/Private/Landscape.cpp
--- a/Engine/Source/Runtime/Landscape/Private/Landscape.cpp
+++ b/Engine/Source/Runtime/Landscape/Private/Landscape.cpp
@@ -5170,7 +5170,7 @@ ALandscapeProxy* ULandscapeInfo::GetCurrentLevelLandscapeProxy(bool bRegistered)
 	ALandscapeProxy* LandscapeProxy = nullptr;
 	ForEachLandscapeProxy([&LandscapeProxy, bRegistered](ALandscapeProxy* Proxy) -> bool
 	{
-		if (!bRegistered || Proxy->GetRootComponent()->IsRegistered())
+		if (!bRegistered || (Proxy->GetRootComponent() && Proxy->GetRootComponent()->IsRegistered()))
 		{
 			UWorld* ProxyWorld = Proxy->GetWorld();
 			if (ProxyWorld &&
