Fix for landscape bounds calculation

Also re-enabling the log warning in PreSave, but only when the bounds were actually incorrect and potentially flash-inducing.  For simple non-optimal bounds it just logs a display message.

diff --git a/Engine/Source/Runtime/Landscape/Private/Landscape.cpp b/Engine/Source/Runtime/Landscape/Private/Landscape.cpp
--- a/Engine/Source/Runtime/Landscape/Private/Landscape.cpp
+++ b/Engine/Source/Runtime/Landscape/Private/Landscape.cpp
@@ -3729,7 +3729,17 @@ void ALandscapeProxy::PreSave(FObjectPreSaveContext ObjectSaveContext)
 		FBox OldCachedLocalBox = LandscapeComponent->CachedLocalBox;
 		if (LandscapeComponent->UpdateCachedBounds(/* bInApproximateBounds= */ false))
 		{
-			UE_LOG(LogLandscape, Display, TEXT("The component %s had incorrect bounds.  The bounds have been recalculated (old CachedLocalBox: %s, new CachedLocalBox: %s)"), *GetPathName(), *OldCachedLocalBox.ToString(), *LandscapeComponent->CachedLocalBox.ToString());
+			// conservative bounds are true bounding boxes, just not as tight/optimal as they could be
+			// if it's not conservative, then visibility flashing issues can occur because of self-occlusion in culling
+			bool bOldBoxIsConservative = LandscapeComponent->CachedLocalBox.IsInsideOrOn(OldCachedLocalBox);
+			if (bOldBoxIsConservative)
+			{
+				UE_LOG(LogLandscape, Display, TEXT("The component %s had non-optimal bounds.  The bounds have been recalculated (old CachedLocalBox: %s, new CachedLocalBox: %s)"), *LandscapeComponent->GetPathName(), *OldCachedLocalBox.ToString(), *LandscapeComponent->CachedLocalBox.ToString());
+			}
+			else
+			{
+				UE_LOG(LogLandscape, Display, TEXT("The component %s had incorrect bounds.  The bounds have been recalculated (old CachedLocalBox: %s, new CachedLocalBox: %s)"), *LandscapeComponent->GetPathName(), *OldCachedLocalBox.ToString(), *LandscapeComponent->CachedLocalBox.ToString());
+			}
 			check(LandscapeComponent->CachedLocalBox.GetVolume() > 0.0);
 		}
 	}
diff --git a/Engine/Source/Runtime/Landscape/Private/LandscapeEdit.cpp b/Engine/Source/Runtime/Landscape/Private/LandscapeEdit.cpp
--- a/Engine/Source/Runtime/Landscape/Private/LandscapeEdit.cpp
+++ b/Engine/Source/Runtime/Landscape/Private/LandscapeEdit.cpp
@@ -164,7 +164,7 @@ namespace UE::Landscape::Private
 		}
 
 		double Min = DBL_MAX;
-		double Max = DBL_MIN;
+		double Max = -DBL_MAX;
 		double Average = 0.0f;
 	};
 
@@ -306,7 +306,7 @@ bool ULandscapeComponent::UpdateCachedBounds(bool bInApproximateBounds)
 			TRACE_CPUPROFILER_EVENT_SCOPE(FetchMipQuads);
 
 			double LocalMin = DBL_MAX;
-			double LocalMax = DBL_MIN;
+			double LocalMax = -DBL_MAX;
 			AllMipsQuadInfos.AddZeroed(ComputeQuadInfosCount(NumRelevantMips, TextureSize));
 			for (int32 MipIndex = 0; MipIndex < NumRelevantMips; ++MipIndex)
 			{
