Fixed incorrect string formatting
#rnx
#tests editor

diff --git a/Engine/Source/Runtime/Landscape/Private/LandscapeEditLayers.cpp b/Engine/Source/Runtime/Landscape/Private/LandscapeEditLayers.cpp
--- a/Engine/Source/Runtime/Landscape/Private/LandscapeEditLayers.cpp
+++ b/Engine/Source/Runtime/Landscape/Private/LandscapeEditLayers.cpp
@@ -4909,13 +4909,13 @@ int32 ALandscape::PerformLayersHeightmapsGlobalMerge(const FUpdateLayersContentC
 
 		// Convert Atlas LandscapeScratchRT1 to the world-projected NonAtlas in LandscapeScratchRT2
 		// TODO: just use this format from the beginning above...
-		DrawHeightmapComponentsToRenderTarget(FString::Printf(TEXT("LS Height: %s += -> NonAtlas %s"), *Layer.Name.ToString(), *LandscapeScratchRT1->GetName(), *LandscapeScratchRT2->GetName()),
+		DrawHeightmapComponentsToRenderTarget(FString::Printf(TEXT("LS Height: %s %s += -> NonAtlas %s"), *Layer.Name.ToString(), *LandscapeScratchRT1->GetName(), *LandscapeScratchRT2->GetName()),
 			InUpdateLayersContentContext.LandscapeComponentsHeightmapsToRender, LandscapeExtent.Min, LandscapeScratchRT1, nullptr, LandscapeScratchRT2, ERTDrawingType::RTAtlasToNonAtlas, true, ShaderParams);
 
 		ShaderParams.ApplyLayerModifiers = true;
 
 		// Combine Current layer NonAtlas LandscapeScratchRT2 with current result in LandscapeScratchRT3, writing final result to CombinedHeightmapNonAtlasRT
-		DrawHeightmapComponentsToRenderTarget(FString::Printf(TEXT("LS Height: %s += -> CombinedNonAtlas %s"), *Layer.Name.ToString(), *LandscapeScratchRT2->GetName(), *CombinedHeightmapNonAtlasRT->GetName()),
+		DrawHeightmapComponentsToRenderTarget(FString::Printf(TEXT("LS Height: %s %s += -> CombinedNonAtlas %s"), *Layer.Name.ToString(), *LandscapeScratchRT2->GetName(), *CombinedHeightmapNonAtlasRT->GetName()),
 			InUpdateLayersContentContext.LandscapeComponentsHeightmapsToRender, LandscapeExtent.Min, LandscapeScratchRT2, FirstLayer ? nullptr : LandscapeScratchRT3, CombinedHeightmapNonAtlasRT, ERTDrawingType::RTNonAtlas, FirstLayer, ShaderParams);
 
 		ShaderParams.ApplyLayerModifiers = false;
