Converting a landscape weightmap reallocation crash into a warning, and adding sanity checks to track down the source of the badly initialized weightmaps (too many mips)

diff --git a/Engine/Source/Runtime/Landscape/Private/LandscapeEdit.cpp b/Engine/Source/Runtime/Landscape/Private/LandscapeEdit.cpp
--- a/Engine/Source/Runtime/Landscape/Private/LandscapeEdit.cpp
+++ b/Engine/Source/Runtime/Landscape/Private/LandscapeEdit.cpp
@@ -7179,6 +7179,12 @@ void ULandscapeComponent::ReallocateWeightmapsInternal(FLandscapeEditDataInterfa
 			}
 		}
 
+		if (CurrentWeightmapTexture != nullptr)
+		{
+			// sanity check - only final weightmaps are allowed more than one mip
+			ensure((CurrentWeightmapTexture->Source.GetNumMips() == 1) || bIsFinalWeightmap);
+		}
+
 		bool NeedsUpdateResource = false;
 		// No suitable weightmap texture
 		if (CurrentWeightmapTexture == nullptr)
@@ -7189,13 +7195,18 @@ void ULandscapeComponent::ReallocateWeightmapsInternal(FLandscapeEditDataInterfa
 			int32 WeightmapSize = (SubsectionSizeQuads + 1) * NumSubsections;
 
 			// We need a new weightmap texture
-			CurrentWeightmapTexture = TargetProxy->CreateLandscapeTexture(WeightmapSize, WeightmapSize, TEXTUREGROUP_Terrain_Weightmap, TSF_BGRA8, nullptr, false, bIsFinalWeightmap); //dmb-nomips
+			CurrentWeightmapTexture = TargetProxy->CreateLandscapeTexture(WeightmapSize, WeightmapSize, TEXTUREGROUP_Terrain_Weightmap, TSF_BGRA8, nullptr, /* bCompress = */ false, /* bMipChain= */ bIsFinalWeightmap); //dmb-nomips
 
 			// Alloc dummy mips
 			if (bIsFinalWeightmap)
 			{
 				CreateEmptyTextureMips(CurrentWeightmapTexture, true);
 			}
+			else
+			{
+				// sanity check - only final weightmaps are allowed more than one mip
+				check(CurrentWeightmapTexture->Source.GetNumMips() == 1);
+			}
 
 			CurrentWeightmapTexture->PostEditChange();
 
@@ -7237,6 +7248,8 @@ void ULandscapeComponent::ReallocateWeightmapsInternal(FLandscapeEditDataInterfa
 					// Copy the data
 					if (ensure(DataInterface != nullptr)) // it's not safe to skip the copy
 					{
+						// sanity check - only final weightmaps are allowed more than one mip
+						ensure((CurrentWeightmapTexture->Source.GetNumMips() == 1) || bIsFinalWeightmap);
 						DataInterface->CopyTextureChannel(CurrentWeightmapTexture, ChanIdx, OldWeightmapTexture, AllocInfo.WeightmapTextureChannel);
 						DataInterface->ZeroTextureChannel(OldWeightmapTexture, AllocInfo.WeightmapTextureChannel);
 						// UE_LOG(LogLandscape, Log, TEXT("Copying old channel (%s).%d to new channel (%s).%d"), *OldWeightmapTexture->GetName(), AllocInfo.WeightmapTextureChannel, *CurrentWeightmapTexture->GetName(), ChanIdx);
diff --git a/Engine/Source/Runtime/Landscape/Private/LandscapeEditInterface.cpp b/Engine/Source/Runtime/Landscape/Private/LandscapeEditInterface.cpp
--- a/Engine/Source/Runtime/Landscape/Private/LandscapeEditInterface.cpp
+++ b/Engine/Source/Runtime/Landscape/Private/LandscapeEditInterface.cpp
@@ -20,6 +20,7 @@ LandscapeEditInterface.cpp: Landscape editing interface
 #include "LandscapeDataAccess.h"
 #include "LandscapeEdit.h"
 #include "LandscapeRender.h"
+#include "LandscapePrivate.h"
 #include "ComponentReregisterContext.h"
 #include "Algo/Transform.h"
 #include "TextureCompiler.h"
@@ -4355,7 +4356,16 @@ void FLandscapeTextureDataInterface::CopyTextureChannel(UTexture2D* Dest, int32
 	int32 MipSize = Dest->Source.GetSizeX();
 	check(Dest->Source.GetSizeX() == Dest->Source.GetSizeY() && Src->Source.GetSizeX() == Dest->Source.GetSizeX());
 
-	for( int32 MipIdx=0;MipIdx<DestDataInfo->NumMips();MipIdx++ )
+	int32 NumMips = DestDataInfo->NumMips();
+	int32 SrcMips = SrcDataInfo->NumMips();
+	if (SrcMips != NumMips)
+	{
+		UE_LOG(LogLandscape, Warning, TEXT("Unexpected mip count mismatch when copying landscape texture channels from '%s' (%d mips) to '%s' (%d mips) -- mip data may be lost, or left uninitialized"),
+			*Src->GetPathName(), SrcMips, *Dest->GetPathName(), NumMips);
+		NumMips = FMath::Min(NumMips, SrcMips);
+	}
+
+	for( int32 MipIdx=0;MipIdx< NumMips;MipIdx++ )
 	{
 		uint8* DestTextureData = (uint8*)DestDataInfo->GetMipData(MipIdx);
 		uint8* SrcTextureData = (uint8*)SrcDataInfo->GetMipData(MipIdx);
