Landscape: User defined layers

diff --git a/Engine/Plugins/Enterprise/DatasmithImporter/Source/DatasmithImporter/Private/DatasmithLandscapeImporter.cpp b/Engine/Plugins/Enterprise/DatasmithImporter/Source/DatasmithImporter/Private/DatasmithLandscapeImporter.cpp
--- a/Engine/Plugins/Enterprise/DatasmithImporter/Source/DatasmithImporter/Private/DatasmithLandscapeImporter.cpp
+++ b/Engine/Plugins/Enterprise/DatasmithImporter/Source/DatasmithImporter/Private/DatasmithLandscapeImporter.cpp
@@ -115,7 +115,7 @@ AActor* FDatasmithLandscapeImporter::ImportLandscapeActor( const TSharedRef< IDa
 	{
 		if(ImportLandscapeLayersList[i].LayerInfo != nullptr)
 		{
-			Landscape->EditorLayerSettings.Add(FLandscapeEditorLayerSettings(ImportLandscapeLayersList[i].LayerInfo, ImportLandscapeLayersList[i].SourceFilePath));
+			Landscape->AddTargetLayer(ImportLandscapeLayersList[i].LayerInfo->LayerName, FLandscapeTargetLayerSettings(ImportLandscapeLayersList[i].LayerInfo, ImportLandscapeLayersList[i].SourceFilePath));
 
 			int32 LayerInfoIndex = LandscapeInfo->GetLayerInfoIndex(ImportLandscapeLayersList[i].LayerName);
 			if(ensure(LayerInfoIndex != INDEX_NONE))
diff --git a/Engine/Source/Editor/EditorStyle/Private/SlateEditorStyle.cpp b/Engine/Source/Editor/EditorStyle/Private/SlateEditorStyle.cpp
--- a/Engine/Source/Editor/EditorStyle/Private/SlateEditorStyle.cpp
+++ b/Engine/Source/Editor/EditorStyle/Private/SlateEditorStyle.cpp
@@ -6968,6 +6968,9 @@ void FSlateEditorStyle::FStyle::SetupLandscapeEditorStyle()
 		Set("LandscapeEditor.TargetList.DropZone.Above", new BOX_BRUSH("Common/VerticalBoxDropZoneIndicator_Above", FMargin(10.0f / 16.0f, 10.0f / 16.0f, 0, 0), SelectionColor_Subdued));
 		Set("LandscapeEditor.TargetList.DropZone.Below", new BOX_BRUSH("Common/VerticalBoxDropZoneIndicator_Below", FMargin(10.0f / 16.0f, 0, 0, 10.0f / 16.0f), SelectionColor_Subdued));
 
+		//todo don.boogert - replace with final icon
+		Set("LandscapeEditor.Layer.Sync", new IMAGE_BRUSH("Icons/icon_MatEd_HLSL_Code_16x", Icon16x16));
+		
 		Set("LandscapeEditor.SpinBox", FSpinBoxStyle(GetWidgetStyle<FSpinBoxStyle>("SpinBox"))
 			.SetTextPadding(FMargin(0))
 			.SetBackgroundBrush(FSlateNoResource())
diff --git a/Engine/Source/Editor/EditorStyle/Private/StarshipStyle.cpp b/Engine/Source/Editor/EditorStyle/Private/StarshipStyle.cpp
--- a/Engine/Source/Editor/EditorStyle/Private/StarshipStyle.cpp
+++ b/Engine/Source/Editor/EditorStyle/Private/StarshipStyle.cpp
@@ -6600,6 +6600,9 @@ void FStarshipEditorStyle::FStyle::SetupLandscapeEditorStyle()
 		Set("LandscapeEditor.TargetList.DropZone.Above", new BOX_BRUSH("Common/VerticalBoxDropZoneIndicator_Above", FMargin(10.0f / 16.0f, 10.0f / 16.0f, 0, 0), SelectionColor_Subdued));
 		Set("LandscapeEditor.TargetList.DropZone.Below", new BOX_BRUSH("Common/VerticalBoxDropZoneIndicator_Below", FMargin(10.0f / 16.0f, 0, 0, 10.0f / 16.0f), SelectionColor_Subdued));
 
+		//todo don.boogert - replace with final icon
+		Set("LandscapeEditor.Layer.Sync", new IMAGE_BRUSH("Icons/icon_MatEd_HLSL_Code_16x", Icon16x16));
+		
 		Set("LandscapeEditor.SpinBox", FSpinBoxStyle(GetWidgetStyle<FSpinBoxStyle>("SpinBox"))
 			.SetTextPadding(FMargin(0))
 			.SetBackgroundBrush(FSlateNoResource())
diff --git a/Engine/Source/Editor/LandscapeEditor/Private/LandscapeEdMode.cpp b/Engine/Source/Editor/LandscapeEditor/Private/LandscapeEdMode.cpp
--- a/Engine/Source/Editor/LandscapeEditor/Private/LandscapeEdMode.cpp
+++ b/Engine/Source/Editor/LandscapeEditor/Private/LandscapeEdMode.cpp
@@ -2271,16 +2271,6 @@ const TArray<FLandscapeListInfo>& FEdModeLandscape::GetLandscapeList()
 	return LandscapeList;
 }
 
-void FEdModeLandscape::AddLayerInfo(ULandscapeLayerInfoObject* LayerInfo)
-{
-	if (CurrentToolTarget.LandscapeInfo.IsValid() && CurrentToolTarget.LandscapeInfo->GetLayerInfoIndex(LayerInfo) == INDEX_NONE)
-	{
-		ALandscapeProxy* Proxy = CurrentToolTarget.LandscapeInfo->GetLandscapeProxy();
-		CurrentToolTarget.LandscapeInfo->Layers.Add(FLandscapeInfoLayerSettings(LayerInfo, Proxy));
-		UpdateTargetList();
-	}
-}
-
 int32 FEdModeLandscape::UpdateLandscapeList()
 {
 	LandscapeList.Empty();
@@ -3716,7 +3706,10 @@ ALandscape* FEdModeLandscape::ChangeComponentSetting(int32 NumComponentsX, int32
 				NewLandscape->BodyInstance.SetObjectType(OldLandscape->BodyInstance.GetObjectType());
 				NewLandscape->BodyInstance.SetResponseToChannels(OldLandscape->BodyInstance.GetResponseToChannels());
 			}
-			NewLandscape->EditorLayerSettings = OldLandscape->EditorLayerSettings;
+			for (const TTuple<FName, FLandscapeTargetLayerSettings>& Layer : OldLandscape->GetTargetLayers())
+			{
+				NewLandscape->AddTargetLayer(Layer.Key, Layer.Value);
+			}
 			NewLandscape->bUsedForNavigation = OldLandscape->bUsedForNavigation;
 			NewLandscape->MaxPaintedLayersPerComponent = OldLandscape->MaxPaintedLayersPerComponent;
 
diff --git a/Engine/Source/Editor/LandscapeEditor/Private/LandscapeEdMode.h b/Engine/Source/Editor/LandscapeEditor/Private/LandscapeEdMode.h
--- a/Engine/Source/Editor/LandscapeEditor/Private/LandscapeEdMode.h
+++ b/Engine/Source/Editor/LandscapeEditor/Private/LandscapeEdMode.h
@@ -64,7 +64,7 @@ struct FLandscapeToolMode
 
 struct FLandscapeTargetListInfo
 {
-	FText TargetName;
+	FText TargetLayerDisplayName;									// UI Display Name
 	ELandscapeToolTargetType TargetType;
 	TWeakObjectPtr<ULandscapeInfo> LandscapeInfo;
 
@@ -76,8 +76,8 @@ struct FLandscapeTargetListInfo
 	uint32 bValid : 1;												// ignored for heightmap
 	int32 LayerIndex;
 
-	FLandscapeTargetListInfo(FText InTargetName, ELandscapeToolTargetType InTargetType, const FLandscapeInfoLayerSettings& InLayerSettings, int32 InLayerIndex)
-		: TargetName(InTargetName)
+	FLandscapeTargetListInfo(FText InTargetLayerDisplayName, ELandscapeToolTargetType InTargetType, const FLandscapeInfoLayerSettings& InLayerSettings, int32 InLayerIndex)
+		: TargetLayerDisplayName(InTargetLayerDisplayName)
 		, TargetType(InTargetType)
 		, LandscapeInfo(InLayerSettings.Owner->GetLandscapeInfo())
 		, LayerInfoObj(InLayerSettings.LayerInfoObj)
@@ -90,8 +90,8 @@ struct FLandscapeTargetListInfo
 	{
 	}
 
-	FLandscapeTargetListInfo(FText InTargetName, ELandscapeToolTargetType InTargetType, ULandscapeInfo* InLandscapeInfo, int32 InLayerIndex)
-		: TargetName(InTargetName)
+	FLandscapeTargetListInfo(FText InTargetLayerDisplayName, ELandscapeToolTargetType InTargetType, ULandscapeInfo* InLandscapeInfo, int32 InLayerIndex)
+		: TargetLayerDisplayName(InTargetLayerDisplayName)
 		, TargetType(InTargetType)
 		, LandscapeInfo(InLandscapeInfo)
 		, LayerInfoObj(nullptr)
@@ -134,21 +134,20 @@ struct FLandscapeTargetListInfo
 		return nullptr;
 	}
 
-	FLandscapeEditorLayerSettings* GetEditorLayerSettings() const
+	const FLandscapeTargetLayerSettings* GetTargetLayerSettings() const
 	{
 		if (TargetType == ELandscapeToolTargetType::Weightmap)
 		{
 			check(LayerInfoObj.IsValid());
 			ALandscapeProxy* Proxy = LandscapeInfo->GetLandscapeProxy();
-			FLandscapeEditorLayerSettings* EditorLayerSettings = Proxy->EditorLayerSettings.FindByKey(LayerInfoObj.Get());
-			if (EditorLayerSettings)
+			const FName* TargetLayerName = Proxy->GetTargetLayers().FindKey(FLandscapeTargetLayerSettings(LayerInfoObj.Get()));
+			if (TargetLayerName)
 			{
-				return EditorLayerSettings;
+				return Proxy->GetTargetLayers().Find(*TargetLayerName);
 			}
 			else
 			{
-				int32 Index = Proxy->EditorLayerSettings.Add(FLandscapeEditorLayerSettings(LayerInfoObj.Get()));
-				return &Proxy->EditorLayerSettings[Index];
+				return &Proxy->AddTargetLayer(LayerInfoObj->LayerName, FLandscapeTargetLayerSettings(LayerInfoObj.Get()));
 			}
 		}
 		return nullptr;
@@ -160,7 +159,7 @@ struct FLandscapeTargetListInfo
 	{
 		if (TargetType == ELandscapeToolTargetType::Weightmap)
 		{
-			FLandscapeEditorLayerSettings* EditorLayerSettings = GetEditorLayerSettings();
+			const FLandscapeTargetLayerSettings* EditorLayerSettings = GetTargetLayerSettings();
 			check(EditorLayerSettings);
 			return EditorLayerSettings->ReimportLayerFilePath;
 		}
@@ -184,9 +183,15 @@ struct FLandscapeTargetListInfo
 	{
 		if (TargetType == ELandscapeToolTargetType::Weightmap)
 		{
-			FLandscapeEditorLayerSettings* EditorLayerSettings = GetEditorLayerSettings();
-			check(EditorLayerSettings);
-			EditorLayerSettings->ReimportLayerFilePath = InNewPath;
+			check(LayerInfoObj.IsValid());
+			ALandscapeProxy* Proxy = LandscapeInfo->GetLandscapeProxy();
+			const FName* TargetLayerName = Proxy->GetTargetLayers().FindKey(FLandscapeTargetLayerSettings(LayerInfoObj.Get()));
+			if (TargetLayerName)
+			{
+				FLandscapeTargetLayerSettings LayerSettings = *Proxy->GetTargetLayers().Find(*TargetLayerName);
+				LayerSettings.ReimportLayerFilePath = InNewPath;
+				Proxy->UpdateTargetLayer(*TargetLayerName, LayerSettings);
+			}
 		}
 		else //if (TargetType == ELandscapeToolTargetType::Heightmap)
 		{
@@ -521,8 +526,6 @@ public:
 	int32 GetTargetLayerStartingIndex() const;
 	const TArray<FLandscapeListInfo>& GetLandscapeList();
 
-	void AddLayerInfo(ULandscapeLayerInfoObject* LayerInfo);
-
 	int32 UpdateLandscapeList();
 	void UpdateTargetList();
 	void SetTargetLandscape(const TWeakObjectPtr<ULandscapeInfo>& InLandscapeInfo);
diff --git a/Engine/Source/Editor/LandscapeEditor/Private/LandscapeEditorDetailCustomization_NewLandscape.cpp b/Engine/Source/Editor/LandscapeEditor/Private/LandscapeEditorDetailCustomization_NewLandscape.cpp
--- a/Engine/Source/Editor/LandscapeEditor/Private/LandscapeEditorDetailCustomization_NewLandscape.cpp
+++ b/Engine/Source/Editor/LandscapeEditor/Private/LandscapeEditorDetailCustomization_NewLandscape.cpp
@@ -1154,11 +1154,11 @@ FReply FLandscapeEditorDetailCustomization_NewLandscape::OnCreateButtonClicked()
 		{
 			if (LandscapeEdMode->NewLandscapePreviewMode == ENewLandscapePreviewMode::ImportLandscape)
 			{
-				Landscape->EditorLayerSettings.Add(FLandscapeEditorLayerSettings(LayerInfo, ImportLandscapeLayersList[i].SourceFilePath));
+				Landscape->AddTargetLayer(LayerInfo->LayerName, FLandscapeTargetLayerSettings(LayerInfo, ImportLandscapeLayersList[i].SourceFilePath));
 			}
 			else
 			{
-				Landscape->EditorLayerSettings.Add(FLandscapeEditorLayerSettings(LayerInfo));
+				Landscape->AddTargetLayer(LayerInfo->LayerName, FLandscapeTargetLayerSettings(LayerInfo));
 			}
 
 			int32 LayerInfoIndex = LandscapeInfo->GetLayerInfoIndex(ImportLandscapeLayersList[i].LayerName);
diff --git a/Engine/Source/Editor/LandscapeEditor/Private/LandscapeEditorDetailCustomization_TargetLayers.cpp b/Engine/Source/Editor/LandscapeEditor/Private/LandscapeEditorDetailCustomization_TargetLayers.cpp
--- a/Engine/Source/Editor/LandscapeEditor/Private/LandscapeEditorDetailCustomization_TargetLayers.cpp
+++ b/Engine/Source/Editor/LandscapeEditor/Private/LandscapeEditorDetailCustomization_TargetLayers.cpp
@@ -18,6 +18,8 @@
 #include "Widgets/Input/SComboButton.h"
 #include "Widgets/Input/SCheckBox.h"
 #include "Widgets/Input/SMultiLineEditableTextBox.h"
+#include "Widgets/SBoxPanel.h"
+#include "Widgets/Text/SInlineEditableTextBlock.h"
 #include "EditorModeManager.h"
 #include "EditorModes.h"
 #include "LandscapeEditorModule.h"
@@ -30,19 +32,18 @@
 #include "DetailCategoryBuilder.h"
 #include "PropertyCustomizationHelpers.h"
 #include "Subsystems/AssetEditorSubsystem.h"
+#include "AssetRegistry/AssetRegistryModule.h"
 
 #include "SLandscapeEditor.h"
 #include "Dialogs/DlgPickAssetPath.h"
 #include "ObjectTools.h"
 #include "ScopedTransaction.h"
 #include "DesktopPlatformModule.h"
-#include "AssetRegistry/AssetRegistryModule.h"
 
 #include "LandscapeRender.h"
-#include "Materials/MaterialExpressionLandscapeVisibilityMask.h"
 #include "LandscapeEdit.h"
-#include "IDetailGroup.h"
-#include "Widgets/SBoxPanel.h"
+#include "Landscape.h"
+#include "LandscapeEditorUtils.h"
 
 #define LOCTEXT_NAMESPACE "LandscapeEditor.TargetLayers"
 
@@ -282,11 +283,98 @@ void FLandscapeEditorCustomNodeBuilder_TargetLayers::GenerateHeaderRowContent(FD
 					]
 				]
 			]
+
+			+ SHorizontalBox::Slot()
+			.AutoWidth()
+			.VAlign(VAlign_Center)
+			.Padding(5.0f, 0.0f, 0.0f, 0.0f)
+			[
+				PropertyCustomizationHelpers::MakeAddButton(FSimpleDelegate::CreateRaw(this, &FLandscapeEditorCustomNodeBuilder_TargetLayers::HandleCreateLayer), NSLOCTEXT("Landscape", "CreateLayer", "Create Layer"))
+			]
+
+			+ SHorizontalBox::Slot()
+			.AutoWidth()
+			.VAlign(VAlign_Center)
+			.Padding(5.0f, 0.0f, 0.0f, 0.0f)
+			[
+				SNew(SButton)
+				.ButtonStyle( FAppStyle::Get(), "SimpleButton" )
+				.ToolTipText(NSLOCTEXT("Landscape", "CreateLayersFromMaterials", "Create Layers From Assigned Materials"))
+				.OnClicked(this, &FLandscapeEditorCustomNodeBuilder_TargetLayers::HandleCreateLayersFromMaterials)
+				[
+					SNew(SImage)
+					.Image(FAppStyle::GetBrush("LandscapeEditor.Layer.Sync"))
+				]
+			]
 		];
 	}
 	NodeRow.IsEnabled(TAttribute<bool>::Create(TAttribute<bool>::FGetter::CreateLambda([LandscapeEdMode]() { return LandscapeEdMode->HasValidLandscapeEditLayerSelection(); })));
 }
 
+
+
+FReply FLandscapeEditorCustomNodeBuilder_TargetLayers::HandleCreateLayersFromMaterials()
+{
+	FScopedTransaction Transaction(LOCTEXT("LandscapeWeightLayer_CreateFromMaterials", "Create Weight Layers from Assigned materials"));
+	
+	FEdModeLandscape* LandscapeEdMode = GetEditorMode(); 
+	if (LandscapeEdMode == NULL)
+	{
+		return FReply::Handled(); 
+	}
+	
+	ALandscape* LandscapeActor = LandscapeEdMode->CurrentToolTarget.LandscapeInfo->LandscapeActor.Get();
+	
+	TSet<FName> LayerNames;
+	LandscapeActor->GetLandscapeInfo()->ForEachLandscapeProxy([&LayerNames](ALandscapeProxy* Proxy)
+	{
+		LayerNames.Append(Proxy->RetrieveAllLayerNamesFromMaterials());
+		return true;
+	});
+
+	UE::Landscape::FLayerInfoFinder LayerInfoFinder;
+	for (const FName& LayerName : LayerNames)
+	{
+		if (!LandscapeActor->GetTargetLayers().Find(LayerName))
+		{
+			ULandscapeLayerInfoObject* LandscapeLayerInfo = LayerInfoFinder.Find(LayerName);
+			LandscapeActor->AddTargetLayer(LayerName, FLandscapeTargetLayerSettings(LandscapeLayerInfo));
+		}
+	}
+	
+	LandscapeEdMode->GetLandscape()->GetLandscapeInfo()->UpdateLayerInfoMap();
+	if (LandscapeEdMode)
+	{
+		LandscapeEdMode->UpdateTargetList();
+	}
+
+	return FReply::Handled();
+}
+
+void FLandscapeEditorCustomNodeBuilder_TargetLayers::HandleCreateLayer()
+{
+	FEdModeLandscape* LandscapeEdMode = GetEditorMode();
+	if (LandscapeEdMode == NULL)
+	{
+		return; 
+	}
+	
+	TWeakObjectPtr<ALandscape> Landscape = LandscapeEdMode->CurrentToolTarget.LandscapeInfo->LandscapeActor;
+
+	if (!Landscape.Get())
+	{
+		return; 
+	}
+
+	FScopedTransaction Transaction(LOCTEXT("LandscapeWeightLayer_Create", "Create a Weight Layer"));
+	
+	Landscape->AddTargetLayer();
+	
+	LandscapeEdMode->CurrentToolTarget.LandscapeInfo->UpdateLayerInfoMap();
+	LandscapeEdMode->UpdateTargetList();
+	LandscapeEdMode->RefreshDetailPanel();
+}
+
 TSharedRef<SWidget> FLandscapeEditorCustomNodeBuilder_TargetLayers::GetTargetLayerDisplayOrderButtonMenuContent()
 {
 	FMenuBuilder MenuBuilder(/*bInShouldCloseWindowAfterMenuSelection=*/true, nullptr, nullptr, /*bCloseSelfOnly=*/ true);
@@ -416,6 +504,8 @@ const FSlateBrush* FLandscapeEditorCustomNodeBuilder_TargetLayers::GetTargetLaye
 		{
 			case ELandscapeLayerDisplayMode::Alphabetical: return FAppStyle::Get().GetBrush("LandscapeEditor.Target_DisplayOrder.Alphabetical");
 			case ELandscapeLayerDisplayMode::UserSpecific: return FAppStyle::Get().GetBrush("LandscapeEditor.Target_DisplayOrder.Custom");
+			default:
+				return nullptr;
 		}
 	}
 
@@ -537,7 +627,7 @@ TSharedPtr<SWidget> FLandscapeEditorCustomNodeBuilder_TargetLayers::GenerateRow(
 							[
 								SNew(STextBlock)
 									.Font(IDetailLayoutBuilder::GetDetailFont())
-									.Text(Target->TargetName)
+									.Text(Target->TargetLayerDisplayName)
 									.ColorAndOpacity_Static(&FLandscapeEditorCustomNodeBuilder_TargetLayers::GetTargetTextColor, Target)
 							]
 					]
@@ -620,15 +710,40 @@ TSharedPtr<SWidget> FLandscapeEditorCustomNodeBuilder_TargetLayers::GenerateRow(
 							+ SVerticalBox::Slot()
 							.AutoHeight()
 							.VAlign(VAlign_Center)
-							.Padding(0, 2, 0, 0)
+							.Padding(4, 3, 0, 3)
 							[
 								SNew(SHorizontalBox)
 									+ SHorizontalBox::Slot()
 									[
-										SNew(STextBlock)
-											.Font(IDetailLayoutBuilder::GetDetailFont())
-											.Text(Target->TargetName)
+										SNew(SInlineEditableTextBlock)
+											.Font(IDetailLayoutBuilder::GetDetailFontBold())
+											.Text(Target->TargetLayerDisplayName)
 											.ColorAndOpacity_Static(&FLandscapeEditorCustomNodeBuilder_TargetLayers::GetTargetTextColor, Target)
+											.OnTextCommitted_Lambda([Target]( const FText& Text, ETextCommit::Type Type)
+											{
+												FScopedTransaction Transaction(LOCTEXT("LandscapeWeightLayer_Rename", "Rename Weight Layer"));
+												ALandscape* Landscape = Cast<ALandscape>(Target->Owner);
+												
+												const TMap<FName, FLandscapeTargetLayerSettings>& TargetLayers = Landscape->GetTargetLayers();
+												const FLandscapeTargetLayerSettings* LayerSettings = nullptr;
+												
+												if (Target->LayerName == FName(Text.ToString()))
+												{
+													LayerSettings = TargetLayers.Find(FName(Target->TargetLayerDisplayName.ToString()));	
+												}
+												
+												Landscape->RemoveTargetLayer(FName(Target->TargetLayerDisplayName.ToString()));
+											
+												Target->TargetLayerDisplayName = Text;
+												Target->LayerName = FName(Text.ToString());
+												Landscape->AddTargetLayer(Target->LayerName, LayerSettings ? *LayerSettings : FLandscapeTargetLayerSettings());	
+												
+												Target->LandscapeInfo->UpdateLayerInfoMap();
+												if (FEdModeLandscape* LandscapeEdMode = GetEditorMode())
+												{
+													LandscapeEdMode->UpdateTargetList();
+												}
+											})
 									]
 									+ SHorizontalBox::Slot()
 									.HAlign(HAlign_Right)
@@ -705,7 +820,6 @@ TSharedPtr<SWidget> FLandscapeEditorCustomNodeBuilder_TargetLayers::GenerateRow(
 											.ForegroundColor(FSlateColor::UseForeground())
 											.IsFocusable(false)
 											.ToolTipText(LOCTEXT("Tooltip_Delete", "Delete Layer"))
-											.Visibility_Static(&FLandscapeEditorCustomNodeBuilder_TargetLayers::GetTargetLayerDeleteVisibility, Target)
 											.OnClicked_Static(&FLandscapeEditorCustomNodeBuilder_TargetLayers::OnTargetLayerDeleteClicked, Target)
 											[
 												SNew(SImage)
@@ -1196,7 +1310,10 @@ void FLandscapeEditorCustomNodeBuilder_TargetLayers::OnTargetLayerSetObject(cons
 		if (ensure(SelectedLayerInfo->LayerName == Target->GetLayerName()))
 		{
 			ULandscapeInfo* LandscapeInfo = Target->LandscapeInfo.Get();
-			LandscapeInfo->Modify();
+			ALandscape* LandscapeActor = LandscapeInfo->LandscapeActor.Get();
+			
+			LandscapeActor->AddTargetLayer(Target->LayerName, FLandscapeTargetLayerSettings(SelectedLayerInfo));
+			
 			if (Target->LayerInfoObj.IsValid())
 			{
 				int32 Index = LandscapeInfo->GetLayerInfoIndex(Target->LayerInfoObj.Get(), Target->Owner.Get());
@@ -1217,7 +1334,7 @@ void FLandscapeEditorCustomNodeBuilder_TargetLayers::OnTargetLayerSetObject(cons
 					FLandscapeInfoLayerSettings& LayerSettings = LandscapeInfo->Layers[Index];
 					LayerSettings.LayerInfoObj = SelectedLayerInfo;
 
-					Target->LandscapeInfo->CreateLayerEditorSettingsFor(SelectedLayerInfo);
+					Target->LandscapeInfo->CreateTargetLayerSettingsFor(SelectedLayerInfo);
 				}
 			}
 						
@@ -1344,7 +1461,7 @@ void FLandscapeEditorCustomNodeBuilder_TargetLayers::OnTargetLayerCreateClicked(
 			}
 
 			Target->LayerInfoObj = LayerInfo;
-			Target->LandscapeInfo->CreateLayerEditorSettingsFor(LayerInfo);
+			Target->LandscapeInfo->CreateTargetLayerSettingsFor(LayerInfo);
 
 			// Notify the asset registry
 			FAssetRegistryModule::AssetCreated(LayerInfo);
@@ -1356,10 +1473,16 @@ void FLandscapeEditorCustomNodeBuilder_TargetLayers::OnTargetLayerCreateClicked(
 			TArray<UObject*> Objects;
 			Objects.Add(LayerInfo);
 			GEditor->SyncBrowserToObjects(Objects);
-			
-			LandscapeEdMode->TargetsListUpdated.Broadcast();
+
+			ALandscape* LandscapeActor = Target->LandscapeInfo->LandscapeActor.Get();
+			LandscapeActor->UpdateTargetLayer(FName(LayerName), FLandscapeTargetLayerSettings(LayerInfo));
 
 			FillEmptyLayers(LandscapeInfo, LayerInfo);
+
+			if (FEdModeLandscape* Mode = (FEdModeLandscape*)GLevelEditorModeTools().GetActiveMode(FBuiltinEditorModes::EM_Landscape))
+			{
+				Mode->UpdateTargetList();
+			}
 		}
 	}
 }
@@ -1374,10 +1497,10 @@ FReply FLandscapeEditorCustomNodeBuilder_TargetLayers::OnTargetLayerMakePublicCl
 	bool bSucceed = ObjectTools::RenameObjects(Objects, false, TEXT(""), Path);
 	if (bSucceed)
 	{
-		FEdModeLandscape* LandscapeEdMode = GetEditorMode();
-		if (LandscapeEdMode)
+		FEdModeLandscape* Mode = (FEdModeLandscape*)GLevelEditorModeTools().GetActiveMode(FBuiltinEditorModes::EM_Landscape);
+		if (Mode)
 		{
-			LandscapeEdMode->UpdateTargetList();
+			Mode->UpdateTargetList();
 		}
 	}
 	else
@@ -1516,8 +1639,7 @@ void FLandscapeEditorCustomNodeBuilder_TargetLayers::OnDebugModeColorChannelChan
 				}
 				LandscapeInfo->UpdateDebugColorMaterial();
 
-				FEdModeLandscape* LandscapeEdMode = GetEditorMode();
-				if (LandscapeEdMode)
+				if (FEdModeLandscape* LandscapeEdMode = GetEditorMode())
 				{
 					LandscapeEdMode->UpdateTargetList();
 				}
diff --git a/Engine/Source/Editor/LandscapeEditor/Private/LandscapeEditorDetailCustomization_TargetLayers.h b/Engine/Source/Editor/LandscapeEditor/Private/LandscapeEditorDetailCustomization_TargetLayers.h
--- a/Engine/Source/Editor/LandscapeEditor/Private/LandscapeEditorDetailCustomization_TargetLayers.h
+++ b/Engine/Source/Editor/LandscapeEditor/Private/LandscapeEditorDetailCustomization_TargetLayers.h
@@ -105,6 +105,8 @@ protected:
 	TOptional<SDragAndDropVerticalBox::EItemDropZone> HandleCanAcceptDrop(const FDragDropEvent& DragDropEvent, SDragAndDropVerticalBox::EItemDropZone DropZone, SVerticalBox::FSlot* Slot);
 	FReply HandleAcceptDrop(FDragDropEvent const& DragDropEvent, SDragAndDropVerticalBox::EItemDropZone DropZone, int32 SlotIndex, SVerticalBox::FSlot* Slot);
 
+	FReply HandleCreateLayersFromMaterials();
+	void HandleCreateLayer();
 	const FSlateBrush* GetTargetLayerDisplayOrderBrush() const;
 	TSharedRef<SWidget> GetTargetLayerDisplayOrderButtonMenuContent();
 	void SetSelectedDisplayOrder(ELandscapeLayerDisplayMode InDisplayOrder);
diff --git a/Engine/Source/Editor/LandscapeEditor/Public/LandscapeImportHelper.h b/Engine/Source/Editor/LandscapeEditor/Public/LandscapeImportHelper.h
--- a/Engine/Source/Editor/LandscapeEditor/Public/LandscapeImportHelper.h
+++ b/Engine/Source/Editor/LandscapeEditor/Public/LandscapeImportHelper.h
@@ -51,7 +51,7 @@ enum class ELandscapeImportTransformType : int8
 	ExpandOffset UMETA(DisplayName="Expand", ToolTip="Will Import the data at the gizmo location and expand the data to fill the landscape") ,
 	ExpandCentered UMETA(Hidden), 
 	Resample UMETA(ToolTip="Will resample Import data to fit landscape"),
-	Subregion UMETA(ToolTop="Import Sub-region of the Image to Landscape")
+	Subregion UMETA(ToolTip="Import Sub-region of the Image to Landscape")
 };
 
 FORCEINLINE bool operator==(const FLandscapeImportResolution& Lhs, const FLandscapeImportResolution& Rhs)
diff --git a/Engine/Source/Editor/WorldBrowser/Private/Tiles/WorldTileCollectionModel.cpp b/Engine/Source/Editor/WorldBrowser/Private/Tiles/WorldTileCollectionModel.cpp
--- a/Engine/Source/Editor/WorldBrowser/Private/Tiles/WorldTileCollectionModel.cpp
+++ b/Engine/Source/Editor/WorldBrowser/Private/Tiles/WorldTileCollectionModel.cpp
@@ -1638,7 +1638,7 @@ void FWorldTileCollectionModel::ImportTiledLandscape_Executed()
 			Landscape->SetLandscapeGuid(LandscapeGuid);
 			for (const auto& ImportLayerInfo : ImportLayers)
 			{
-				Landscape->EditorLayerSettings.Add(FLandscapeEditorLayerSettings(ImportLayerInfo.LayerInfo));
+				Landscape->AddTargetLayer(ImportLayerInfo.LayerName, FLandscapeTargetLayerSettings(ImportLayerInfo.LayerInfo));
 			}
 			Landscape->CreateLandscapeInfo();
 		}
@@ -1785,8 +1785,9 @@ void FWorldTileCollectionModel::ReimportTiledLandscape_Executed(FName TargetLaye
 		}
 		else // Weightmap
 		{
-			for (FLandscapeEditorLayerSettings& LayerSettings : Landscape->EditorLayerSettings)
+			for (const TPair<FName, FLandscapeTargetLayerSettings>& Pair : Landscape->GetTargetLayers())
 			{
+				const FLandscapeTargetLayerSettings& LayerSettings = Pair.Value;
 				if (LayerSettings.LayerInfoObj && (LayerSettings.LayerInfoObj->LayerName == TargetLayer || TargetLayer == NAME_None))
 				{
 					if (!LayerSettings.ReimportLayerFilePath.IsEmpty())
diff --git a/Engine/Source/Editor/WorldBrowser/Private/Tiles/WorldTileModel.cpp b/Engine/Source/Editor/WorldBrowser/Private/Tiles/WorldTileModel.cpp
--- a/Engine/Source/Editor/WorldBrowser/Private/Tiles/WorldTileModel.cpp
+++ b/Engine/Source/Editor/WorldBrowser/Private/Tiles/WorldTileModel.cpp
@@ -1053,7 +1053,7 @@ ALandscapeProxy* FWorldTileModel::ImportLandscapeTile(const FLandscapeImportSett
 
 	for (const FLandscapeImportLayerInfo& ImportLayerInfo : Settings.ImportLayers)
 	{
-		LandscapeProxy->EditorLayerSettings.Add(FLandscapeEditorLayerSettings(ImportLayerInfo.LayerInfo, ImportLayerInfo.SourceFilePath));
+		LandscapeProxy->AddTargetLayer(ImportLayerInfo.LayerInfo->LayerName, FLandscapeTargetLayerSettings(ImportLayerInfo.LayerInfo, ImportLayerInfo.SourceFilePath));
 	}
 
 	return LandscapeProxy;
diff --git a/Engine/Source/Runtime/Core/Public/UObject/FortniteMainBranchObjectVersion.h b/Engine/Source/Runtime/Core/Public/UObject/FortniteMainBranchObjectVersion.h
--- a/Engine/Source/Runtime/Core/Public/UObject/FortniteMainBranchObjectVersion.h
+++ b/Engine/Source/Runtime/Core/Public/UObject/FortniteMainBranchObjectVersion.h
@@ -456,6 +456,9 @@ struct FFortniteMainBranchObjectVersion
 		// PCGPartitionActorDesc
 		PCGPartitionActorDesc,
 
+		// Target layers are now defined in the Landscape actor and not continuously synced from the assigned material.
+		LandscapeTargetLayersInLandscapeActor,
+
 		// -----<new versions can be added above this line>-------------------------------------------------
 		VersionPlusOne,
 		LatestVersion = VersionPlusOne - 1
diff --git a/Engine/Source/Runtime/Landscape/Classes/Landscape.h b/Engine/Source/Runtime/Landscape/Classes/Landscape.h
--- a/Engine/Source/Runtime/Landscape/Classes/Landscape.h
+++ b/Engine/Source/Runtime/Landscape/Classes/Landscape.h
@@ -11,6 +11,7 @@
 #include "Delegates/DelegateCombinations.h"
 #include "Templates/SubclassOf.h"
 #include "Algo/Transform.h"
+#include "LandscapeLayerInfoObject.h"
 
 #include "Landscape.generated.h"
 
diff --git a/Engine/Source/Runtime/Landscape/Classes/LandscapeInfo.h b/Engine/Source/Runtime/Landscape/Classes/LandscapeInfo.h
--- a/Engine/Source/Runtime/Landscape/Classes/LandscapeInfo.h
+++ b/Engine/Source/Runtime/Landscape/Classes/LandscapeInfo.h
@@ -25,6 +25,7 @@ class ULandscapeLayerInfoObject;
 class ULevel;
 class UMaterialInstanceConstant;
 struct FLandscapeEditorLayerSettings;
+struct FLandscapeTargetLayerSettings;
 class ULandscapeSplinesComponent;
 class ULandscapeSplineControlPoint;
 class ULandscapeSplineSegment;
@@ -114,7 +115,13 @@ struct FLandscapeInfoLayerSettings
 	LANDSCAPE_API FName GetLayerName() const;
 
 #if WITH_EDITORONLY_DATA
+PRAGMA_DISABLE_DEPRECATION_WARNINGS	
+	UE_DEPRECATED(5.5, "This property has been deprecated, please use the GetTargetLayerSettings instead")
 	LANDSCAPE_API FLandscapeEditorLayerSettings& GetEditorSettings() const;
+PRAGMA_ENABLE_DEPRECATION_WARNINGS
+	
+	LANDSCAPE_API const FLandscapeTargetLayerSettings& GetTargetLayerSettings() const;
+
 #endif
 };
 
@@ -240,15 +247,21 @@ public:
 	LANDSCAPE_API void UpdateAllAddCollisions();
 	LANDSCAPE_API void UpdateAddCollision(FIntPoint LandscapeKey);
 
-	LANDSCAPE_API FLandscapeEditorLayerSettings& GetLayerEditorSettings(ULandscapeLayerInfoObject* LayerInfo) const;
-	LANDSCAPE_API void CreateLayerEditorSettingsFor(ULandscapeLayerInfoObject* LayerInfo);
+	UE_DEPRECATED(5.5, "This property has been deprecated, please use the GetTargetLayerSettings(ULandscapeLayerInfoObject* LayerInfo) instead")
+	LANDSCAPE_API FLandscapeTargetLayerSettings& GetLayerEditorSettings(ULandscapeLayerInfoObject* LayerInfo) const;
+	
+	LANDSCAPE_API const FLandscapeTargetLayerSettings& GetTargetLayerSettings(ULandscapeLayerInfoObject* LayerInfo) const;
+
+	UE_DEPRECATED(5.5, "This property has been deprecated, please use the CreateTargetLayerSettingsFor instead")
+	LANDSCAPE_API void CreateLayerEditorSettingsFor(ULandscapeLayerInfoObject* LayerInfo) {}
+	
+	LANDSCAPE_API void CreateTargetLayerSettingsFor(ULandscapeLayerInfoObject* LayerInfo);
 
 	LANDSCAPE_API ULandscapeLayerInfoObject* GetLayerInfoByName(FName LayerName, ALandscapeProxy* Owner = nullptr) const;
 	LANDSCAPE_API int32 GetLayerInfoIndex(FName LayerName, ALandscapeProxy* Owner = nullptr) const;
 	LANDSCAPE_API int32 GetLayerInfoIndex(ULandscapeLayerInfoObject* LayerInfo, ALandscapeProxy* Owner = nullptr) const;
 	LANDSCAPE_API bool UpdateLayerInfoMap(ALandscapeProxy* Proxy = nullptr, bool bInvalidate = false);
 
-
 	LANDSCAPE_API bool CanDeleteLandscape(FText& OutReason) const;
 
 	/**
@@ -445,7 +458,7 @@ private:
 	bool ApplySplinesInternal(bool bOnlySelected, TScriptInterface<ILandscapeSplineInterface> SplineOwner, TSet<TObjectPtr<ULandscapeComponent>>* OutModifiedComponents, bool bMarkPackageDirty, int32 LandscapeMinX, int32 LandscapeMinY, int32 LandscapeMaxX, int32 LandscapeMaxY, TFunctionRef<TSharedPtr<FModulateAlpha>(ULandscapeLayerInfoObject*)> GetOrCreateModulate);
 	void MoveSegment(ULandscapeSplineSegment* InSegment, TScriptInterface<ILandscapeSplineInterface> From, TScriptInterface<ILandscapeSplineInterface> To);
 	void MoveControlPoint(ULandscapeSplineControlPoint* InControlPoint, TScriptInterface<ILandscapeSplineInterface> From, TScriptInterface<ILandscapeSplineInterface> To);
-	bool UpdateLayerInfoMapInternal(ALandscapeProxy* Proxy, bool bInvalidate);
+	void UpdateLayerInfoMapInternal(ALandscapeProxy* Proxy);
 	bool TryAddToModifiedPackages(UPackage* InPackage, const ALandscape* InLandscapeOverride = nullptr);
 #endif
 };
diff --git a/Engine/Source/Runtime/Landscape/Classes/LandscapeProxy.h b/Engine/Source/Runtime/Landscape/Classes/LandscapeProxy.h
--- a/Engine/Source/Runtime/Landscape/Classes/LandscapeProxy.h
+++ b/Engine/Source/Runtime/Landscape/Classes/LandscapeProxy.h
@@ -74,7 +74,34 @@ DECLARE_MULTICAST_DELEGATE_TwoParams(FOnLandscapeProxyComponentDataChanged, ALan
 // ----------------------------------------------------------------------------------
 
 USTRUCT()
-struct FLandscapeEditorLayerSettings
+struct FLandscapeTargetLayerSettings
+{
+	GENERATED_USTRUCT_BODY()
+
+	UPROPERTY(Category = "General", EditAnywhere)
+	TObjectPtr<ULandscapeLayerInfoObject> LayerInfoObj;
+	
+#if WITH_EDITORONLY_DATA
+	UPROPERTY()
+	FString ReimportLayerFilePath;
+
+	FLandscapeTargetLayerSettings() = default;
+
+	explicit FLandscapeTargetLayerSettings(ULandscapeLayerInfoObject* InLayerInfo, const FString& InFilePath = FString())
+		: LayerInfoObj(InLayerInfo)
+		, ReimportLayerFilePath(InFilePath)
+	{
+	}
+
+	bool operator==(const FLandscapeTargetLayerSettings& Other) const
+    {
+    	return LayerInfoObj == Other.LayerInfoObj;
+    }
+#endif // WITH_EDITORONLY_DATA
+};
+
+USTRUCT()
+struct UE_DEPRECATED(5.5, "FLandscapeEditorLayerSettings is deprecated; please use FLandscapeTargetLayerSettings instead")  FLandscapeEditorLayerSettings
 {
 	GENERATED_USTRUCT_BODY()
 
@@ -442,6 +469,10 @@ protected:
 	bool bDisableRuntimeGrassMapGeneration = false;
 
 public:
+#if WITH_EDITORONLY_DATA
+	UPROPERTY(Transient)
+	bool bInitWeightLayersFromMaterial;
+#endif //WITH_EDITORONLY_DATA
 	/** Offset in quads from global components grid origin (in quads) **/
 	UPROPERTY()
 	FIntPoint LandscapeSectionOffset;
@@ -804,9 +835,12 @@ public:
 	UPROPERTY()
 	FGuid ReimportDestinationLayerGuid;
 
+PRAGMA_DISABLE_DEPRECATION_WARNINGS
+	UE_DEPRECATED(5.5, "TargetLayers will be used instead")
 	UPROPERTY()
-	TArray<FLandscapeEditorLayerSettings> EditorLayerSettings;
-
+	TArray<FLandscapeEditorLayerSettings> EditorLayerSettings_DEPRECATED;
+PRAGMA_ENABLE_DEPRECATION_WARNINGS
+	
 	TMap<UTexture2D*, FLandscapeEditLayerReadback*> HeightmapsCPUReadback;
 	TMap<UTexture2D*, FLandscapeEditLayerReadback*> WeightmapsCPUReadback;
 
@@ -1001,6 +1035,9 @@ public:
 	static constexpr int32 NaniteComponentMaxSide = 8;
 	static constexpr int32 NaniteMaxComponents = NaniteComponentMaxSide * NaniteComponentMaxSide;
 	int32 NumNaniteRequiredComponents() const { return FMath::DivideAndRoundUp(LandscapeComponents.Num(), NaniteMaxComponents); }
+
+	UE_DEPRECATED(5.5, "This version of HasLayer is deprecated.  Use HasTargetLayer(ULandscapeLayerInfoObject* LayerInfoObject) instead.")
+	bool HasLayer(ULandscapeLayerInfoObject* LayerInfoObject) const;
 #endif	//WITH_EDITOR
 
 	bool AreNaniteComponentsValid(const FGuid& InProxyContentId) const;
@@ -1236,6 +1273,12 @@ public:
 	LANDSCAPE_API static const TArray<FName>& GetLayersFromMaterial(UMaterialInterface* Material);
 	LANDSCAPE_API const TArray<FName>& GetLayersFromMaterial() const;
 
+	// Get all Layer names used by the materials assigned the proxy and components ( Proxy Material, Component, Component LOD Material and Hole Material)
+	LANDSCAPE_API TArray<FName> RetrieveAllLayerNamesFromMaterials() const;
+	
+	// Query all the LandscapeLayerInfo Objects from the weight layers allocated on this proxy.
+	LANDSCAPE_API TMap<FName, ULandscapeLayerInfoObject*> RetrieveAllocationInfos() const;
+
 	/**
 	* Creates a new LandscapeLayerInfoObject
 	* @param	LayerName	The name of the created Layer
@@ -1482,7 +1525,33 @@ public:
 	 */
 	UFUNCTION(BlueprintCallable, Category = "Landscape")
 	LANDSCAPE_API virtual void DeleteUnusedLayers();
+	
+	// TargetLayer API Start
+	// ---------------------
+	
+	// Add a TargetLayer with a unique name and no layer object info
+	LANDSCAPE_API FLandscapeTargetLayerSettings& AddTargetLayer();
+
+	// Add a named TargetLayer 
+	LANDSCAPE_API FLandscapeTargetLayerSettings& AddTargetLayer(const FName& Name, const FLandscapeTargetLayerSettings& TargetLayerSettings, bool bPostEditChange = true);
+
+	// Update a existing Target Layer
+	LANDSCAPE_API bool UpdateTargetLayer(const FName& Name, const FLandscapeTargetLayerSettings& TargetLayerSettings, bool bPostEditChange = true);
+
+	// Remove Named TargetLayer,
+	// Returns true if a layer has been removed.
+	LANDSCAPE_API bool RemoveTargetLayer(const FName& Name, bool bPostEditChange = true);
+
+	// Query if a TargetLayer exists by Name or by checking the ULandscapeLayerInfoObject in  TargetLayerSettings
+	LANDSCAPE_API bool HasTargetLayer(const FName& Name) const;
+	LANDSCAPE_API bool HasTargetLayer(const FLandscapeTargetLayerSettings& TargetLayerSettings) const;
+	LANDSCAPE_API bool HasTargetLayer(const ULandscapeLayerInfoObject* LayerInfoObject) const;
 
+	// All the Target layers 
+	LANDSCAPE_API const TMap<FName, FLandscapeTargetLayerSettings>& GetTargetLayers() const;
+
+	// TargetLayer API End
+	// -------------------
 protected:
 	friend class ALandscape;
 
@@ -1520,7 +1589,10 @@ protected:
 	bool bUpgradeSharedPropertiesPerformed = false;
 #endif // WITH_EDITOR
 private:
-
+	
+	UPROPERTY(Category = "Target Layers", VisibleAnywhere, meta = (NoResetToDefault, LandscapeInherited) )
+	TMap<FName, FLandscapeTargetLayerSettings> TargetLayers;
+	
 	/** Returns Grass Update interval */
 	FORCEINLINE int32 GetGrassUpdateInterval() const 
 	{
diff --git a/Engine/Source/Runtime/Landscape/Private/Landscape.cpp b/Engine/Source/Runtime/Landscape/Private/Landscape.cpp
--- a/Engine/Source/Runtime/Landscape/Private/Landscape.cpp
+++ b/Engine/Source/Runtime/Landscape/Private/Landscape.cpp
@@ -3837,38 +3837,49 @@ FName FLandscapeInfoLayerSettings::GetLayerName() const
 	return LayerName;
 }
 
+PRAGMA_DISABLE_DEPRECATION_WARNINGS
 FLandscapeEditorLayerSettings& FLandscapeInfoLayerSettings::GetEditorSettings() const
+{
+	static FLandscapeEditorLayerSettings DeprecatedSettings;
+	return DeprecatedSettings;
+}
+PRAGMA_ENABLE_DEPRECATION_WARNINGS
+
+const FLandscapeTargetLayerSettings& FLandscapeInfoLayerSettings::GetTargetLayerSettings() const
 {
 	check(LayerInfoObj);
 
 	ULandscapeInfo* LandscapeInfo = Owner->GetLandscapeInfo();
-	return LandscapeInfo->GetLayerEditorSettings(LayerInfoObj);
+	return LandscapeInfo->GetTargetLayerSettings(LayerInfoObj);
 }
 
-FLandscapeEditorLayerSettings& ULandscapeInfo::GetLayerEditorSettings(ULandscapeLayerInfoObject* LayerInfo) const
+FLandscapeTargetLayerSettings& ULandscapeInfo::GetLayerEditorSettings(ULandscapeLayerInfoObject* LayerInfo) const
+{
+	static FLandscapeTargetLayerSettings DeprecatedSettings;
+	return DeprecatedSettings;
+}
+
+const FLandscapeTargetLayerSettings& ULandscapeInfo::GetTargetLayerSettings(ULandscapeLayerInfoObject* LayerInfo) const
 {
 	ALandscapeProxy* Proxy = GetLandscapeProxy();
-	FLandscapeEditorLayerSettings* EditorLayerSettings = Proxy->EditorLayerSettings.FindByKey(LayerInfo);
-	if (EditorLayerSettings)
+	const FName* LayerName = Proxy->GetTargetLayers().FindKey(FLandscapeTargetLayerSettings(LayerInfo));
+	if (LayerName)
 	{
-		return *EditorLayerSettings;
+		return *Proxy->GetTargetLayers().Find(*LayerName);
 	}
 	else
 	{
-		int32 Index = Proxy->EditorLayerSettings.Add(FLandscapeEditorLayerSettings(LayerInfo));
-		return Proxy->EditorLayerSettings[Index];
+		return Proxy->AddTargetLayer(LayerInfo->LayerName, FLandscapeTargetLayerSettings(LayerInfo));
 	}
 }
 
-void ULandscapeInfo::CreateLayerEditorSettingsFor(ULandscapeLayerInfoObject* LayerInfo)
+void ULandscapeInfo::CreateTargetLayerSettingsFor(ULandscapeLayerInfoObject* LayerInfo)
 {
 	ForEachLandscapeProxy([LayerInfo](ALandscapeProxy* Proxy)
 	{
-		FLandscapeEditorLayerSettings* EditorLayerSettings = Proxy->EditorLayerSettings.FindByKey(LayerInfo);
-		if (!EditorLayerSettings)
+		if (!Proxy->HasTargetLayer(LayerInfo->LayerName))
 		{
-			Proxy->Modify();
-			Proxy->EditorLayerSettings.Add(FLandscapeEditorLayerSettings(LayerInfo));
+			Proxy->AddTargetLayer(LayerInfo->LayerName, FLandscapeTargetLayerSettings(LayerInfo));
 		}
 		return true;
 	});
@@ -3917,193 +3928,71 @@ int32 ULandscapeInfo::GetLayerInfoIndex(FName LayerName, ALandscapeProxy* Owner
 }
 
 
-bool ULandscapeInfo::UpdateLayerInfoMapInternal(ALandscapeProxy* Proxy, bool bInvalidate)
+void ULandscapeInfo::UpdateLayerInfoMapInternal(ALandscapeProxy* Proxy)
 {
 	TRACE_CPUPROFILER_EVENT_SCOPE(ULandscapeInfo::UpdateLayerInfoMapInternal);
 
-	bool bHasCollision = false;
-	if (GIsEditor)
+	if (GIsEditor && Proxy && Proxy->bInitWeightLayersFromMaterial)
 	{
-		if (Proxy)
+		Proxy->bInitWeightLayersFromMaterial = false;
+		// add the landscape layers from the material on the proxy & LandscapeLayerInfos from the WeightMap Allocations
+		if (LandscapeActor.IsValid())
 		{
-			if (bInvalidate)
+			TArray<FName> LayerNames = Proxy->RetrieveAllLayerNamesFromMaterials();
+			const TMap<FName, ULandscapeLayerInfoObject*> InfoObjects = Proxy->RetrieveAllocationInfos();
+			for(const FName& LayerName : LayerNames)
 			{
-				// this is a horribly dangerous combination of parameters...
-
-				for (int32 i = 0; i < Layers.Num(); i++)
+				ULandscapeLayerInfoObject* const * LayerInfoObject = InfoObjects.Find(LayerName);
+				if (!LandscapeActor->HasTargetLayer(LayerName))
 				{
-					if (Layers[i].Owner == Proxy)
-					{
-						Layers.RemoveAt(i--);
-					}
+					LandscapeActor->AddTargetLayer(LayerName, FLandscapeTargetLayerSettings(LayerInfoObject ? *LayerInfoObject : nullptr), false);
 				}
 			}
-			else // Proxy && !bInvalidate
+			
+			for (const auto& Layer : LandscapeActor->GetTargetLayers())
 			{
-				TArray<FName> LayerNames = Proxy->GetLayersFromMaterial();
-
-				// Validate any existing layer infos owned by this proxy
-				for (int32 i = 0; i < Layers.Num(); i++)
-				{
-					if (Layers[i].Owner == Proxy)
-					{
-						Layers[i].bValid = LayerNames.Contains(Layers[i].GetLayerName());
-					}
-				}
-
-				// Add placeholders for any unused material layers
-				for (int32 i = 0; i < LayerNames.Num(); i++)
-				{
-					int32 LayerInfoIndex = GetLayerInfoIndex(LayerNames[i]);
-					if (LayerInfoIndex == INDEX_NONE)
-					{
-						FLandscapeInfoLayerSettings LayerSettings(LayerNames[i], Proxy);
-						LayerSettings.bValid = true;
-						Layers.Add(LayerSettings);
-					}
-				}
-
-				// Populate from layers used in components
-				for (int32 ComponentIndex = 0; ComponentIndex < Proxy->LandscapeComponents.Num(); ComponentIndex++)
-				{
-					ULandscapeComponent* Component = Proxy->LandscapeComponents[ComponentIndex];
-
-					// Add layers from per-component override materials
-					if ((Component != nullptr) && (Component->OverrideMaterial != nullptr))
-					{
-						TArray<FName> ComponentLayerNames = Proxy->GetLayersFromMaterial(Component->OverrideMaterial);
-						for (int32 i = 0; i < ComponentLayerNames.Num(); i++)
-						{
-							int32 LayerInfoIndex = GetLayerInfoIndex(ComponentLayerNames[i]);
-							if (LayerInfoIndex == INDEX_NONE)
-							{
-								FLandscapeInfoLayerSettings LayerSettings(ComponentLayerNames[i], Proxy);
-								LayerSettings.bValid = true;
-								Layers.Add(LayerSettings);
-							}
-						}
-					}
-
-					const TArray<FWeightmapLayerAllocationInfo>& ComponentWeightmapLayerAllocations = Component->GetWeightmapLayerAllocations();
-
-					for (int32 AllocationIndex = 0; AllocationIndex < ComponentWeightmapLayerAllocations.Num(); AllocationIndex++)
-					{
-						ULandscapeLayerInfoObject* LayerInfo = ComponentWeightmapLayerAllocations[AllocationIndex].LayerInfo;
-						if (LayerInfo)
-						{
-							int32 LayerInfoIndex = GetLayerInfoIndex(LayerInfo);
-							bool bValid = LayerNames.Contains(LayerInfo->LayerName);
-
-							if (bValid)
-							{
-								//LayerInfo->IsReferencedFromLoadedData = true;
-							}
-
-							if (LayerInfoIndex != INDEX_NONE)
-							{
-								FLandscapeInfoLayerSettings& LayerSettings = Layers[LayerInfoIndex];
-
-								// Valid layer infos take precedence over invalid ones
-								// Landscape Actors take precedence over Proxies
-								if ((bValid && !LayerSettings.bValid)
-									|| (bValid == LayerSettings.bValid && Proxy->IsA<ALandscape>()))
-								{
-									LayerSettings.Owner = Proxy;
-									LayerSettings.bValid = bValid;
-									LayerSettings.ThumbnailMIC = nullptr;
-								}
-							}
-							else
-							{
-								// handle existing placeholder layers
-								LayerInfoIndex = GetLayerInfoIndex(LayerInfo->LayerName);
-								if (LayerInfoIndex != INDEX_NONE)
-								{
-									FLandscapeInfoLayerSettings& LayerSettings = Layers[LayerInfoIndex];
-
-									//if (LayerSettings.Owner == Proxy)
-									{
-										LayerSettings.Owner = Proxy;
-										LayerSettings.LayerInfoObj = LayerInfo;
-										LayerSettings.bValid = bValid;
-										LayerSettings.ThumbnailMIC = nullptr;
-									}
-								}
-								else
-								{
-									FLandscapeInfoLayerSettings LayerSettings(LayerInfo, Proxy);
-									LayerSettings.bValid = bValid;
-									Layers.Add(LayerSettings);
-								}
-							}
-						}
-					}
-				}
-
-				// Add any layer infos cached in the actor
-				Proxy->EditorLayerSettings.RemoveAll([](const FLandscapeEditorLayerSettings& Settings) { return Settings.LayerInfoObj == nullptr; });
-				for (int32 i = 0; i < Proxy->EditorLayerSettings.Num(); i++)
+				if (!Layer.Value.LayerInfoObj)
 				{
-					FLandscapeEditorLayerSettings& EditorLayerSettings = Proxy->EditorLayerSettings[i];
-					if (LayerNames.Contains(EditorLayerSettings.LayerInfoObj->LayerName))
+					ULandscapeLayerInfoObject* const * LayerInfoObject = InfoObjects.Find(Layer.Key);
+					if (LayerInfoObject && LandscapeActor->HasTargetLayer(Layer.Key))
 					{
-						// intentionally using the layer name here so we don't add layer infos from
-						// the cache that have the same name as an actual assignment from a component above
-						int32 LayerInfoIndex = GetLayerInfoIndex(EditorLayerSettings.LayerInfoObj->LayerName);
-						if (LayerInfoIndex != INDEX_NONE)
-						{
-							FLandscapeInfoLayerSettings& LayerSettings = Layers[LayerInfoIndex];
-							if (LayerSettings.LayerInfoObj == nullptr)
-							{
-								LayerSettings.Owner = Proxy;
-								LayerSettings.LayerInfoObj = EditorLayerSettings.LayerInfoObj;
-								LayerSettings.bValid = true;
-							}
-						}
-					}
-					else
-					{
-						Proxy->Modify();
-						Proxy->EditorLayerSettings.RemoveAt(i--);
+						LandscapeActor->UpdateTargetLayer(Layer.Key, FLandscapeTargetLayerSettings(*LayerInfoObject), false);
 					}
 				}
+			}
+		}
+	}
 
-				// Add Visibility Layer info if not initialized
-				if (ALandscapeProxy::VisibilityLayer != nullptr)
-				{
-					int32 LayerInfoIndex = GetLayerInfoIndex(ALandscapeProxy::VisibilityLayer->LayerName);
+	if (LandscapeActor.IsValid())
+	{
+		Layers.Empty();
 
-					if ((LayerInfoIndex != INDEX_NONE) && (Layers[LayerInfoIndex].LayerInfoObj == nullptr))
-					{
-						Layers[LayerInfoIndex].LayerInfoObj = ALandscapeProxy::VisibilityLayer;
-					}
-				}
-			}
+		for (const TTuple<FName, FLandscapeTargetLayerSettings>& TargetLayer : LandscapeActor->GetTargetLayers())
+		{
+			FLandscapeInfoLayerSettings InfoLayerSettings (TargetLayer.Key, LandscapeActor.Get());
+			InfoLayerSettings.bValid = true;
+			InfoLayerSettings.LayerInfoObj = TargetLayer.Value.LayerInfoObj;
+			Layers.Add(InfoLayerSettings);
 		}
-		else // !Proxy
+		
+		// Add Visibility Layer info if not initialized
+		if (ALandscapeProxy::VisibilityLayer != nullptr)
 		{
-			Layers.Empty();
+			int32 LayerInfoIndex = GetLayerInfoIndex(ALandscapeProxy::VisibilityLayer->LayerName);
 
-			if (!bInvalidate)
+			if ((LayerInfoIndex != INDEX_NONE) && (Layers[LayerInfoIndex].LayerInfoObj == nullptr))
 			{
-				ForEachLandscapeProxy([this](ALandscapeProxy* EachProxy)
-				{
-					if (!EachProxy->IsPendingKillPending())
-					{
-						checkSlow(EachProxy->GetLandscapeInfo() == this);
-						UpdateLayerInfoMapInternal(EachProxy, false);
-					}
-					return true;
-				});
+				Layers[LayerInfoIndex].LayerInfoObj = ALandscapeProxy::VisibilityLayer;
 			}
 		}
 	}
-	return bHasCollision;
+	
 }
 
 bool ULandscapeInfo::UpdateLayerInfoMap(ALandscapeProxy* Proxy /*= nullptr*/, bool bInvalidate /*= false*/)
 {
-	bool bResult = UpdateLayerInfoMapInternal(Proxy, bInvalidate);
+	UpdateLayerInfoMapInternal(Proxy);
+
 	if (GIsEditor)
 	{
 		ALandscape* Landscape = LandscapeActor.Get();
@@ -4112,7 +4001,9 @@ bool ULandscapeInfo::UpdateLayerInfoMap(ALandscapeProxy* Proxy /*= nullptr*/, bo
 			Landscape->RequestLayersInitialization(/*bInRequestContentUpdate*/false);
 		}
 	}
-	return bResult;
+	// return value is unused and has no meaning.
+	return true;
+
 }
 
 #endif // WITH_EDITOR
@@ -4227,6 +4118,9 @@ void ALandscapeProxy::PostLoad()
 	}
 
 #if WITH_EDITOR
+	const int32 LinkerVersion = GetLinkerCustomVersion(FFortniteMainBranchObjectVersion::GUID);
+	bInitWeightLayersFromMaterial = LinkerVersion < FFortniteMainBranchObjectVersion::LandscapeTargetLayersInLandscapeActor;
+	
 	PRAGMA_DISABLE_DEPRECATION_WARNINGS
 		if (!LandscapeMaterialsOverride_DEPRECATED.IsEmpty())
 		{
@@ -4268,15 +4162,14 @@ void ALandscapeProxy::PostLoad()
 				RecreateCollisionComponents();
 			}
 		}
-
-	EditorLayerSettings.RemoveAll([](const FLandscapeEditorLayerSettings& Settings) { return Settings.LayerInfoObj == nullptr; });
-
+	
 PRAGMA_DISABLE_DEPRECATION_WARNINGS
+
 	if (EditorCachedLayerInfos_DEPRECATED.Num() > 0)
 	{
 		for (int32 i = 0; i < EditorCachedLayerInfos_DEPRECATED.Num(); i++)
 		{
-			EditorLayerSettings.Add(FLandscapeEditorLayerSettings(EditorCachedLayerInfos_DEPRECATED[i]));
+			TargetLayers.Add(EditorCachedLayerInfos_DEPRECATED[i]->LayerName, FLandscapeTargetLayerSettings(EditorCachedLayerInfos_DEPRECATED[i]));
 		}
 		EditorCachedLayerInfos_DEPRECATED.Empty();
 	}
@@ -5436,20 +5329,28 @@ void ULandscapeInfo::ForEachLandscapeProxy(TFunctionRef<bool(ALandscapeProxy*)>
 {
 	if (ALandscape* Landscape = LandscapeActor.Get())
 	{
-		if (!Fn(Landscape))
+		if (!Landscape->IsPendingKillPending())
 		{
-			return;
+			if ( !Fn(Landscape))
+			{
+				return;
+			}	
 		}
+		
 	}
 
 	for (TWeakObjectPtr<ALandscapeStreamingProxy> StreamingProxyPtr : StreamingProxies)
 	{
 		if (ALandscapeProxy* LandscapeProxy = StreamingProxyPtr.Get())
 		{
-			if (!Fn(LandscapeProxy))
+			if (!LandscapeProxy->IsPendingKillPending())
 			{
-				return;
+				if (!Fn(LandscapeProxy))
+				{
+					return;
+				}	
 			}
+			
 		}
 	}
 }
@@ -5579,6 +5480,10 @@ void ULandscapeInfo::RegisterActor(ALandscapeProxy* Proxy, bool bMapCheck, bool
 			LandscapeActor = Landscape;
 
 #if WITH_EDITOR
+			// Now we have associated a LandscapeActor with this info
+			// we can ask for the WeightMaps
+			UpdateLayerInfoMap(LandscapeActor.Get());
+			
 			// Update registered splines so they can pull the actor pointer
 			for (TScriptInterface<ILandscapeSplineInterface> SplineActor : SplineActors)
 			{
@@ -6523,6 +6428,105 @@ void ALandscapeProxy::EnableNaniteComponents(bool bInNaniteActive)
 	}
 }
 
+#if WITH_EDITOR
+bool ALandscapeProxy::HasLayer(ULandscapeLayerInfoObject* LayerInfoObject) const
+{
+	return TargetLayers.FindKey(FLandscapeTargetLayerSettings(LayerInfoObject)) == nullptr;
+}
+
+bool  ALandscapeProxy::RemoveTargetLayer(const FName& Name,  bool bPostEditChange)
+{
+	Modify();
+	
+	int32 NumItemsRemoved = TargetLayers.Remove(Name);
+	if (FProperty* Property = StaticClass()->FindPropertyByName(GET_MEMBER_NAME_CHECKED(ALandscapeProxy, TargetLayers)); Property != nullptr && bPostEditChange)
+	{
+		FPropertyChangedEvent PropertyChangedEvent(Property);
+		PostEditChangeProperty(PropertyChangedEvent);	
+	}
+	
+	return NumItemsRemoved > 0;
+}
+
+FLandscapeTargetLayerSettings& ALandscapeProxy::AddTargetLayer()
+{
+	int32 StartIndex = GetTargetLayers().Num();
+	FName NewName;
+	do
+	{
+		NewName = FName(FString::Format(TEXT("Layer_{0}"), { StartIndex++ } ));
+	} while (HasTargetLayer(NewName));
+
+	return AddTargetLayer(NewName, FLandscapeTargetLayerSettings());
+}
+	
+FLandscapeTargetLayerSettings& ALandscapeProxy::AddTargetLayer(const FName& Name, const FLandscapeTargetLayerSettings& TargetLayerSettings, bool bPostEditChange)
+{
+	Modify();
+
+	check(!HasTargetLayer(Name));
+	
+	FLandscapeTargetLayerSettings& Settings = TargetLayers.Add(Name, TargetLayerSettings);
+
+	if (FProperty* Property = StaticClass()->FindPropertyByName(GET_MEMBER_NAME_CHECKED(ALandscapeProxy, TargetLayers)); Property != nullptr && bPostEditChange)
+	{
+		FPropertyChangedEvent PropertyChangedEvent(Property);
+		PostEditChangeProperty(PropertyChangedEvent);	
+	}
+	
+	return Settings;
+}
+
+bool ALandscapeProxy::UpdateTargetLayer(const FName& Name, const FLandscapeTargetLayerSettings& InTargetLayerSettings, bool bPostEditChange)
+{
+	FLandscapeTargetLayerSettings* TargetLayerSettings = TargetLayers.Find(Name);
+
+	check(TargetLayerSettings);
+	if (TargetLayerSettings)
+	{
+		*TargetLayerSettings = InTargetLayerSettings;
+		
+		if (FProperty* Property = StaticClass()->FindPropertyByName(GET_MEMBER_NAME_CHECKED(ALandscapeProxy, TargetLayers));  Property != nullptr && bPostEditChange)
+		{
+			FPropertyChangedEvent PropertyChangedEvent(Property);
+			PostEditChangeProperty(PropertyChangedEvent);
+		}
+
+		Modify();
+		return true;
+	}
+
+	return false;
+}
+
+bool ALandscapeProxy::HasTargetLayer(const FName& Name) const
+{
+	return TargetLayers.Find(Name) != nullptr; 
+}
+
+bool ALandscapeProxy::HasTargetLayer(const FLandscapeTargetLayerSettings& TargetLayerSettings) const
+{
+	return TargetLayers.FindKey(TargetLayerSettings) != nullptr;
+}
+
+bool ALandscapeProxy::HasTargetLayer(const ULandscapeLayerInfoObject* LayerInfoObject) const
+{
+	for (const auto& It : TargetLayers)
+	{
+		if (It.Value.LayerInfoObj == LayerInfoObject)
+		{
+			return true;
+		}
+	}
+	return false;
+}
+
+const TMap<FName, FLandscapeTargetLayerSettings>& ALandscapeProxy::GetTargetLayers() const
+{
+	return TargetLayers;
+}
+#endif
+
 bool ALandscapeProxy::AreNaniteComponentsValid(const FGuid& InProxyContentId) const
 {
 	if (NaniteComponents.IsEmpty())
diff --git a/Engine/Source/Runtime/Landscape/Private/LandscapeEdit.cpp b/Engine/Source/Runtime/Landscape/Private/LandscapeEdit.cpp
--- a/Engine/Source/Runtime/Landscape/Private/LandscapeEdit.cpp
+++ b/Engine/Source/Runtime/Landscape/Private/LandscapeEdit.cpp
@@ -90,6 +90,7 @@ LandscapeEdit.cpp: Landscape editing
 #include "Misc/TransactionObjectEvent.h"
 #endif
 #include "Algo/Compare.h"
+#include "LandscapeProxy.h"
 #include "Algo/Count.h"
 #include "Algo/Transform.h"
 #include "Algo/ForEach.h"
@@ -1117,6 +1118,7 @@ void ULandscapeComponent::FixupWeightmaps(const FGuid& InEditLayerGuid)
 	{
 		ULandscapeInfo* Info = GetLandscapeInfo();
 		ALandscapeProxy* Proxy = GetLandscapeProxy();
+		ALandscape* Landscape = Proxy->GetLandscapeActor();
 
 		TArray<TObjectPtr<UTexture2D>>& LocalWeightmapTextures = GetWeightmapTextures(InEditLayerGuid);
 		TArray<TObjectPtr<ULandscapeWeightmapUsage>>& LocalWeightmapTextureUsages = GetWeightmapTexturesUsage(InEditLayerGuid);
@@ -1145,14 +1147,16 @@ void ULandscapeComponent::FixupWeightmaps(const FGuid& InEditLayerGuid)
 			for (const auto& Allocation : LocalWeightmapLayerAllocations)
 			{
 				if (!Allocation.LayerInfo
-					|| (Allocation.LayerInfo != ALandscapeProxy::VisibilityLayer && Info->GetLayerInfoIndex(Allocation.LayerInfo) == INDEX_NONE))
+					|| (Allocation.LayerInfo != ALandscapeProxy::VisibilityLayer && Proxy && !Landscape->HasTargetLayer(Allocation.LayerInfo->LayerName)))
 				{
 					if (!bFixedLayerDeletion)
 					{
 						FFormatNamedArguments Arguments;
 						Arguments.Add(TEXT("LandscapeName"), FText::FromString(GetPathName()));
+						Arguments.Add(TEXT("TargetLayerName"), FText::FromString(Allocation.LayerInfo->LayerName.ToString()));
+						Arguments.Add(TEXT("EditLayerName"), FText::FromString(InEditLayerGuid.IsValid() ? LayersData[InEditLayerGuid].DebugName.ToString() : FString(TEXT("unknown"))));
 						FMessageLog("MapCheck").Warning()
-							->AddToken(FTextToken::Create(FText::Format(LOCTEXT("MapCheck_Message_FixedUpDeletedLayerWeightmap", "{LandscapeName} : Fixed up deleted layer weightmap"), Arguments)))
+							->AddToken(FTextToken::Create(FText::Format(LOCTEXT("MapCheck_Message_FixedUpDeletedLayerWeightmap", "{LandscapeName} : Fixed up deleted layer weightmap Edit Layer: '{EditLayerName}' Target Layer: '{TargetLayerName}'"), Arguments)))
 							->AddToken(FMapErrorToken::Create(FMapErrors::FixedUpDeletedLayerWeightmap));
 					}
 
@@ -1160,7 +1164,7 @@ void ULandscapeComponent::FixupWeightmaps(const FGuid& InEditLayerGuid)
 					LayersToDelete.Add(Allocation.LayerInfo);
 				}
 			}
-
+			
 			if (bFixedLayerDeletion)
 			{
 				// Delete material layer in the base/render layer
@@ -2974,6 +2978,50 @@ const TArray<FName>& ALandscapeProxy::GetLayersFromMaterial() const
 	return GetLayersFromMaterial(LandscapeMaterial);
 }
 
+TArray<FName> ALandscapeProxy::RetrieveAllLayerNamesFromMaterials() const
+{
+	TSet<FName> LayerNames;
+	LayerNames.Append(GetLayersFromMaterial());
+
+	for (int32 ComponentIndex = 0; ComponentIndex < LandscapeComponents.Num(); ComponentIndex++)
+	{
+		const ULandscapeComponent* Component = LandscapeComponents[ComponentIndex];
+
+		// Add layers from per-component override materials
+		if ((Component != nullptr) && (Component->OverrideMaterial != nullptr))
+		{
+			LayerNames.Append(GetLayersFromMaterial(Component->OverrideMaterial));
+
+			// Hole Override material
+			LayerNames.Append(GetLayersFromMaterial( Component->OverrideHoleMaterial));
+				
+			// LOD Materials
+			for(const FLandscapePerLODMaterialOverride& OverrideMaterial : Component->GetPerLODOverrideMaterials())
+			{
+				LayerNames.Append(GetLayersFromMaterial(OverrideMaterial.Material));	
+			}
+		}
+	}
+	return LayerNames.Array();
+}
+
+TMap<FName, ULandscapeLayerInfoObject*> ALandscapeProxy::RetrieveAllocationInfos() const
+{
+	TMap<FName, ULandscapeLayerInfoObject*> InfoObjects;
+	for (const TObjectPtr<ULandscapeComponent>& Component : LandscapeComponents)
+	{
+		for (const FWeightmapLayerAllocationInfo& Allocation : Component->GetWeightmapLayerAllocations())
+		{
+			if (Allocation.LayerInfo)
+			{
+				InfoObjects.Add(Allocation.GetLayerName(), Allocation.LayerInfo);
+			}
+		}
+	}
+	
+	return InfoObjects;
+}
+
 ULandscapeLayerInfoObject* ALandscapeProxy::CreateLayerInfo(const TCHAR* InLayerName, const ULevel* InLevel, const ULandscapeLayerInfoObject* InTemplate)
 {
 	FName LayerObjectName;
@@ -4848,31 +4896,25 @@ void ULandscapeInfo::DeleteLayer(ULandscapeLayerInfoObject* LayerInfo, const FNa
 {
 	GWarn->BeginSlowTask(LOCTEXT("BeginDeletingLayerTask", "Deleting Layer"), true);
 
-	// Remove data from all components
-	FLandscapeEditDataInterface LandscapeEdit(this);
-	LandscapeEdit.DeleteLayer(LayerInfo);
-
+	LandscapeActor->Modify();
 	// Remove from layer settings array
 	{
 		int32 LayerIndex = Layers.IndexOfByPredicate([LayerInfo, LayerName](const FLandscapeInfoLayerSettings& LayerSettings) { return LayerSettings.LayerInfoObj == LayerInfo && LayerSettings.LayerName == LayerName; });
 		if (LayerIndex != INDEX_NONE)
 		{
 			Layers.RemoveAt(LayerIndex);
+			LandscapeActor->RemoveTargetLayer(LayerName);
 		}
 	}
-
-	ForEachLandscapeProxy([LayerInfo](ALandscapeProxy* Proxy)
-	{
-		Proxy->Modify();
-		int32 Index = Proxy->EditorLayerSettings.IndexOfByKey(LayerInfo);
-		if (Index != INDEX_NONE)
-		{
-			Proxy->EditorLayerSettings.RemoveAt(Index);
-		}
-		return true;
-	});
-
-	//UpdateLayerInfoMap();
+	
+	LandscapeActor->RemoveTargetLayer(LayerName);
+	
+	
+	// Remove data from all components
+	FLandscapeEditDataInterface LandscapeEdit(this);
+	LandscapeEdit.DeleteLayer(LayerInfo);
+	
+	UpdateLayerInfoMap();
 
 	GWarn->EndSlowTask();
 }
@@ -4896,37 +4938,9 @@ void ULandscapeInfo::ReplaceLayer(ULandscapeLayerInfoObject* FromLayerInfo, ULan
 			}
 		}
 
-		ForEachLandscapeProxy([FromLayerInfo, ToLayerInfo](ALandscapeProxy* Proxy)
-		{
-			Proxy->Modify();
-			FLandscapeEditorLayerSettings* ToEditorLayerSettings = Proxy->EditorLayerSettings.FindByKey(ToLayerInfo);
-			if (ToEditorLayerSettings != nullptr)
-			{
-				// If the new layer already exists, simple remove the old layer
-				int32 Index = Proxy->EditorLayerSettings.IndexOfByKey(FromLayerInfo);
-				if (Index != INDEX_NONE)
-				{
-					Proxy->EditorLayerSettings.RemoveAt(Index);
-				}
-			}
-			else
-			{
-				FLandscapeEditorLayerSettings* FromEditorLayerSettings = Proxy->EditorLayerSettings.FindByKey(FromLayerInfo);
-				if (FromEditorLayerSettings != nullptr)
-				{
-					// If only the old layer exists (most common case), change it to point to the new layer info
-					FromEditorLayerSettings->LayerInfoObj = ToLayerInfo;
-				}
-				else
-				{
-					// If neither exists in the EditorLayerSettings cache, add it
-					Proxy->EditorLayerSettings.Add(FLandscapeEditorLayerSettings(ToLayerInfo));
-				}
-			}
-			return true;
-		});
-
-		//UpdateLayerInfoMap();
+		LandscapeActor->RemoveTargetLayer(FromLayerInfo->LayerName);
+		LandscapeActor->AddTargetLayer(ToLayerInfo->LayerName, FLandscapeTargetLayerSettings(ToLayerInfo));
+		UpdateLayerInfoMap();
 
 		GWarn->EndSlowTask();
 	}
@@ -5076,7 +5090,7 @@ void ALandscape::PostEditMove(bool bFinished)
 	if (bFinished)
 	{
 		// align all proxies to landscape actor
-		auto* LandscapeInfo = GetLandscapeInfo();
+		ULandscapeInfo* LandscapeInfo = GetLandscapeInfo();
 		if (LandscapeInfo)
 		{
 			LandscapeInfo->FixupProxiesTransform(true);
@@ -6157,9 +6171,6 @@ void ALandscapeProxy::PostEditChangeProperty(FPropertyChangedEvent& PropertyChan
 		MarkComponentsRenderStateDirty();
 	}
 
-	// Remove null layer infos
-	EditorLayerSettings.RemoveAll([](const FLandscapeEditorLayerSettings& Entry) { return Entry.LayerInfoObj == nullptr; });
-
 	// Remove any null landscape components
 	LandscapeComponents.RemoveAll([](const ULandscapeComponent* Component) { return Component == nullptr; });
 
diff --git a/Engine/Source/Runtime/Landscape/Private/LandscapeProxy.cpp b/Engine/Source/Runtime/Landscape/Private/LandscapeProxy.cpp
--- a/Engine/Source/Runtime/Landscape/Private/LandscapeProxy.cpp
+++ b/Engine/Source/Runtime/Landscape/Private/LandscapeProxy.cpp
@@ -20,7 +20,7 @@
 LANDSCAPE_API FLandscapeImportLayerInfo::FLandscapeImportLayerInfo(const FLandscapeInfoLayerSettings& InLayerSettings)
 	: LayerName(InLayerSettings.GetLayerName())
 	, LayerInfo(InLayerSettings.LayerInfoObj)
-	, SourceFilePath(InLayerSettings.GetEditorSettings().ReimportLayerFilePath)
+	, SourceFilePath(InLayerSettings.GetTargetLayerSettings().ReimportLayerFilePath)
 {
 }
 
diff --git a/Engine/Source/Runtime/Landscape/Private/LandscapeUtils.cpp b/Engine/Source/Runtime/Landscape/Private/LandscapeUtils.cpp
--- a/Engine/Source/Runtime/Landscape/Private/LandscapeUtils.cpp
+++ b/Engine/Source/Runtime/Landscape/Private/LandscapeUtils.cpp
@@ -12,6 +12,7 @@
 #include "Engine/Texture2D.h"
 #include "LandscapeComponent.h"
 #include "ObjectTools.h"
+#include "AssetRegistry/AssetRegistryModule.h"
 #endif
 
 // Channel remapping
@@ -283,6 +284,30 @@ bool UseWeightmapTextureArray(EShaderPlatform InPlatform)
 	return IsMobilePlatform(InPlatform) && (LandscapeMobileWeightTextureArray != 0);	
 }
 
+FLayerInfoFinder::FLayerInfoFinder()
+{
+	const UClass* AssetClass = ULandscapeLayerInfoObject::StaticClass();
+	FAssetRegistryModule& AssetRegistryModule = FModuleManager::LoadModuleChecked<FAssetRegistryModule>("AssetRegistry");
+	FARFilter Filter;
+	FName PackageName = *AssetClass->GetPackage()->GetName();
+	FName AssetName = AssetClass->GetFName();
+											
+	Filter.ClassPaths.Add(FTopLevelAssetPath(PackageName,AssetName));
+	AssetRegistryModule.Get().GetAssets(Filter, LayerInfoAssets);
+}
+
+ULandscapeLayerInfoObject* FLayerInfoFinder::Find(const FName& LayerName) const
+{
+	for (const FAssetData& LayerInfoAsset : LayerInfoAssets)
+	{
+		ULandscapeLayerInfoObject* LayerInfo = CastChecked<ULandscapeLayerInfoObject>(LayerInfoAsset.GetAsset());
+		if (LayerInfo && LayerInfo->LayerName == LayerName)
+		{
+			return LayerInfo;
+		}
+	}
+	return nullptr;
+}
 #endif //!WITH_EDITOR
 
 } // end namespace UE::Landscape
\ No newline at end of file
diff --git a/Engine/Source/Runtime/Landscape/Public/LandscapeUtils.h b/Engine/Source/Runtime/Landscape/Public/LandscapeUtils.h
--- a/Engine/Source/Runtime/Landscape/Public/LandscapeUtils.h
+++ b/Engine/Source/Runtime/Landscape/Public/LandscapeUtils.h
@@ -8,6 +8,7 @@
 #include "RHIDefinitions.h"
 #endif
 #include "UObject/NameTypes.h"
+#include "AssetRegistry/AssetData.h"
 
 class ULevel;
 enum EShaderPlatform : uint16;
@@ -111,7 +112,15 @@ bool UseWeightmapTextureArray(EShaderPlatform InPlatform);
 
 /** Check if Landscape.MobileWeightTextureArray CVar set and we should attempt to use texture arrays for weight maps on mobile platforms  */
 bool IsMobileWeightmapTextureArrayEnabled();
-	
+
+struct FLayerInfoFinder
+{
+	LANDSCAPE_API FLayerInfoFinder();
+	LANDSCAPE_API ~FLayerInfoFinder() = default;
+	LANDSCAPE_API ULandscapeLayerInfoObject* Find(const FName& LayerName) const;
+	TArray<FAssetData> LayerInfoAssets;
+};
+
 #endif //!WITH_EDITOR
 
 } // end namespace UE::Landscape
