Change null check to Landscape to match the object a method is potentially called on.
Landscape can be null during cook, when proxies may be cooked without the landscape.

diff --git a/Engine/Source/Runtime/Landscape/Private/LandscapeEdit.cpp b/Engine/Source/Runtime/Landscape/Private/LandscapeEdit.cpp
--- a/Engine/Source/Runtime/Landscape/Private/LandscapeEdit.cpp
+++ b/Engine/Source/Runtime/Landscape/Private/LandscapeEdit.cpp
@@ -1147,7 +1147,7 @@ void ULandscapeComponent::FixupWeightmaps(const FGuid& InEditLayerGuid)
 			for (const auto& Allocation : LocalWeightmapLayerAllocations)
 			{
 				if (!Allocation.LayerInfo
-					|| (Allocation.LayerInfo != ALandscapeProxy::VisibilityLayer && Proxy && !Landscape->HasTargetLayer(Allocation.LayerInfo->LayerName)))
+					|| (Allocation.LayerInfo != ALandscapeProxy::VisibilityLayer && Landscape && !Landscape->HasTargetLayer(Allocation.LayerInfo->LayerName)))
 				{
 					if (!bFixedLayerDeletion)
 					{
