Landscape: Disable deletion of invalid weightmaps

diff --git a/Engine/Source/Runtime/Landscape/Private/LandscapeEdit.cpp b/Engine/Source/Runtime/Landscape/Private/LandscapeEdit.cpp
--- a/Engine/Source/Runtime/Landscape/Private/LandscapeEdit.cpp
+++ b/Engine/Source/Runtime/Landscape/Private/LandscapeEdit.cpp
@@ -1149,19 +1149,21 @@ void ULandscapeComponent::FixupWeightmaps(const FGuid& InEditLayerGuid)
 				if (!Allocation.LayerInfo
 					|| (Allocation.LayerInfo != ALandscapeProxy::VisibilityLayer && Landscape && !Landscape->HasTargetLayer(Allocation.LayerInfo->LayerName)))
 				{
-					if (!bFixedLayerDeletion)
-					{
-						FFormatNamedArguments Arguments;
-						Arguments.Add(TEXT("LandscapeName"), FText::FromString(GetPathName()));
-						Arguments.Add(TEXT("TargetLayerName"), FText::FromString(Allocation.LayerInfo->LayerName.ToString()));
-						Arguments.Add(TEXT("EditLayerName"), FText::FromString(InEditLayerGuid.IsValid() ? LayersData[InEditLayerGuid].DebugName.ToString() : FString(TEXT("unknown"))));
-						FMessageLog("MapCheck").Warning()
-							->AddToken(FTextToken::Create(FText::Format(LOCTEXT("MapCheck_Message_FixedUpDeletedLayerWeightmap", "{LandscapeName} : Fixed up deleted layer weightmap Edit Layer: '{EditLayerName}' Target Layer: '{TargetLayerName}'"), Arguments)))
-							->AddToken(FMapErrorToken::Create(FMapErrors::FixedUpDeletedLayerWeightmap));
-					}
-
-					bFixedLayerDeletion = true;
-					LayersToDelete.Add(Allocation.LayerInfo);
+					//don.boogert: disable deleting layers temp
+					// if (!bFixedLayerDeletion)
+					// {
+					// 	FFormatNamedArguments Arguments;
+					// 	Arguments.Add(TEXT("LandscapeName"), FText::FromString(GetPathName()));
+					// 	Arguments.Add(TEXT("TargetLayerName"), FText::FromString(Allocation.LayerInfo->LayerName.ToString()));
+					// 	Arguments.Add(TEXT("EditLayerName"), FText::FromString(InEditLayerGuid.IsValid() ? LayersData[InEditLayerGuid].DebugName.ToString() : FString(TEXT("unknown"))));
+					// 	FMessageLog("MapCheck").Warning()
+					// 		->AddToken(FTextToken::Create(FText::Format(LOCTEXT("MapCheck_Message_FixedUpDeletedLayerWeightmap", "{LandscapeName} : Fixed up deleted layer weightmap Edit Layer: '{EditLayerName}' Target Layer: '{TargetLayerName}'"), Arguments)))
+					// 		->AddToken(FMapErrorToken::Create(FMapErrors::FixedUpDeletedLayerWeightmap));
+					// }
+					//
+					//
+					// bFixedLayerDeletion = true;
+					// LayersToDelete.Add(Allocation.LayerInfo);
 				}
 			}
 			
