Fix for excessive material instance invalidation in Landscape PostLoad

Initial Editor Load:
	ALandscapeProxy::PostLoad		14.3s --> 385 ms
PIE:
	ALandscapeProxy::PostLoad		16.5s --> 31.6 ms

diff --git a/Engine/Source/Runtime/Landscape/Private/Landscape.cpp b/Engine/Source/Runtime/Landscape/Private/Landscape.cpp
--- a/Engine/Source/Runtime/Landscape/Private/Landscape.cpp
+++ b/Engine/Source/Runtime/Landscape/Private/Landscape.cpp
@@ -528,8 +528,6 @@ void ALandscapeProxy::InvalidateNaniteRepresentation(bool bInCheckContentId)
 
 void ALandscapeProxy::InvalidateOrUpdateNaniteRepresentation(bool bInCheckContentId, const ITargetPlatform* InTargetPlatform)
 {
-	TRACE_BOOKMARK(TEXT("ALandscapeProxy::InvalidateOrUpdateNaniteRepresentation"));	
-
 	ULandscapeSubsystem* Subsystem = GetWorld()->GetSubsystem<ULandscapeSubsystem>();
 	if (Subsystem->IsLiveNaniteRebuildEnabled())
 	{
@@ -1192,6 +1190,8 @@ void ULandscapeComponent::UpdatedSharedPropertiesFromActor()
 
 void ULandscapeComponent::PostLoad()
 {
+	TRACE_CPUPROFILER_EVENT_SCOPE(ULandscapeComponent::PostLoad);
+
 	using namespace UE::Landscape;
 
 	Super::PostLoad();
@@ -2837,6 +2837,8 @@ void ALandscape::PostInitProperties()
 
 void ALandscape::PostLoad()
 {
+	TRACE_CPUPROFILER_EVENT_SCOPE(ALandscape::PostLoad);
+
 	if (!LandscapeGuid.IsValid())
 	{
 		LandscapeGuid = FGuid::NewGuid();
@@ -4053,6 +4055,8 @@ void ALandscapeProxy::PostLoadFixupLandscapeGuidsIfInstanced()
 
 void ALandscapeProxy::PostLoad()
 {
+	TRACE_CPUPROFILER_EVENT_SCOPE(ALandscapeProxy::PostLoad);
+
 	Super::PostLoad();
 
 	PostLoadFixupLandscapeGuidsIfInstanced();
@@ -4177,7 +4181,6 @@ PRAGMA_ENABLE_DEPRECATION_WARNINGS
 
 		bool bFixedUpInvalidMaterialInstances = false;
 
-	UMaterial* BaseLandscapeMaterial = GetLandscapeMaterial()->GetMaterial();
 
 	for (ULandscapeComponent* Comp : LandscapeComponents)
 	{
@@ -4198,6 +4201,7 @@ PRAGMA_ENABLE_DEPRECATION_WARNINGS
 		// Only validate if uncooked and in the editor/commandlet mode (we cannot re-build material instance constants if this is not the case : see UMaterialInstance::CacheResourceShadersForRendering, which is only called if FApp::CanEverRender() returns true) 
 		if (!Comp->GetOutermost()->HasAnyPackageFlags(PKG_FilterEditorOnly) && (GIsEditor && FApp::CanEverRender()))
 		{
+			UMaterial* BaseLandscapeMaterial = Comp->GetLandscapeMaterial()->GetMaterial();
 			// MaterialInstance is different from the used LandscapeMaterial, we need to update the material as we cannot properly validate used combinations.
 			if (MaterialInstance->GetMaterial() != BaseLandscapeMaterial)
 			{
diff --git a/Engine/Source/Runtime/Landscape/Private/LandscapeEdit.cpp b/Engine/Source/Runtime/Landscape/Private/LandscapeEdit.cpp
--- a/Engine/Source/Runtime/Landscape/Private/LandscapeEdit.cpp
+++ b/Engine/Source/Runtime/Landscape/Private/LandscapeEdit.cpp
@@ -600,6 +600,8 @@ UMaterialInstanceConstant* ULandscapeComponent::GetCombinationMaterial(FMaterial
 
 void ULandscapeComponent::UpdateMaterialInstances_Internal(FMaterialUpdateContext& Context)
 {
+	TRACE_CPUPROFILER_EVENT_SCOPE(ULandscapeComponent::UpdateMaterialInstances_Internal);
+
 	const int8 MaxLOD = IntCastChecked<int8>(FMath::CeilLogTwo(SubsectionSizeQuads + 1) - 1);
 	decltype(MaterialPerLOD) NewMaterialPerLOD;
 	LODIndexToMaterialIndex.SetNumUninitialized(MaxLOD+1);
diff --git a/Engine/Source/Runtime/Landscape/Private/LandscapeSubsystem.cpp b/Engine/Source/Runtime/Landscape/Private/LandscapeSubsystem.cpp
--- a/Engine/Source/Runtime/Landscape/Private/LandscapeSubsystem.cpp
+++ b/Engine/Source/Runtime/Landscape/Private/LandscapeSubsystem.cpp
@@ -689,7 +689,7 @@ TArray<TTuple<ALandscapeProxy*, UE::Landscape::EOutdatedDataFlags>> ULandscapeSu
 
 void ULandscapeSubsystem::BuildNanite(TArrayView<ALandscapeProxy*> InProxiesToBuild, bool bForceRebuild)
 {
-	TRACE_BOOKMARK(TEXT("ULandscapeSubsystem::BuildNanite"));
+	TRACE_CPUPROFILER_EVENT_SCOPE(ULandscapeSubsystem::BuildNanite);
 
 	UWorld* World = GetWorld();
 	if (!World || World->IsGameWorld())
