Fixed landscape material property assignments not being properly undone when using edit layers : in that case, even if the landscape layers update runs properly, if no change was detected on weightmaps, the component material instances wouldn't get updated and the material changes wouldn't get applied

diff --git a/Engine/Source/Runtime/Landscape/Private/LandscapeEdit.cpp b/Engine/Source/Runtime/Landscape/Private/LandscapeEdit.cpp
--- a/Engine/Source/Runtime/Landscape/Private/LandscapeEdit.cpp
+++ b/Engine/Source/Runtime/Landscape/Private/LandscapeEdit.cpp
@@ -976,13 +976,13 @@ void ULandscapeComponent::PostEditUndo()
 			ComponentProxy->FixupWeightmaps();
 		}
 
-		// Update Material Instances on each modified component (uses the weightmap allocations fixed by FixupWeightmaps)
+		// Update Material Instances on each modified component, in case the material property has changed (uses the weightmap allocations fixed by FixupWeightmaps)
 		for (ULandscapeComponent* Component : UndoRedoModifiedComponents)
 		{
 			if (IsValid(Component)) // Components can be pending kill, if they were removed by the undo operation
 			{
 				ALandscapeProxy* ComponentProxy = Component->GetLandscapeProxy();
-				if (IsValid(ComponentProxy) && !ComponentProxy->HasLayersContent())
+				if (IsValid(ComponentProxy))
 				{
 					Component->UpdateMaterialInstances();
 				}
