The splines edit layer and the standard edit layers now inherit from the same base class, ULandscapeEditLayerPersistent, so that they can share their rendering / blending technique (preparatory work for upcoming change)

diff --git a/Engine/Source/Runtime/Landscape/Classes/LandscapeEditLayer.h b/Engine/Source/Runtime/Landscape/Classes/LandscapeEditLayer.h
--- a/Engine/Source/Runtime/Landscape/Classes/LandscapeEditLayer.h
+++ b/Engine/Source/Runtime/Landscape/Classes/LandscapeEditLayer.h
@@ -8,6 +8,10 @@ enum class ELandscapeToolTargetType : uint8;
 struct FLandscapeLayer;
 class ALandscape;
 
+/** 
+* Base class for all landscape edit layers. By implementing the various virtual functions, we are able to customize the behavior of the edit layer
+*  wrt the landscape tools in a generic way (e.g. does it support sculpting tools? painting tools? can it be collapsed?, etc.)
+*/
 UCLASS(MinimalAPI, Abstract)
 class ULandscapeEditLayerBase : public UObject
 {
@@ -153,10 +157,25 @@ protected:
 };
 
 /** 
-	This is the standard type of edit layer. It can be manually authored (sculpted, painted, etc.) in the landscape editor 
+* Base class for persistent layers, i.e. layers that have a set of backing textures (heightmaps, weightmaps) and can therefore be rendered in a similar fashion
+*/
+UCLASS(MinimalAPI, Abstract)
+class ULandscapeEditLayerPersistent : public ULandscapeEditLayerBase
+{
+	GENERATED_BODY()
+
+public:
+	// Begin ULandscapeEditLayerBase implementation
+	virtual bool NeedsPersistentTextures() const override { return true; };
+	virtual bool SupportsCollapsingTo() const override { return true; } // If the layer has persistent textures, it can be collapsed to another layer (one that supports being collapsed away, that is)
+	// End ULandscapeEditLayerBase implementation
+};
+
+/** 
+* This is the standard type of edit layer. It can be manually authored (sculpted, painted, etc.) in the landscape editor 
 */
 UCLASS(MinimalAPI)
-class ULandscapeEditLayer : public ULandscapeEditLayerBase
+class ULandscapeEditLayer : public ULandscapeEditLayerPersistent
 {
 	GENERATED_BODY()
 
@@ -164,10 +183,8 @@ public:
 	// Begin ULandscapeEditLayerBase implementation
 	virtual bool SupportsTargetType(ELandscapeToolTargetType InType) const override;
 	virtual bool SupportsEditingTools() const override { return true; }
-	virtual bool NeedsPersistentTextures() const override { return true; };
 	virtual bool SupportsMultiple() const override { return true; }
 	virtual bool SupportsBeingCollapsedAway() const override { return true; }
-	virtual bool SupportsCollapsingTo() const override { return true; }
 	virtual FString GetDefaultName() const { return TEXT("Layer"); }
 	// End ULandscapeEditLayerBase implementation
 
@@ -176,7 +193,7 @@ protected:
 };
 
 /** 
-	Base class for procedural layers. Procedural layers cannot be edited through standard editing tools
+* Base class for procedural layers. Procedural layers cannot be edited through standard editing tools
 */
 UCLASS(MinimalAPI, Abstract)
 class ULandscapeEditLayerProcedural : public ULandscapeEditLayerBase
@@ -186,24 +203,26 @@ class ULandscapeEditLayerProcedural : public ULandscapeEditLayerBase
 public:
 	// Begin ULandscapeEditLayerBase implementation
 	virtual bool SupportsEditingTools() const override { return false; } // procedural layers cannot be edited through standard editing tools
+	virtual bool SupportsCollapsingTo() const override { return false; } // for now, don't support collapsing to a layer underneath for a procedural layer (this may become unneeded if we make the collapse happen on the GPU)
 	virtual bool SupportsBeingCollapsedAway() const override { return false; } // this is a procedural and therefore cannot be collapsed 
-	virtual bool SupportsCollapsingTo() const override { return NeedsPersistentTextures(); } // if the layer has persistent textures, it supports being collapsed into a layer underneath (this may become unneeded if we make the collapse happen on the GPU)
 	// End ULandscapeEditLayerBase implementation
 };
 
 /** 
-	Procedural edit layer that lets the user manipulate its content using landscape splines (Splines tool in the Manage panel) 
+* Procedural edit layer that lets the user manipulate its content using landscape splines (Splines tool in the Manage panel) 
 */
 UCLASS(MinimalAPI)
-class ULandscapeEditLayerSplines : public ULandscapeEditLayerProcedural
+class ULandscapeEditLayerSplines : public ULandscapeEditLayerPersistent
 {
 	GENERATED_BODY()
 
 public:
 	// Begin ULandscapeEditLayerBase implementation
+	virtual bool SupportsEditingTools() const override { return false; } // procedural layers cannot be edited through standard editing tools
 	virtual bool SupportsTargetType(ELandscapeToolTargetType InType) const override;
 	virtual bool NeedsPersistentTextures() const override { return true; }; // it's a layer computed on the CPU and outputting to persistent textures
 	virtual bool SupportsMultiple() const override { return false; } // only one layer of this type is allowed
+	virtual bool SupportsBeingCollapsedAway() const override { return false; } // this is a procedural and therefore cannot be collapsed 
 	virtual FString GetDefaultName() const override { return TEXT("Splines"); }
 	virtual void OnLayerCreated(FLandscapeLayer& Layer) override;
 	virtual TArray<FEditLayerAction> GetActions() const override;
