Landscape: Fix layer deletion on Proxy load
After the change to target layer definitions the FixupWeightLayers was happening too early and deleting weight layers which were still in use.

Now this is called only after the LandscapeActor is registered.

diff --git a/Engine/Source/Runtime/Landscape/Classes/LandscapeInfo.h b/Engine/Source/Runtime/Landscape/Classes/LandscapeInfo.h
--- a/Engine/Source/Runtime/Landscape/Classes/LandscapeInfo.h
+++ b/Engine/Source/Runtime/Landscape/Classes/LandscapeInfo.h
@@ -398,9 +398,9 @@ public:
 	void UpdateNanite(const ITargetPlatform* InTargetPlatform);
 
 	/** Associates passed actor with this info object
- *  @param	Proxy		Landscape actor to register
- *  @param  bMapCheck	Whether to warn about landscape errors
- */
+	 *  @param	Proxy		Landscape actor to register
+	 *  @param  bMapCheck	Whether to warn about landscape errors
+	 */
 	LANDSCAPE_API void RegisterActor(ALandscapeProxy* Proxy, bool bMapCheck = false, bool bUpdateAllAddCollisions = true);
 
 	/** Deassociates passed actor with this info object*/
@@ -461,4 +461,7 @@ private:
 	void UpdateLayerInfoMapInternal(ALandscapeProxy* Proxy);
 	bool TryAddToModifiedPackages(UPackage* InPackage, const ALandscape* InLandscapeOverride = nullptr);
 #endif
+
+private:
+	void RegisterLandscapeActorWithProxyInternal(ALandscapeProxy* Proxy, bool bMapCheck);
 };
diff --git a/Engine/Source/Runtime/Landscape/Private/Landscape.cpp b/Engine/Source/Runtime/Landscape/Private/Landscape.cpp
--- a/Engine/Source/Runtime/Landscape/Private/Landscape.cpp
+++ b/Engine/Source/Runtime/Landscape/Private/Landscape.cpp
@@ -2726,37 +2726,6 @@ void ALandscapeProxy::PostRegisterAllComponents()
 			{
 				LandscapeInfo = CreateLandscapeInfo(true);
 			}
-
-#if WITH_EDITOR
-			if (GIsEditor)
-			{
-				// Note: This can happen when loading certain cooked assets in an editor
-				// Todo: Determine the root cause of this and fix it at a higher level!
-				if (LandscapeComponents.Num() > 0 && LandscapeComponents[0] == nullptr)
-				{
-					LandscapeComponents.Empty();
-				}
-
-				if (WeightmapFixupVersion != CurrentVersion)
-				{
-					FixupWeightmaps();
-				}
-
-				UpdateCachedHasLayersContent(true);
-
-				// Cache the value at this point as CreateLandscapeInfo (-> RegisterActor) might create/destroy layers content if there was a mismatch between landscape & proxy
-				// Check the actual flag here not HasLayersContent() which could return true if the LandscapeActorRef is valid.
-				bool bHasLayersContentBefore = bHasLayersContent;
-
-				check(WeightmapFixupVersion == CurrentVersion);
-
-				const bool bNeedOldDataMigration = !bHasLayersContentBefore && CanHaveLayersContent();
-				if (bNeedOldDataMigration && LandscapeInfo->LandscapeActor.IsValid() && LandscapeInfo->LandscapeActor->HasLayersContent())
-				{
-					LandscapeInfo->LandscapeActor->CopyOldDataToDefaultLayer(this);
-				}
-			}
-#endif // WITH_EDITOR
 		}
 
 		if (UWorld* OwningWorld = GetWorld())
@@ -3938,6 +3907,11 @@ void ULandscapeInfo::UpdateLayerInfoMapInternal(ALandscapeProxy* Proxy)
 {
 	TRACE_CPUPROFILER_EVENT_SCOPE(ULandscapeInfo::UpdateLayerInfoMapInternal);
 
+	if( !LandscapeActor.IsValid())
+	{
+		return;
+	}
+
 	if (GIsEditor && Proxy && Proxy->bInitWeightLayersFromMaterial)
 	{
 		Proxy->bInitWeightLayersFromMaterial = false;
@@ -3969,30 +3943,26 @@ void ULandscapeInfo::UpdateLayerInfoMapInternal(ALandscapeProxy* Proxy)
 		}
 	}
 
-	if (LandscapeActor.IsValid())
+	Layers.Empty();
+
+	for (const TTuple<FName, FLandscapeTargetLayerSettings>& TargetLayer : LandscapeActor->GetTargetLayers())
 	{
-		Layers.Empty();
+		FLandscapeInfoLayerSettings InfoLayerSettings (TargetLayer.Key, LandscapeActor.Get());
+		InfoLayerSettings.bValid = true;
+		InfoLayerSettings.LayerInfoObj = TargetLayer.Value.LayerInfoObj;
+		Layers.Add(InfoLayerSettings);
+	}
+	
+	// Add Visibility Layer info if not initialized
+	if (ALandscapeProxy::VisibilityLayer != nullptr)
+	{
+		int32 LayerInfoIndex = GetLayerInfoIndex(ALandscapeProxy::VisibilityLayer->LayerName);
 
-		for (const TTuple<FName, FLandscapeTargetLayerSettings>& TargetLayer : LandscapeActor->GetTargetLayers())
-		{
-			FLandscapeInfoLayerSettings InfoLayerSettings (TargetLayer.Key, LandscapeActor.Get());
-			InfoLayerSettings.bValid = true;
-			InfoLayerSettings.LayerInfoObj = TargetLayer.Value.LayerInfoObj;
-			Layers.Add(InfoLayerSettings);
-		}
-		
-		// Add Visibility Layer info if not initialized
-		if (ALandscapeProxy::VisibilityLayer != nullptr)
+		if ((LayerInfoIndex != INDEX_NONE) && (Layers[LayerInfoIndex].LayerInfoObj == nullptr))
 		{
-			int32 LayerInfoIndex = GetLayerInfoIndex(ALandscapeProxy::VisibilityLayer->LayerName);
-
-			if ((LayerInfoIndex != INDEX_NONE) && (Layers[LayerInfoIndex].LayerInfoObj == nullptr))
-			{
-				Layers[LayerInfoIndex].LayerInfoObj = ALandscapeProxy::VisibilityLayer;
-			}
+			Layers[LayerInfoIndex].LayerInfoObj = ALandscapeProxy::VisibilityLayer;
 		}
 	}
-	
 }
 
 bool ULandscapeInfo::UpdateLayerInfoMap(ALandscapeProxy* Proxy /*= nullptr*/, bool bInvalidate /*= false*/)
@@ -5448,6 +5418,63 @@ bool ULandscapeInfo::IsRegistered(const ALandscapeProxy* Proxy) const
 	return bResult;
 }
 
+
+// this function contains all of the registration code that requires the ALandscape actor to be present
+void ULandscapeInfo::RegisterLandscapeActorWithProxyInternal(ALandscapeProxy* Proxy, bool bMapCheck)
+{
+	ALandscape* Landscape = LandscapeActor.Get();
+	check(Landscape);
+
+	if (ALandscapeStreamingProxy* StreamingProxy = Cast<ALandscapeStreamingProxy>(Proxy))
+	{
+		// streaming proxy specific setup here
+		StreamingProxy->SetLandscapeActor(Landscape);
+
+#if WITH_EDITOR
+		StreamingProxy->FixupSharedData(Landscape, bMapCheck);
+#endif // WITH_EDITOR
+
+		StreamingProxy->SetLODGroupKeyInternal(Landscape->LODGroupKey);
+	}
+
+#if WITH_EDITOR
+	// generic proxy setup (that requires ALandscape actor) here
+	UpdateLayerInfoMap(Proxy);
+
+	if (GIsEditor)
+	{
+		// Note: This can happen when loading certain cooked assets in an editor
+		// Todo: Determine the root cause of this and fix it at a higher level!
+		if (Proxy->LandscapeComponents.Num() > 0 && Proxy->LandscapeComponents[0] == nullptr)
+		{
+			Proxy->LandscapeComponents.Empty();
+		}
+
+		if (Proxy->WeightmapFixupVersion != Proxy->CurrentVersion)
+		{
+			Proxy->FixupWeightmaps();
+		}
+
+		Proxy->UpdateCachedHasLayersContent(true);
+
+		// Cache the value at this point as CreateLandscapeInfo (-> RegisterActor) might create/destroy layers content if there was a mismatch between landscape & proxy
+		// Check the actual flag here not HasLayersContent() which could return true if the LandscapeActorRef is valid.
+		bool bHasLayersContentBefore = Proxy->bHasLayersContent;
+
+		check(Proxy->WeightmapFixupVersion == Proxy->CurrentVersion);
+
+		const bool bNeedOldDataMigration = !bHasLayersContentBefore && CanHaveLayersContent();
+		if (bNeedOldDataMigration && LandscapeActor->HasLayersContent())
+		{
+			LandscapeActor->CopyOldDataToDefaultLayer(Proxy);
+		}
+	}
+#endif // WITH_EDITOR
+
+
+
+}
+
 void ULandscapeInfo::RegisterActor(ALandscapeProxy* Proxy, bool bMapCheck, bool bUpdateAllAddCollisions)
 {
 	UWorld* OwningWorld = Proxy->GetWorld();
@@ -5504,22 +5531,23 @@ void ULandscapeInfo::RegisterActor(ALandscapeProxy* Proxy, bool bMapCheck, bool
 			LandscapeActor->SetLockLocation(bIsLockLocation);
 #endif // WITH_EDITOR
 
-			// update proxies reference actor
+#if WITH_EDITORONLY_DATA
+			Landscape->bIsRegisteredWithLandscapeInfo = true;
+#endif // WITH_EDITORONLY_DATA
+
+			// run post-landscape actor registration on the LandscapeActor first, then on each streaming proxy
+			RegisterLandscapeActorWithProxyInternal(Landscape, bMapCheck);
 			for (TWeakObjectPtr<ALandscapeStreamingProxy> StreamingProxyPtr : StreamingProxies)
 			{
 				if (ALandscapeStreamingProxy* StreamingProxy = StreamingProxyPtr.Get())
 				{
-					StreamingProxy->SetLandscapeActor(LandscapeActor.Get());
-#if WITH_EDITOR
-					StreamingProxy->FixupSharedData(Landscape, bMapCheck);
-#endif // WITH_EDITOR
-					StreamingProxy->SetLODGroupKeyInternal(Landscape->LODGroupKey);
+					RegisterLandscapeActorWithProxyInternal(StreamingProxy, bMapCheck);
 				}
 			}
 		}
 		else if (LandscapeActor != Landscape)
 		{
-			UE_LOG(LogLandscape, Warning, TEXT("Multiple landscape actors with the same GUID detected: %s vs %s"), * LandscapeActor->GetPathName(), * Landscape->GetPathName());
+			UE_LOG(LogLandscape, Warning, TEXT("Multiple landscape actors with the same GUID detected: %s vs %s"), *LandscapeActor->GetPathName(), *Landscape->GetPathName());
 		}
 	}
 	else
@@ -5568,18 +5596,18 @@ void ULandscapeInfo::RegisterActor(ALandscapeProxy* Proxy, bool bMapCheck, bool
 			StreamingProxies.Insert(StreamingProxyPtr, InsertIndex);
 		}
 
-		if (LandscapeActor.IsValid())	// don't overwrite the proxy's landscape actor if we don't have one registered yet
+#if WITH_EDITORONLY_DATA
+		StreamingProxy->bIsRegisteredWithLandscapeInfo = true;
+#endif // WITH_EDITORONLY_DATA
+
+		// If we have a LandscapeActor, register it with the streaming proxy.  If not, it is deferred until a LandscapeActor is registered.
+		if (LandscapeActor.IsValid())
 		{
-			StreamingProxy->SetLandscapeActor(LandscapeActor.Get());
-			StreamingProxy->SetLODGroupKeyInternal(LandscapeActor.Get()->LODGroupKey);
-#if WITH_EDITOR
-			StreamingProxy->FixupSharedData(LandscapeActor.Get(), bMapCheck);
-#endif // WITH_EDITOR
+			RegisterLandscapeActorWithProxyInternal(StreamingProxy, bMapCheck);
 		}
 	}
 
 #if WITH_EDITOR
-	UpdateLayerInfoMap(Proxy);
 	if(bUpdateAllAddCollisions)
 	{
 		UpdateAllAddCollisions();
@@ -5599,10 +5627,6 @@ void ULandscapeInfo::RegisterActor(ALandscapeProxy* Proxy, bool bMapCheck, bool
 	{
 		RegisterCollisionComponent(CollComp);
 	}
-
-#if WITH_EDITORONLY_DATA
-	Proxy->bIsRegisteredWithLandscapeInfo = true;
-#endif // WITH_EDITORONLY_DATA
 }
 
 void ULandscapeInfo::UnregisterActor(ALandscapeProxy* Proxy)
diff --git a/Engine/Source/Runtime/Landscape/Private/LandscapeEdit.cpp b/Engine/Source/Runtime/Landscape/Private/LandscapeEdit.cpp
--- a/Engine/Source/Runtime/Landscape/Private/LandscapeEdit.cpp
+++ b/Engine/Source/Runtime/Landscape/Private/LandscapeEdit.cpp
@@ -1151,21 +1151,19 @@ void ULandscapeComponent::FixupWeightmaps(const FGuid& InEditLayerGuid)
 				if (!Allocation.LayerInfo
 					|| (Allocation.LayerInfo != ALandscapeProxy::VisibilityLayer && Landscape && !Landscape->HasTargetLayer(Allocation.LayerInfo->LayerName)))
 				{
-					//don.boogert: disable deleting layers temp
-					// if (!bFixedLayerDeletion)
-					// {
-					// 	FFormatNamedArguments Arguments;
-					// 	Arguments.Add(TEXT("LandscapeName"), FText::FromString(GetPathName()));
-					// 	Arguments.Add(TEXT("TargetLayerName"), FText::FromString(Allocation.LayerInfo->LayerName.ToString()));
-					// 	Arguments.Add(TEXT("EditLayerName"), FText::FromString(InEditLayerGuid.IsValid() ? LayersData[InEditLayerGuid].DebugName.ToString() : FString(TEXT("unknown"))));
-					// 	FMessageLog("MapCheck").Warning()
-					// 		->AddToken(FTextToken::Create(FText::Format(LOCTEXT("MapCheck_Message_FixedUpDeletedLayerWeightmap", "{LandscapeName} : Fixed up deleted layer weightmap Edit Layer: '{EditLayerName}' Target Layer: '{TargetLayerName}'"), Arguments)))
-					// 		->AddToken(FMapErrorToken::Create(FMapErrors::FixedUpDeletedLayerWeightmap));
-					// }
-					//
-					//
-					// bFixedLayerDeletion = true;
-					// LayersToDelete.Add(Allocation.LayerInfo);
+					if (!bFixedLayerDeletion)
+					{
+						FFormatNamedArguments Arguments;
+						Arguments.Add(TEXT("LandscapeName"), FText::FromString(GetPathName()));
+						Arguments.Add(TEXT("TargetLayerName"), FText::FromString(Allocation.LayerInfo->LayerName.ToString()));
+						Arguments.Add(TEXT("EditLayerName"), FText::FromString(InEditLayerGuid.IsValid() ? LayersData[InEditLayerGuid].DebugName.ToString() : FString(TEXT("unknown"))));
+						FMessageLog("MapCheck").Warning()
+							->AddToken(FTextToken::Create(FText::Format(LOCTEXT("MapCheck_Message_FixedUpDeletedLayerWeightmap", "{LandscapeName} : Fixed up deleted layer weightmap Edit Layer: '{EditLayerName}' Target Layer: '{TargetLayerName}'"), Arguments)))
+							->AddToken(FMapErrorToken::Create(FMapErrors::FixedUpDeletedLayerWeightmap));
+					}
+
+					bFixedLayerDeletion = true;
+					LayersToDelete.Add(Allocation.LayerInfo);
 				}
 			}
 			
