Fix for rare crash on landscape save
This is a fix for a possible assert when a landscape component that was unregistered ends up deleting itself within the presave update tick.
Because the amortization component index was not updated correctly to reflect the deleted element, some components may not be updated.

diff --git a/Engine/Source/Runtime/Landscape/Private/LandscapeGrassMapsBuilder.cpp b/Engine/Source/Runtime/Landscape/Private/LandscapeGrassMapsBuilder.cpp
--- a/Engine/Source/Runtime/Landscape/Private/LandscapeGrassMapsBuilder.cpp
+++ b/Engine/Source/Runtime/Landscape/Private/LandscapeGrassMapsBuilder.cpp
@@ -548,6 +548,7 @@ bool FLandscapeGrassMapsBuilder::UpdateTrackedComponents(const TArray<FVector>&
 
 				It.RemoveCurrent();
 				AmortizedUpdate.HandleDeletion(ComponentStateIndex);
+				ComponentStateIndex--;
 				check(PendingCount > 0);
 				PendingCount--;
 			}
