Fixed / improved landscape texture diff CVar utilities :
* One texture / diff report is dumped per landscape paint layer (it used to be one per weightmap, which was incorrect if multiple layers from a given weightmap texture changed, as the last layer would override the other ones
* Misc minor improvements in the diff report

diff --git a/Engine/Source/Runtime/Landscape/Classes/Landscape.h b/Engine/Source/Runtime/Landscape/Classes/Landscape.h
--- a/Engine/Source/Runtime/Landscape/Classes/Landscape.h
+++ b/Engine/Source/Runtime/Landscape/Classes/Landscape.h
@@ -486,7 +486,7 @@ private:
 	int32 PerformLayersWeightmapsGlobalMerge(FUpdateLayersContentContext& InUpdateLayersContentContext, const FEditLayersWeightmapMergeParams& InMergeParams);
 	void ResolveLayersWeightmapTexture(const FTextureToComponentHelper& MapHelper, const TSet<UTexture2D*>& WeightmapsToResolve, bool bIntermediateRender, bool bFlushRender, TArray<FLandscapeEditLayerComponentReadbackResult>& InOutComponentReadbackResults);
 
-	bool HasTextureDataChanged(TArrayView<const FColor> InOldData, TArrayView<const FColor> InNewData, bool bInIsWeightmap, uint64 InPreviousHash, uint64& OutNewHash) const;
+	bool HasTextureDataChanged(TArrayView<const FColor> InOldData, TArrayView<const FColor> InNewData, bool bInIsWeightmap, uint64 InPreviousHash, uint64& OutNewHash, TOptional<uint8>& OutChangedWeightmapChannelsMasks) const;
 	bool ResolveLayersTexture(FTextureToComponentHelper const& MapHelper, FLandscapeEditLayerReadback* InCPUReadBack, UTexture2D* InOutputTexture, bool bIntermediateRender,	bool bFlushRender,
 		TArray<FLandscapeEditLayerComponentReadbackResult>& InOutComponentReadbackResults, bool bIsWeightmap);
 
@@ -543,7 +543,7 @@ private:
 	void PrintLayersDebugWeightData(const FString& InContext, const TArray<FColor>& InWeightmapData, const FIntPoint& InDataSize, uint8 InMipRender) const;
 
 	void UpdateWeightDirtyData(ULandscapeComponent* InLandscapeComponent, UTexture2D const* InWeightmap, FColor const* InOldData, FColor const* InNewData, uint8 InChannel);
-	void OnDirtyWeightmap(FTextureToComponentHelper const& MapHelper, UTexture2D const* InWeightmap, FColor const* InOldData, FColor const* InNewData, int32 InMipLevel);
+	void OnDirtyWeightmap(FTextureToComponentHelper const& MapHelper, UTexture2D const* InWeightmap, FColor const* InOldData, FColor const* InNewData, int32 InMipLevel, uint8 ChangedChannelsMask);
 	void UpdateHeightDirtyData(ULandscapeComponent* InLandscapeComponent, UTexture2D const* InHeightmap, FColor const* InOldData, FColor const* InNewData);
 	void OnDirtyHeightmap(FTextureToComponentHelper const& MapHelper, UTexture2D const* InWeightmap, FColor const* InOldData, FColor const* InNewData, int32 InMipLevel);
 
diff --git a/Engine/Source/Runtime/Landscape/Private/LandscapeEditLayers.cpp b/Engine/Source/Runtime/Landscape/Private/LandscapeEditLayers.cpp
--- a/Engine/Source/Runtime/Landscape/Private/LandscapeEditLayers.cpp
+++ b/Engine/Source/Runtime/Landscape/Private/LandscapeEditLayers.cpp
@@ -5316,7 +5316,7 @@ void ALandscape::UpdateWeightDirtyData(ULandscapeComponent* InLandscapeComponent
 	LandscapeEdit.SetDirtyData(X1, Y1, X2, Y2, DirtyData.Get(), 0);
 }
 
-void ALandscape::OnDirtyWeightmap(FTextureToComponentHelper const& MapHelper, UTexture2D const* InWeightmap, FColor const* InOldData, FColor const* InNewData, int32 InMipLevel)
+void ALandscape::OnDirtyWeightmap(FTextureToComponentHelper const& MapHelper, UTexture2D const* InWeightmap, FColor const* InOldData, FColor const* InNewData, int32 InMipLevel, uint8 ChangedChannelsMask)
 {
 	int32 DumpWeightmapDiff = CVarLandscapeDumpWeightmapDiff.GetValueOnGameThread();
 	const bool bDumpDiff = (DumpWeightmapDiff > 0);
@@ -5334,6 +5334,8 @@ void ALandscape::OnDirtyWeightmap(FTextureToComponentHelper const& MapHelper, UT
 		return;
 	}
 
+	check(ChangedChannelsMask != 0);
+
 	TArray<ULandscapeComponent*> const* Components = MapHelper.WeightmapToComponents.Find(InWeightmap);
 	if (Components != nullptr)
 	{
@@ -5345,7 +5347,9 @@ void ALandscape::OnDirtyWeightmap(FTextureToComponentHelper const& MapHelper, UT
 			for (FWeightmapLayerAllocationInfo const& AllocInfo : AllocInfos)
 			{
 				check(AllocInfo.IsAllocated() && AllocInfo.WeightmapTextureIndex < WeightmapTextures.Num());
-				if (InWeightmap == WeightmapTextures[AllocInfo.WeightmapTextureIndex])
+				if (InWeightmap == WeightmapTextures[AllocInfo.WeightmapTextureIndex] 
+					// Only dump if that particular weightmap channel has changed
+					&& ((1 << AllocInfo.WeightmapTextureChannel) & ChangedChannelsMask) != 0)
 				{
 					if (bTrackDirty)
 					{
@@ -5362,7 +5366,8 @@ void ALandscape::OnDirtyWeightmap(FTextureToComponentHelper const& MapHelper, UT
 						ALandscapeProxy* Proxy = Cast<ALandscapeProxy>(Component->GetOwner());
 						check(Proxy);
 						FString ActorName = Proxy->GetActorLabel();
-						FString FilePattern = FString::Format(TEXT("{0}/LandscapeLayers/{1}/{2}/{3}/Weightmaps/{4}-{5}-{6}[mip{7}]"), { FPaths::ProjectSavedDir(), CurrentTime.ToString(), WorldName, ParentLandscapeActorName, ActorName, Component->GetName(), InWeightmap->GetName(), InMipLevel });
+						FString FilePattern = FString::Format(TEXT("{0}/LandscapeLayers/{1}/{2}/{3}/Weightmaps/{4}/{5}-{6}-{7}[mip{8}]"), 
+							{ FPaths::ProjectSavedDir(), CurrentTime.ToString(), WorldName, ParentLandscapeActorName, AllocInfo.GetLayerName().ToString(), ActorName, Component->GetName(), InWeightmap->GetName(), InMipLevel});
 
 						FFileHelper::EColorChannel ColorChannel = GetWeightmapColorChannel(AllocInfo);
 						FFileHelper::CreateBitmap(*(FilePattern + "_a(pre).bmp"), SizeU, SizeV, InOldData, /*SubRectangle = */nullptr, &IFileManager::Get(), /*OutFilename = */nullptr, /*bool bInWriteAlpha = */true, ColorChannel);
@@ -5370,9 +5375,11 @@ void ALandscape::OnDirtyWeightmap(FTextureToComponentHelper const& MapHelper, UT
 			
 						if (bDumpDiffDetails)
 						{
+							static const TCHAR* Channels = TEXT("RGBA");
 							int32 NumDifferentPixels = 0;
 							uint8 MaxDiff = 0;
 							FStringBuilderBase StrBuilder;
+							FIntPoint MaxDiffUV(ForceInit);
 							for (int32 V = 0; V < SizeV; ++V)
 							{
 								for (int32 U = 0; U < SizeU; ++U)
@@ -5381,30 +5388,39 @@ void ALandscape::OnDirtyWeightmap(FTextureToComponentHelper const& MapHelper, UT
 									const FColor* NewDataPtr = InNewData + (V * SizeU + U);
 									if (*OldDataPtr != *NewDataPtr)
 									{
-										auto ComputeMaxDiff = [&MaxDiff](uint8 OldValue, uint8 NewValue) -> uint8
+										uint32 OldValueAsUInt32 = OldDataPtr->ToPackedRGBA();
+										uint8 OldValue = (OldValueAsUInt32 >> ((3 - AllocInfo.WeightmapTextureChannel) * 8)) & 0xff;
+										uint32 NewValueAsUInt32 = NewDataPtr->ToPackedRGBA();
+										uint8 NewValue = (NewValueAsUInt32 >> ((3 - AllocInfo.WeightmapTextureChannel) * 8)) & 0xff;
+										uint8 Diff = (NewValue > OldValue) ? NewValue - OldValue : OldValue - NewValue;
+										if (Diff > 0)
 										{
-											uint8 Diff = (NewValue > OldValue) ? (NewValue - OldValue) : (OldValue - NewValue);
-											MaxDiff = FMath::Max<uint8>(Diff, MaxDiff);
-											return Diff;
-										};
-
-										StrBuilder.Appendf(TEXT("Pixel (%4u,%4u) : R (%3u -> %3u, absdiff %3u), G (%3u -> %3u, absdiff %3u), B (%3u -> %3u, absdiff %3u), A (%3u -> %3u, absdiff %3u)\n"), U, V, 
-											OldDataPtr->R, NewDataPtr->R, ComputeMaxDiff(OldDataPtr->R, NewDataPtr->R),
-											OldDataPtr->G, NewDataPtr->G, ComputeMaxDiff(OldDataPtr->G, NewDataPtr->G),
-											OldDataPtr->B, NewDataPtr->B, ComputeMaxDiff(OldDataPtr->B, NewDataPtr->B),
-											OldDataPtr->A, NewDataPtr->A, ComputeMaxDiff(OldDataPtr->A, NewDataPtr->A));
-
-										++NumDifferentPixels;
+											if (Diff > MaxDiff)
+											{
+												MaxDiffUV = FIntPoint(U, V);
+												MaxDiff = Diff;
+											}
+
+											TCHAR Channel = (AllocInfo.WeightmapTextureChannel == 0) ? TEXT('R') 
+												: (AllocInfo.WeightmapTextureChannel == 1) ? TEXT('G') 
+												: (AllocInfo.WeightmapTextureChannel == 2) ? TEXT('B') 
+												: TEXT('A');
+
+											StrBuilder.Appendf(TEXT("Pixel (%4u,%4u) : RGBA ((%3u,%3u,%3u,%3u) -> (%3u,%3u,%3u,%3u)) : channel %c (%3u -> %3u, absdiff %3u)\n"), 
+												U, V, OldDataPtr->R, OldDataPtr->G, OldDataPtr->B, OldDataPtr->A, NewDataPtr->R, NewDataPtr->G, NewDataPtr->B, NewDataPtr->A, Channels[AllocInfo.WeightmapTextureChannel], OldValue, NewValue, Diff);
+
+											++NumDifferentPixels;
+										}
 									}
 								}
-					}
-							StrBuilder.Appendf(TEXT("----------------------------------------\n"));
-							StrBuilder.Appendf(TEXT("Num diffs = %u\n"), NumDifferentPixels);
-							StrBuilder.Appendf(TEXT("Max diff = %u (%1.3f%%)\n"), MaxDiff, 100.0 * static_cast<float>(MaxDiff) / MAX_uint8);
+							}
+							StrBuilder.InsertAt(0, FString::Printf(TEXT("----------------------------------------\n")));
+							StrBuilder.InsertAt(0, FString::Printf(TEXT("Max diff (at %s) = %u (%1.3f%%)\n"), *MaxDiffUV.ToString(), MaxDiff, 100.0 * static_cast<float>(MaxDiff) / MAX_uint8));
+							StrBuilder.InsertAt(0, FString::Printf(TEXT("Num diffs = %u\n"), NumDifferentPixels));
+							StrBuilder.InsertAt(0, FString::Printf(TEXT("Layer %s is packed in channel %c\n"), *AllocInfo.GetLayerName().ToString(), Channels[AllocInfo.WeightmapTextureChannel]));
 							FFileHelper::SaveStringToFile(StrBuilder.ToView(), *(FilePattern + "_diff.txt"));
 						}
 					}
-
 				}
 			}
 		}
@@ -5499,7 +5515,9 @@ void ALandscape::OnDirtyHeightmap(FTextureToComponentHelper const& MapHelper, UT
 				{
 					int32 NumDifferentPixels = 0;
 					uint16 MaxHeightDiff = 0;
+					FIntPoint MaxHeightDiffUV(ForceInit);
 					uint8 MaxNormalDiff = 0;
+					FIntPoint MaxNormalDiffUV(ForceInit);
 					FStringBuilderBase StrBuilder;
 					const FColor* OldDataStartPtr = InOldData + (HeightmapOffsetY * SizeU + HeightmapOffsetX);
 					const FColor* NewDataStartPtr = InNewData + (HeightmapOffsetY * SizeU + HeightmapOffsetX);
@@ -5513,32 +5531,42 @@ void ALandscape::OnDirtyHeightmap(FTextureToComponentHelper const& MapHelper, UT
 							{
 								uint16 OldHeight = ((static_cast<uint16>(OldDataPtr->R) << 8) | static_cast<uint16>(OldDataPtr->G));
 								uint16 NewHeight = ((static_cast<uint16>(NewDataPtr->R) << 8) | static_cast<uint16>(NewDataPtr->G));
-								uint16 HeightDiff = (NewHeight > OldHeight) ? (NewHeight - OldHeight) : (OldHeight - NewHeight);
-								MaxHeightDiff = FMath::Max<uint16>(HeightDiff, MaxHeightDiff);
+								uint16 HeightDiff = (NewHeight > OldHeight) ? NewHeight - OldHeight : OldHeight - NewHeight;
+								if (HeightDiff > MaxHeightDiff)
+								{
+									MaxHeightDiffUV = FIntPoint(U, V);
+									MaxHeightDiff = HeightDiff;
+								}
 
 								uint8 OldNormalX = (OldDataPtr->B);
 								uint8 NewNormalX = (NewDataPtr->B);
-								uint8 NormalXDiff = (NewNormalX > OldNormalX) ? (NewNormalX - OldNormalX) : (OldNormalX - NewNormalX);
-								MaxNormalDiff = FMath::Max<uint8>(NormalXDiff, MaxNormalDiff);
+								uint16 NormalXDiff = (NewNormalX > OldNormalX) ? NewNormalX - OldNormalX : OldNormalX - NewNormalX;
+								if (NormalXDiff > MaxNormalDiff)
+								{
+									MaxNormalDiffUV = FIntPoint(U, V);
+									MaxNormalDiff = NormalXDiff;
+								}
 
 								uint8 OldNormalY = (OldDataPtr->A);
 								uint8 NewNormalY = (NewDataPtr->A);
-								uint8 NormalYDiff = (NewNormalY > OldNormalY) ? (NewNormalY - OldNormalY) : (OldNormalY - NewNormalY);
-								MaxNormalDiff = FMath::Max<uint8>(NormalYDiff, MaxNormalDiff);
+								uint16 NormalYDiff = (NewNormalY > OldNormalY) ? NewNormalY - OldNormalY : OldNormalY - NewNormalY;
+								if (NormalYDiff > MaxNormalDiff)
+								{
+									MaxNormalDiffUV = FIntPoint(U, V);
+									MaxNormalDiff = NormalYDiff;
+								}
 
-								StrBuilder.Appendf(TEXT("Pixel (%4u,%4u) : Height (%5u -> %5u, absdiff %5u), NormalX (%3u -> %3u, absdiff %3u), NormalY (%3u -> %3u, absdiff %3u)\n"), U, V, 
-									OldHeight, NewHeight, HeightDiff, 
-									OldNormalX, NewNormalX, NormalXDiff, 
-									OldNormalY, NewNormalY, NormalYDiff);
+								StrBuilder.Appendf(TEXT("Pixel (%4u,%4u) : Height (%5u -> %5u, absdiff %5u), Normal ((%3u,%3u) -> (%3u,%3u), absdiff %3u)\n"), 
+									U, V, OldHeight, NewHeight, HeightDiff, OldNormalX, OldNormalY, NewNormalX, NewNormalY, FMath::Max(NormalXDiff, NormalYDiff));
 
 								++NumDifferentPixels;
 							}
 						}
 					}
-					StrBuilder.Appendf(TEXT("----------------------------------------\n"));
-					StrBuilder.Appendf(TEXT("Num diffs = %u\n"), NumDifferentPixels);
-					StrBuilder.Appendf(TEXT("Max height diff = %u (%1.3f%%)\n"), MaxHeightDiff, 100.0f * static_cast<float>(MaxHeightDiff) / MAX_uint16);
-					StrBuilder.Appendf(TEXT("Max normal diff = %u (%1.3f%%)\n"), MaxNormalDiff, 100.0f * static_cast<float>(MaxNormalDiff) / MAX_uint8);
+					StrBuilder.InsertAt(0, FString::Printf(TEXT("----------------------------------------\n")));
+					StrBuilder.InsertAt(0, FString::Printf(TEXT("Max normal diff (at %s) = %u (%1.3f%%)\n"), *MaxNormalDiffUV.ToString(), MaxNormalDiff, 100.0f * static_cast<float>(MaxNormalDiff) / MAX_uint8));
+					StrBuilder.InsertAt(0, FString::Printf(TEXT("Max height diff (at %s) = %u (%1.3f%%)\n"), *MaxHeightDiffUV.ToString(), MaxHeightDiff, 100.0f * static_cast<float>(MaxHeightDiff) / MAX_uint16));
+					StrBuilder.InsertAt(0, FString::Printf(TEXT("Num diffs = %u\n"), NumDifferentPixels));
 					FFileHelper::SaveStringToFile(StrBuilder.ToView(), *(FilePattern + "_diff.txt"));
 				}
 			}
@@ -5546,7 +5574,7 @@ void ALandscape::OnDirtyHeightmap(FTextureToComponentHelper const& MapHelper, UT
 	}
 }
 
-bool ALandscape::HasTextureDataChanged(TArrayView<const FColor> InOldData, TArrayView<const FColor> InNewData, bool bInIsWeightmap, uint64 InPreviousHash, uint64& OutNewHash) const
+bool ALandscape::HasTextureDataChanged(TArrayView<const FColor> InOldData, TArrayView<const FColor> InNewData, bool bInIsWeightmap, uint64 InPreviousHash, uint64& OutNewHash, TOptional<uint8>& OutChangedWeightmapChannelsMasks) const
 {
 	TRACE_CPUPROFILER_EVENT_SCOPE(ALandscape::HasTextureDataChanged);
 
@@ -5559,8 +5587,11 @@ bool ALandscape::HasTextureDataChanged(TArrayView<const FColor> InOldData, TArra
 		// If necessary, perform a deep comparison to avoid taking into account precision issues as actual changes :
 		if (bInIsWeightmap)
 		{
-			if (int32 DirtyWeightmapThreshold = CVarLandscapeDirtyWeightmapThreshold.GetValueOnGameThread(); DirtyWeightmapThreshold > 0)
+			int32 DirtyWeightmapThreshold = CVarLandscapeDirtyWeightmapThreshold.GetValueOnGameThread();
+			if ((DirtyWeightmapThreshold > 0) || OutChangedWeightmapChannelsMasks.IsSet())
 			{
+				bool bHasChanged = false;
+
 				TRACE_CPUPROFILER_EVENT_SCOPE(DeepCompareWeightmap);
 				for (int32 Index = 0; Index < TextureSize; ++Index)
 				{
@@ -5568,27 +5599,34 @@ bool ALandscape::HasTextureDataChanged(TArrayView<const FColor> InOldData, TArra
 					const FColor& NewColor = InNewData[Index];
 					if (OldColor != NewColor)
 					{
-						if (uint8 Diff = (NewColor.R > OldColor.R) ? (NewColor.R - OldColor.R) : (OldColor.R - NewColor.R); Diff > DirtyWeightmapThreshold)
-						{
-							return true;
-						}
-						if (uint8 Diff = (NewColor.G > OldColor.G) ? (NewColor.G - OldColor.G) : (OldColor.G - NewColor.G); Diff > DirtyWeightmapThreshold)
-						{
-							return true;
-						}
-						if (uint8 Diff = (NewColor.B > OldColor.B) ? (NewColor.B - OldColor.B) : (OldColor.B - NewColor.B); Diff > DirtyWeightmapThreshold)
+						auto DiffChannel = [DirtyWeightmapThreshold] (uint8 InOldValue, uint8 InNewValue) -> bool
 						{
-							return true;
-						}
-						if (uint8 Diff = (NewColor.A > OldColor.A) ? (NewColor.A - OldColor.A) : (OldColor.A - NewColor.A); Diff > DirtyWeightmapThreshold)
+							uint8 Diff = (InNewValue > InOldValue) ? (InNewValue - InOldValue) : (InOldValue - InNewValue);
+							return (Diff > DirtyWeightmapThreshold);
+						};
+
+						uint8 DiffMask = ((uint8)(DiffChannel(OldColor.R, NewColor.R) ? 1 : 0) << 0)
+							| ((uint8)(DiffChannel(OldColor.G, NewColor.G) ? 1 : 0) << 1)
+							| ((uint8)(DiffChannel(OldColor.B, NewColor.B) ? 1 : 0) << 2)
+							| ((uint8)(DiffChannel(OldColor.A, NewColor.A) ? 1 : 0) << 3);
+
+						if (DiffMask != 0)
 						{
-							return true;
+							bHasChanged = true;
+							if (OutChangedWeightmapChannelsMasks.IsSet())
+							{
+								*OutChangedWeightmapChannelsMasks |= DiffMask;
+							}
+							else
+							{
+								// no need to report which channel has been changed, early out : 
+								return true;
+							}
 						}
 					}
 				}
 
-				// No significant difference detected : 
-				return false;
+				return bHasChanged;
 			}
 		}
 		else 
@@ -5671,6 +5709,16 @@ bool ALandscape::ResolveLayersTexture(
 	}
 	
 	bool bChanged = false;
+	TOptional<uint8> ChangedChannelsMask;
+	if (bIsWeightmap)
+	{
+		// Request a precise report of which channel have changed if we need to dump the weightmap diffs :
+		if ((CVarLandscapeDumpWeightmapDiff.GetValueOnGameThread() != 0) 
+			|| CVarLandscapeTrackDirty.GetValueOnGameThread() != 0)
+		{
+			ChangedChannelsMask.Emplace(0);
+		}
+	}
 
 	if (CompletedReadbackNum > 0)
 	{
@@ -5691,7 +5739,7 @@ bool ALandscape::ResolveLayersTexture(
 				if (MipIndex == 0 && !bIntermediateRender)
 				{
 					uint64 Hash = 0;
-					if (HasTextureDataChanged(MakeArrayView(reinterpret_cast<const FColor*>(TextureData), TextureSize), MakeArrayView(OutMipsData[MipIndex].GetData(), TextureSize), bIsWeightmap, InCPUReadback->GetHash(), Hash))
+					if (HasTextureDataChanged(MakeArrayView(reinterpret_cast<const FColor*>(TextureData), TextureSize), MakeArrayView(OutMipsData[MipIndex].GetData(), TextureSize), bIsWeightmap, InCPUReadback->GetHash(), Hash, ChangedChannelsMask))
 					{
 						bChanged |= InCPUReadback->SetHash(Hash);
 						check(bChanged);
@@ -5716,7 +5764,7 @@ bool ALandscape::ResolveLayersTexture(
 				{
 					if (bIsWeightmap)
 					{
-						OnDirtyWeightmap(MapHelper, InOutputTexture, (FColor*)TextureData, OutMipsData[MipIndex].GetData(), MipIndex);
+						OnDirtyWeightmap(MapHelper, InOutputTexture, (FColor*)TextureData, OutMipsData[MipIndex].GetData(), MipIndex, ChangedChannelsMask.GetValue());
 					}
 					else
 					{
