Reduce Unused FNames in RHI_ENABLE_RESOURCE_INFO=0 builds (Test/Shipping)
When RHI_ENABLE_RESOURCE_INFO=0 then the FNames passed to FRenderResource's are not stored and the name is always returned as NAME_None (See FRenderResource::[Get|Set]OwnerName)
The constructor for the FName will still add the name to the table. In some situations the name already exists, but in others we are storing strings we wouldn't use. In the situations where the name exists we still have to exclusively lock the slot in the FName pool (See FNamePool::Store / FNamePoolShard.Insert)
#pf 664b6052623a833f59900315

diff --git a/Engine/Source/Runtime/Engine/Private/Animation/SkinWeightProfile.cpp b/Engine/Source/Runtime/Engine/Private/Animation/SkinWeightProfile.cpp
--- a/Engine/Source/Runtime/Engine/Private/Animation/SkinWeightProfile.cpp
+++ b/Engine/Source/Runtime/Engine/Private/Animation/SkinWeightProfile.cpp
@@ -367,8 +367,10 @@ void FSkinWeightProfilesData::SetDynamicDefaultSkinWeightProfile(USkeletalMesh*
 								bDefaultOverridden = true;
 								DefaultProfileName = Profiles[DefaultProfileIndex].Name;
 								
+#if RHI_ENABLE_RESOURCE_INFO
 								const FName OwnerName(USkinnedAsset::GetLODPathName(Mesh, LODIndex));
 								OverrideBuffer->SetOwnerName(OwnerName);
+#endif
 								OverrideBuffer->BeginInitResources();
 							}
 						}
diff --git a/Engine/Source/Runtime/Engine/Private/Components/SkinnedMeshComponent.cpp b/Engine/Source/Runtime/Engine/Private/Components/SkinnedMeshComponent.cpp
--- a/Engine/Source/Runtime/Engine/Private/Components/SkinnedMeshComponent.cpp
+++ b/Engine/Source/Runtime/Engine/Private/Components/SkinnedMeshComponent.cpp
@@ -4811,8 +4811,10 @@ void USkinnedMeshComponent::SetSkinWeightOverride(int32 LODIndex, const TArray<F
 			Info.OverrideSkinWeights->SetUse16BitBoneIndex(bUse16BitBoneIndex);
 
 			const USkinnedAsset* Mesh = GetSkinnedAsset();
+#if RHI_ENABLE_RESOURCE_INFO
 			const FName OwnerName(USkinnedAsset::GetLODPathName(Mesh, LODIndex));
 			Info.OverrideSkinWeights->SetOwnerName(OwnerName);
+#endif
 
 			const FReferenceSkeleton& RefSkel = Mesh->GetRefSkeleton();
 			TArray<FSkinWeightInfo> GPUWeights;
diff --git a/Engine/Source/Runtime/Engine/Private/Texture.cpp b/Engine/Source/Runtime/Engine/Private/Texture.cpp
--- a/Engine/Source/Runtime/Engine/Private/Texture.cpp
+++ b/Engine/Source/Runtime/Engine/Private/Texture.cpp
@@ -396,7 +396,10 @@ void UTexture::UpdateResource()
 			}
 
 			PrivateResource = NewResource;
+
+#if RHI_ENABLE_RESOURCE_INFO
 			NewResource->SetOwnerName(FName(GetPathName()));
+#endif
 
 			// Init the texture reference, which needs to be set from a render command, since TextureReference.TextureReferenceRHI is gamethread coherent.
 			ENQUEUE_RENDER_COMMAND(SetTextureReference)([this, NewResource](FRHICommandListImmediate& RHICmdList)
diff --git a/Engine/Source/Runtime/Engine/Private/TextureCollection.cpp b/Engine/Source/Runtime/Engine/Private/TextureCollection.cpp
--- a/Engine/Source/Runtime/Engine/Private/TextureCollection.cpp
+++ b/Engine/Source/Runtime/Engine/Private/TextureCollection.cpp
@@ -181,7 +181,9 @@ void UTextureCollection::UpdateResource()
 		LLM_SCOPE(ELLMTag::Textures);
 
 		PrivateResource = NewResource;
+#if RHI_ENABLE_RESOURCE_INFO
 		NewResource->SetOwnerName(FName(GetPathName()));
+#endif
 
 		ENQUEUE_RENDER_COMMAND(SetTextureCollectionResource)([this, NewResource](FRHICommandListImmediate& RHICmdList)
 		{
diff --git a/Engine/Source/Runtime/Landscape/Private/LandscapeCollision.cpp b/Engine/Source/Runtime/Landscape/Private/LandscapeCollision.cpp
--- a/Engine/Source/Runtime/Landscape/Private/LandscapeCollision.cpp
+++ b/Engine/Source/Runtime/Landscape/Private/LandscapeCollision.cpp
@@ -652,10 +652,12 @@ FPrimitiveSceneProxy* ULandscapeHeightfieldCollisionComponent::CreateSceneProxy(
 			// Allocate the static vertex resources now 
 			if (Vertices.Num() > 0)
 			{
+#if RHI_ENABLE_RESOURCE_INFO
 				FName Name = FName(TEXT("FLandscapeHeightfieldCollisionComponentSceneProxy ") + GetOwnerName().ToString());
 				VertexBuffers.SetOwnerName(Name);
 				IndexBuffer.SetOwnerName(Name);
 				VertexFactory.SetOwnerName(Name);
+#endif
 
 				VertexBuffers.InitFromDynamicVertex(&VertexFactory, Vertices);
 				BeginInitResource(&VertexBuffers.PositionVertexBuffer);
