Fix for crash in serialization for PIE when using Mobile preview mode

The textures and materials array must have the same number of entries in serialize and deserialize, and thus cannot depend on whether the pointer is null or not (in fact we replace the property pointers with null on serialize, so it's guaranteed to be null on deserialize, and always have a mismatch in the count)

diff --git a/Engine/Source/Runtime/Landscape/Private/Landscape.cpp b/Engine/Source/Runtime/Landscape/Private/Landscape.cpp
--- a/Engine/Source/Runtime/Landscape/Private/Landscape.cpp
+++ b/Engine/Source/Runtime/Landscape/Private/Landscape.cpp
@@ -803,12 +803,9 @@ void ULandscapeComponent::Serialize(FArchive& Ar)
 		{
 			TexturesAndMaterials.Add((UObject**)&static_cast<UTexture2D*&>(MobileWeightmapTexture));
 		}
-		
-		if (MobileWeightmapTextureArray)
-		{
-			TexturesAndMaterials.Add((UObject**)&static_cast<UTexture2DArray*&>(MobileWeightmapTextureArray));
-		}
-		
+
+		TexturesAndMaterials.Add((UObject**)&static_cast<UTexture2DArray*&>(MobileWeightmapTextureArray));
+
 		for (auto& ItPair : LayersData)
 		{
 			FLandscapeLayerComponentData& LayerComponentData = ItPair.Value;
