Implement PSO precaching miss fixes for FLandscapeGrassWeightMeshProcessor

diff --git a/Engine/Source/Runtime/Landscape/Private/LandscapeGrassWeightExporter.cpp b/Engine/Source/Runtime/Landscape/Private/LandscapeGrassWeightExporter.cpp
--- a/Engine/Source/Runtime/Landscape/Private/LandscapeGrassWeightExporter.cpp
+++ b/Engine/Source/Runtime/Landscape/Private/LandscapeGrassWeightExporter.cpp
@@ -183,6 +183,12 @@ public:
 	{
 		checkf(false, TEXT("Default AddMeshBatch can't be used as rendering requires extra parameters per pass."));
 	}
+	
+	virtual void CollectPSOInitializers(const FSceneTexturesConfig& SceneTexturesConfig, 
+		const FMaterial& Material, 
+		const FPSOPrecacheVertexFactoryData& VertexFactoryData, 
+		const FPSOPrecacheParams& PreCacheParams, 
+		TArray<FPSOPrecacheData>& PSOInitializers) override final;
 
 private:
 	bool TryAddMeshBatch(
@@ -213,8 +219,10 @@ private:
 	FMeshPassProcessorRenderState PassDrawRenderState;
 };
 
+static const TCHAR* LandscapeGrassWeightMeshPassName = TEXT("LandscapeGrassWeight");
+
 FLandscapeGrassWeightMeshProcessor::FLandscapeGrassWeightMeshProcessor(const FScene* Scene, ERHIFeatureLevel::Type InFeatureLevel, const FSceneView* InViewIfDynamicMeshCommand, FMeshPassDrawListContext* InDrawListContext)
-	: FMeshPassProcessor(EMeshPass::Num, Scene, InFeatureLevel, InViewIfDynamicMeshCommand, InDrawListContext)
+	: FMeshPassProcessor(LandscapeGrassWeightMeshPassName, Scene, InFeatureLevel, InViewIfDynamicMeshCommand, InDrawListContext)
 {
 	PassDrawRenderState.SetBlendState(TStaticBlendState<>::GetRHI());
 	PassDrawRenderState.SetDepthStencilState(TStaticDepthStencilState<false, CF_Always>::GetRHI());
@@ -327,6 +335,58 @@ bool FLandscapeGrassWeightMeshProcessor::Process(
 }
 
 
+void FLandscapeGrassWeightMeshProcessor::CollectPSOInitializers(const FSceneTexturesConfig& SceneTexturesConfig, const FMaterial& Material, const FPSOPrecacheVertexFactoryData& VertexFactoryData, const FPSOPrecacheParams& PreCacheParams, TArray<FPSOPrecacheData>& PSOInitializers)
+{
+	// Only support the Landscape fixed grid vertex factory type.
+	if (VertexFactoryData.VertexFactoryType != &FLandscapeFixedGridVertexFactory::StaticType)
+	{
+		return;
+	}
+
+	FMaterialShaderTypes ShaderTypes;
+	ShaderTypes.AddShaderType<FLandscapeGrassWeightVS>();
+	ShaderTypes.AddShaderType<FLandscapeGrassWeightPS>();
+
+	FMaterialShaders Shaders;
+	if (!Material.TryGetShaders(ShaderTypes, VertexFactoryData.VertexFactoryType, Shaders))
+	{
+		return;
+	}
+
+	TMeshProcessorShaders<
+		FLandscapeGrassWeightVS,
+		FLandscapeGrassWeightPS> PassShaders;
+	Shaders.TryGetVertexShader(PassShaders.VertexShader);
+	Shaders.TryGetPixelShader(PassShaders.PixelShader);
+
+	const FMeshDrawingPolicyOverrideSettings OverrideSettings = ComputeMeshOverrideSettings(PreCacheParams);
+	const ERasterizerFillMode MeshFillMode = ComputeMeshFillMode(Material, OverrideSettings);
+	const ERasterizerCullMode MeshCullMode = CM_None;
+
+	FGraphicsPipelineRenderTargetsInfo RenderTargetsInfo;
+	RenderTargetsInfo.NumSamples = 1;
+	AddRenderTargetInfo(PF_B8G8R8A8, ETextureCreateFlags::RenderTargetable, RenderTargetsInfo);
+	
+	AddGraphicsPipelineStateInitializer(
+		VertexFactoryData,
+		Material,
+		PassDrawRenderState,
+		RenderTargetsInfo,
+		PassShaders,
+		MeshFillMode,
+		MeshCullMode,
+		PT_PointList,
+		EMeshPassFeatures::Default,
+		true /*bRequired*/,
+		PSOInitializers);
+}
+
+IPSOCollector* CreateLandscapeGrassWeightPSOCollector(ERHIFeatureLevel::Type FeatureLevel)
+{
+	return new FLandscapeGrassWeightMeshProcessor(nullptr, FeatureLevel, nullptr, nullptr);
+}
+FRegisterPSOCollectorCreateFunction RegisterLandscapeGrassWeightPSOCollector(&CreateLandscapeGrassWeightPSOCollector, EShadingPath::Deferred, LandscapeGrassWeightMeshPassName);
+
 void FLandscapeGrassWeightExporter_RenderThread::RenderLandscapeComponentToTexture_RenderThread(FRHICommandListImmediate& RHICmdList)
 {
 	FRDGBuilder GraphBuilder(RHICmdList);
