Landscape: Fix incorrect transform being passed to legacy global merge adapters when landscape has components deleted.

diff --git a/Engine/Source/Runtime/Landscape/Private/LandscapeEditLayers.cpp b/Engine/Source/Runtime/Landscape/Private/LandscapeEditLayers.cpp
--- a/Engine/Source/Runtime/Landscape/Private/LandscapeEditLayers.cpp
+++ b/Engine/Source/Runtime/Landscape/Private/LandscapeEditLayers.cpp
@@ -76,6 +76,18 @@ LandscapeEditLayers.cpp: Landscape editing layers mode
 
 #define LOCTEXT_NAMESPACE "Landscape"
 
+namespace LandscapeEditLayersLocals
+{
+	FTransform GetLandscapeExtentTransform(const ALandscape& Landscape, const FIntRect& InLandscapeExtent)
+	{
+		FTransform NewLandscapeTransform = Landscape.GetTransform();
+		FVector OffsetVector(InLandscapeExtent.Min.X, InLandscapeExtent.Min.Y, 0.f);
+		FVector Translation = NewLandscapeTransform.TransformFVector4(OffsetVector);
+		NewLandscapeTransform.SetTranslation(Translation);
+		return NewLandscapeTransform;
+	}
+}
+
 // Channel remapping
 extern const size_t ChannelOffsets[4];
 
@@ -4976,7 +4988,9 @@ int32 ALandscape::PerformLayersHeightmapsGlobalMerge(const FUpdateLayersContentC
 				FLandscapeBrushParameters BrushParameters = BuildLandscapeBrushParameters(
 					/*bIsHeightmapMerge = */true, LandscapeExtent, CombinedHeightmapNonAtlasRT);
 
-				UTextureRenderTarget2D* BrushOutputNonAtlasRT = RenderCallAdapter->RenderAsBlueprintBrush(BrushParameters, GetTransform(), 
+				UTextureRenderTarget2D* BrushOutputNonAtlasRT = RenderCallAdapter->RenderAsBlueprintBrush(
+					BrushParameters, 
+					LandscapeEditLayersLocals::GetLandscapeExtentTransform(*this, LandscapeExtent),
 					// TODO: Someday this will be part of BrushParameters
 					LandscapeExtent.Max - LandscapeExtent.Min);
 
@@ -7460,7 +7474,9 @@ int32 ALandscape::PerformLayersWeightmapsGlobalMerge(FUpdateLayersContentContext
 							FLandscapeBrushParameters BrushParameters = BuildLandscapeBrushParameters(
 								/*bIsHeightmapMerge = */false, LandscapeExtent, LandscapeScratchRT3, LayerInfoObj->LayerName);
 							
-							UTextureRenderTarget2D* BrushOutputRT = RenderCallAdapter->RenderAsBlueprintBrush(BrushParameters, GetTransform(),
+							UTextureRenderTarget2D* BrushOutputRT = RenderCallAdapter->RenderAsBlueprintBrush(
+								BrushParameters, 
+								LandscapeEditLayersLocals::GetLandscapeExtentTransform(*this, LandscapeExtent),
 								// TODO: Someday this will be part of BrushParameters
 								LandscapeExtent.Max - LandscapeExtent.Min);
 
@@ -10471,10 +10487,7 @@ bool FLandscapeLayerBrush::Initialize(const FIntRect& InLandscapeExtent, UTextur
 		if (ALandscape* Landscape = BlueprintBrush->GetOwningLandscape())
 		{
 			const FIntPoint NewLandscapeRenderTargetSize = FIntPoint(InLandscapeRenderTarget->SizeX, InLandscapeRenderTarget->SizeY);
-			FTransform NewLandscapeTransform = Landscape->GetTransform();
-			FVector OffsetVector(InLandscapeExtent.Min.X, InLandscapeExtent.Min.Y, 0.f);
-			FVector Translation = NewLandscapeTransform.TransformFVector4(OffsetVector);
-			NewLandscapeTransform.SetTranslation(Translation);
+			FTransform NewLandscapeTransform = LandscapeEditLayersLocals::GetLandscapeExtentTransform(*Landscape, InLandscapeExtent);
 			FIntPoint NewLandscapeSize = InLandscapeExtent.Max - InLandscapeExtent.Min;
 			if (!LandscapeTransform.Equals(NewLandscapeTransform) || (LandscapeSize != NewLandscapeSize) || LandscapeRenderTargetSize != NewLandscapeRenderTargetSize)
 			{
