Fix for spline merge not correcting the control point rotations for parent component transforms

diff --git a/Engine/Source/Runtime/Landscape/Private/LandscapeSplines.cpp b/Engine/Source/Runtime/Landscape/Private/LandscapeSplines.cpp
--- a/Engine/Source/Runtime/Landscape/Private/LandscapeSplines.cpp
+++ b/Engine/Source/Runtime/Landscape/Private/LandscapeSplines.cpp
@@ -3695,8 +3695,11 @@ void ULandscapeInfo::MoveControlPoint(ULandscapeSplineControlPoint* InControlPoi
 
 	ToSplineComponent->Modify();
 
-	const FTransform OldToNewTransform =
-		FromSplineComponent->GetComponentTransform().GetRelativeTransform(ToSplineComponent->GetComponentTransform());
+	const FTransform FromComponentTransform = FromSplineComponent->GetComponentTransform();	// (From-->World)
+	const FTransform ToComponentTransform = ToSplineComponent->GetComponentTransform();		// (To-->World)
+
+	const FTransform OldToNewTransform =		// (From-->To) == (From-->World) * (World-->To)  == From * Inverse(To)
+		FromComponentTransform.GetRelativeTransform(ToComponentTransform);
 		
 	// Delete all Mesh Components associated with the ControlPoint. (Will get recreated in UpdateSplinePoints)
 	if (InControlPoint->LocalMeshComponent)
@@ -3724,6 +3727,7 @@ void ULandscapeInfo::MoveControlPoint(ULandscapeSplineControlPoint* InControlPoi
 	ToSplineComponent->ControlPoints.Add(InControlPoint);
 
 	InControlPoint->Location = OldToNewTransform.TransformPosition(InControlPoint->Location);
+	InControlPoint->Rotation = FRotator(OldToNewTransform.TransformRotation(InControlPoint->Rotation.Quaternion()));
 
 	const bool bUpdateCollision = true; // default value
 	const bool bUpdateSegments = false; // done in next loop
