Fix for landscape save taking longer than it should when runtime gen is enabled

We turn off grass map distance eviction in editor, which prevents us from throwing away grass maps that were serialized in at load time.

diff --git a/Engine/Source/Runtime/Landscape/Private/LandscapeGrassMapsBuilder.cpp b/Engine/Source/Runtime/Landscape/Private/LandscapeGrassMapsBuilder.cpp
--- a/Engine/Source/Runtime/Landscape/Private/LandscapeGrassMapsBuilder.cpp
+++ b/Engine/Source/Runtime/Landscape/Private/LandscapeGrassMapsBuilder.cpp
@@ -312,7 +312,10 @@ FLandscapeGrassMapsBuilder::~FLandscapeGrassMapsBuilder()
 	{
 		TArray<FVector> EmptyCamerasArray;
 		int32 UpdateAllComponentCount = ComponentStates.Num();
-		UpdateTrackedComponents(EmptyCamerasArray, 0, UpdateAllComponentCount, /* bCancelAndEvictAllImmediately = */ true);
+
+		const bool bCancelAndEvictAllImmediately = true;
+		const bool bEvictWhenBeyondEvictionRange = false;
+		UpdateTrackedComponents(EmptyCamerasArray, 0, UpdateAllComponentCount, bCancelAndEvictAllImmediately, bEvictWhenBeyondEvictionRange);
 
 		ensure(NotReadyCount == 0 && StreamingCount == 0 && RenderingCount == 0 && AsyncFetchCount == 0 && PopulatedCount == 0);
 
@@ -387,7 +390,7 @@ void FAsyncFetchTask::DoWork()
 	Results = ActiveRender->FetchResults(bFreeAsyncReadback);
 }
 
-bool FLandscapeGrassMapsBuilder::UpdateTrackedComponents(const TArray<FVector>& Cameras, int32 LocalMaxRendering, int32 MaxExpensiveUpdateChecksToPerform, bool bCancelAndEvictAllImmediately)
+bool FLandscapeGrassMapsBuilder::UpdateTrackedComponents(const TArray<FVector>& Cameras, int32 LocalMaxRendering, int32 MaxExpensiveUpdateChecksToPerform, bool bCancelAndEvictAllImmediately, bool bEvictWhenBeyondEvictionRange)
 {
 	TRACE_CPUPROFILER_EVENT_SCOPE(FLandscapeGrassMapsBuilder::UpdateTrackedComponents);
 	SCOPE_CYCLE_COUNTER(STAT_UpdateComponentGrassMaps);
@@ -491,7 +494,7 @@ bool FLandscapeGrassMapsBuilder::UpdateTrackedComponents(const TArray<FVector>&
 					}
 
 					// check if the component is too far from the camera and we can reclaim the grass data
-					if (GGrassMapUseRuntimeGeneration &&
+					if (bEvictWhenBeyondEvictionRange &&
 						State->IsBeyondEvictionRange(Cameras))
 					{
 						GRASS_DEBUG_LOG(TEXT("Evicting for being beyond eviction range"));
@@ -589,7 +592,7 @@ bool FLandscapeGrassMapsBuilder::UpdateTrackedComponents(const TArray<FVector>&
 	return bChanged;
 }
 
-void FLandscapeGrassMapsBuilder::StartPrioritizedGrassMapGeneration(const TArray<FVector>& Cameras, int32 MaxComponentsToStart)
+void FLandscapeGrassMapsBuilder::StartPrioritizedGrassMapGeneration(const TArray<FVector>& Cameras, int32 MaxComponentsToStart, bool bOnlyWhenCloserThanEvictionRange)
 {
 	SCOPE_CYCLE_COUNTER(STAT_PrioritizePendingGrassMaps);
 
@@ -644,8 +647,8 @@ void FLandscapeGrassMapsBuilder::StartPrioritizedGrassMapGeneration(const TArray
 		const ULandscapeComponent* Component = State->Component;
 		check(State->Stage == EComponentStage::Pending);
 
-		// runtime generation doesn't generate past this distance
-		if (GGrassMapUseRuntimeGeneration)
+		// if threshold is enabled, check if the nearest pending component is close enough
+		if (bOnlyWhenCloserThanEvictionRange)
 		{
 			const double MaxSpawnDistance = Component->GrassTypeSummary.MaxInstanceDiscardDistance * MustHaveDistanceScale;
 			if (Pending.PriorityKey > MaxSpawnDistance * MaxSpawnDistance)
@@ -753,15 +756,23 @@ void FLandscapeGrassMapsBuilder::AmortizedUpdateGrassMaps(
 		AmortizedMaxRendering *= GGrassMapPrioritizedMultiplier;
 	}
 
+#if WITH_EDITOR
+	const bool bCancelAndEvictAllImmediately = false;
+	const bool bEvictWhenBeyondEvictionRange = false;
+#else
 	const bool bCancelAndEvictAllImmediately = !GGrassEnable;
-	UpdateTrackedComponents(Cameras, AmortizedMaxRendering, GGrassMapMaxDiscardChecksPerFrame, bCancelAndEvictAllImmediately);
+	const bool bEvictWhenBeyondEvictionRange = !!GGrassMapUseRuntimeGeneration;
+#endif // WITH_EDITOR
+	UpdateTrackedComponents(Cameras, AmortizedMaxRendering, GGrassMapMaxDiscardChecksPerFrame, bCancelAndEvictAllImmediately, bEvictWhenBeyondEvictionRange);
 
 	// no point in looking to start new grass map generation if nothing is pending, if grass is disabled or there are no cameras
 	if (bAllowStartGrassMapGeneration && PendingCount > 0 && GGrassEnable && Cameras.Num() > 0)
 	{
 		// check our pipeline limits to make sure we have room to start components
 		const int32 AvailableStreamingSlots = AmortizedMaxStreaming - StreamingCount;
-		StartPrioritizedGrassMapGeneration(Cameras, AvailableStreamingSlots);
+		// do not build grass maps beyond eviction range if runtime generation is enabled
+		const bool bOnlyWhenCloserThanEvictionRange = !!GGrassMapUseRuntimeGeneration;
+		StartPrioritizedGrassMapGeneration(Cameras, AvailableStreamingSlots, bOnlyWhenCloserThanEvictionRange);
 	}
 }
 
@@ -813,7 +824,10 @@ bool FLandscapeGrassMapsBuilder::BuildGrassMapsNowForComponents(
 
 		// update all components that are tracked (without evicting)
 		const int32 UpdateAllComponentCount = ComponentStates.Num();
-		bool bChanged = UpdateTrackedComponents(EmptyCamerasArray, MaxStreamingRendering, UpdateAllComponentCount, /* bCancelAndEvictAllImmediately= */ false);
+
+		const bool bCancelAndEvictAllImmediately = false;
+		const bool bEvictWhenBeyondEvictionRange = false;
+		bool bChanged = UpdateTrackedComponents(EmptyCamerasArray, MaxStreamingRendering, UpdateAllComponentCount, bCancelAndEvictAllImmediately, bEvictWhenBeyondEvictionRange);
 
 		UpToDateCount = 0;
 		int32 AvailableStreamingSlots = MaxStreamingRendering - StreamingCount; // here we don't limit by overall population count
diff --git a/Engine/Source/Runtime/Landscape/Private/LandscapeGrassMapsBuilder.h b/Engine/Source/Runtime/Landscape/Private/LandscapeGrassMapsBuilder.h
--- a/Engine/Source/Runtime/Landscape/Private/LandscapeGrassMapsBuilder.h
+++ b/Engine/Source/Runtime/Landscape/Private/LandscapeGrassMapsBuilder.h
@@ -77,14 +77,16 @@ private:
 	// false if the world can not currently render the grass (but this may change later, for example if preview modes are modified)
 	bool CanCurrentlyRender() const;
 
-	// Update all of the non-pending components.
+	// Update all of the non-pending components.  Returns true if any components changed states.
 	// If passed an empty Cameras array, distances are calculated as zero (i.e. it won't evict for distance)
-	// returns true if any components changed states
-	bool UpdateTrackedComponents(const TArray<FVector>& Cameras, int32 LocalMaxRendering, int32 MaxExpensiveUpdateChecksToPerform, bool bCancelAndEvictAll);
+	// MaxExpensiveUpdateChecksToPerform controls how many expensive component updates are performed (pass ComponentStates.Num() to process all)
+	// bCancelAndEvictAllImmediately will move all tracked component states to the pending state and block until any in-flight processing is cancelled and cleaned up
+	// bEvictWhenBeyondEvictionRange will evict any populated components that are beyond the eviction range distance (used to reclaim memory when using runtime generation)
+	bool UpdateTrackedComponents(const TArray<FVector>& Cameras, int32 LocalMaxRendering, int32 MaxExpensiveUpdateChecksToPerform, bool bCancelAndEvictAllImmediately, bool bEvictWhenBeyondEvictionRange);
 
 	// Start the grass map generation process on pending components in priority order
 	// (based on distance from the given Camera set) -- Cameras must not be empty.
-	void StartPrioritizedGrassMapGeneration(const TArray<FVector>& Cameras, int32 MaxComponentsToStart);
+	void StartPrioritizedGrassMapGeneration(const TArray<FVector>& Cameras, int32 MaxComponentsToStart, bool bOnlyWhenCloserThanEvictionRange);
 	
 	enum class EComponentStage
 	{
