Fix Substrate shader compilation with terrain layer blend node and material attributes.

diff --git a/Engine/Source/Runtime/Engine/Public/Materials/MaterialExpressionShadingModel.h b/Engine/Source/Runtime/Engine/Public/Materials/MaterialExpressionShadingModel.h
--- a/Engine/Source/Runtime/Engine/Public/Materials/MaterialExpressionShadingModel.h
+++ b/Engine/Source/Runtime/Engine/Public/Materials/MaterialExpressionShadingModel.h
@@ -18,7 +18,7 @@ class FMaterialCompiler;
  * @param Alpha					Bland factor [0..1]
  * @return						Index to a new code chunk
  */
-extern int32 CompileShadingModelBlendFunction(FMaterialCompiler* Compiler, int32 A, int32 B, int32 Alpha);
+ENGINE_API extern int32 CompileShadingModelBlendFunction(FMaterialCompiler* Compiler, int32 A, int32 B, int32 Alpha);
 
 UCLASS(collapsecategories, hidecategories = Object, MinimalAPI)
 class UMaterialExpressionShadingModel : public UMaterialExpression
diff --git a/Engine/Source/Runtime/Landscape/Private/Materials/MaterialExpressionLandscapeLayerBlend.cpp b/Engine/Source/Runtime/Landscape/Private/Materials/MaterialExpressionLandscapeLayerBlend.cpp
--- a/Engine/Source/Runtime/Landscape/Private/Materials/MaterialExpressionLandscapeLayerBlend.cpp
+++ b/Engine/Source/Runtime/Landscape/Private/Materials/MaterialExpressionLandscapeLayerBlend.cpp
@@ -7,6 +7,8 @@
 #include "MaterialCompiler.h"
 #include "DataDrivenShaderPlatformInfo.h"
 #include "LandscapeUtils.h"
+#include "Materials/MaterialAttributeDefinitionMap.h"
+#include "Materials/MaterialExpressionShadingModel.h"
 
 #if WITH_EDITOR
 #include "MaterialGraph/MaterialGraphNode.h"
@@ -205,7 +207,18 @@ int32 UMaterialExpressionLandscapeLayerBlend::Compile(class FMaterialCompiler* C
 
 	int32 InvWeightSumCode = Compiler->Div(Compiler->Constant(1.f), WeightSumCode);
 
-	int32 OutputCode = Compiler->Constant(0);
+	// If the output has a ShadingModel type, regular arithmetic nodes (weight/blend) needs to be handled differently
+	int32 OutputCode = INDEX_NONE;
+	bool bIsShadingModel = false;
+	if (FMaterialAttributeDefinitionMap::GetProperty(Compiler->GetMaterialAttribute()) == MP_ShadingModel)
+	{
+		OutputCode = Compiler->ShadingModel(EMaterialShadingModel::MSM_DefaultLit);
+		bIsShadingModel = true;
+	}
+	else
+	{
+		OutputCode = Compiler->Constant(0);
+	}
 
 	for (int32 LayerIdx = 0; LayerIdx<Layers.Num(); LayerIdx++)
 	{
@@ -220,7 +233,11 @@ int32 UMaterialExpressionLandscapeLayerBlend::Compile(class FMaterialCompiler* C
 															  static_cast<float>(Layer.ConstLayerInput.Y),
 				                                              static_cast<float>(Layer.ConstLayerInput.Z));
 
-			if (bNeedsRenormalize)
+			if (bIsShadingModel && Compiler->GetType(LayerCode) == EMaterialValueType::MCT_ShadingModel)
+			{
+				OutputCode = CompileShadingModelBlendFunction(Compiler, OutputCode, LayerCode, WeightCodes[LayerIdx]);
+			}
+			else if (bNeedsRenormalize)
 			{
 				// Renormalize the weights as our height modification has made them non-uniform
 				OutputCode = Compiler->Add(OutputCode, Compiler->Mul(LayerCode, Compiler->Mul(InvWeightSumCode, WeightCodes[LayerIdx])));
@@ -247,8 +264,16 @@ int32 UMaterialExpressionLandscapeLayerBlend::Compile(class FMaterialCompiler* C
 					                        : Compiler->Constant3(static_cast<float>(Layer.ConstLayerInput.X),
 					                                              static_cast<float>(Layer.ConstLayerInput.Y),
 					                                              static_cast<float>(Layer.ConstLayerInput.Z));
+				
 				// Blend in the layer using the alpha value
-				OutputCode = Compiler->Lerp(OutputCode, LayerCode, WeightCode);
+				if (bIsShadingModel && Compiler->GetType(LayerCode) == EMaterialValueType::MCT_ShadingModel)
+				{
+					OutputCode = CompileShadingModelBlendFunction(Compiler, OutputCode, LayerCode, WeightCode);
+				}
+				else
+				{
+					OutputCode = Compiler->Lerp(OutputCode, LayerCode, WeightCode);
+				}
 			}
 		}
 	}
