Don't display landscape on-screen warnings while the landscape is being edited to avoid flickering of the message while stuff get updated in the background

diff --git a/Engine/Source/Runtime/Landscape/Private/LandscapeSubsystem.cpp b/Engine/Source/Runtime/Landscape/Private/LandscapeSubsystem.cpp
--- a/Engine/Source/Runtime/Landscape/Private/LandscapeSubsystem.cpp
+++ b/Engine/Source/Runtime/Landscape/Private/LandscapeSubsystem.cpp
@@ -911,17 +911,23 @@ bool ULandscapeSubsystem::GetActionableMessage(FActionableMessage& OutActionable
 	EOutdatedDataFlags OutdatedFlagsUnion = EOutdatedDataFlags::None;
 	Algo::ForEach(OutdatedProxies, [&OutdatedProxies, &NumOutdatedProxyPerFlag, &NumTotalOutdatedProxy, &OutdatedFlagsUnion](const TTuple<ALandscapeProxy*, EOutdatedDataFlags>& ProxyAndFlag)
 	{
-		++NumTotalOutdatedProxy;
-		OutdatedFlagsUnion |= ProxyAndFlag.Value;
+		ALandscape* ParentLandscape = ProxyAndFlag.Key->GetLandscapeActor();
+		// Don't display any message for this landscape when it's being edited : we consider the landscape to be in "WIP state" while editing. 
+		//  This avoids flickering of the message while the async stuff (grass, Nanite, ...) gets updated in the background
+		if ((ParentLandscape == nullptr) || !ParentLandscape->HasLandscapeEdMode())
+		{ 
+			++NumTotalOutdatedProxy;
+			OutdatedFlagsUnion |= ProxyAndFlag.Value;
 
-		uint32 RemainingFlags = static_cast<uint32>(ProxyAndFlag.Value);
-		while (RemainingFlags != 0)
-		{
-			uint32 FlagIndex = FBitSet::GetAndClearNextBit(RemainingFlags);
-			++NumOutdatedProxyPerFlag[FlagIndex];
+			uint32 RemainingFlags = static_cast<uint32>(ProxyAndFlag.Value);
+			while (RemainingFlags != 0)
+			{
+				uint32 FlagIndex = FBitSet::GetAndClearNextBit(RemainingFlags);
+				++NumOutdatedProxyPerFlag[FlagIndex];
+			}
 		}
 	});
-
+	
 	// If more than 1 action is required, go with a BuildAll action
 	if (FMath::CountBits(static_cast<uint64>(OutdatedFlagsUnion)) > 1)
 	{
