Fixed incorrect ensure triggering when loading old landscape levels that have landscape edit layers textures with more than a single mip. This ensure is only valid if CVar landscape.StripLayerMipsOnLoad is active (it's currently not)

diff --git a/Engine/Source/Runtime/Landscape/Private/Landscape.cpp b/Engine/Source/Runtime/Landscape/Private/Landscape.cpp
--- a/Engine/Source/Runtime/Landscape/Private/Landscape.cpp
+++ b/Engine/Source/Runtime/Landscape/Private/Landscape.cpp
@@ -208,7 +208,7 @@ static FAutoConsoleVariable CVarLandscapeSupressMapCheckWarnings_Nanite(
 	false,
 	TEXT("Issue MapCheck Info messages instead of warnings if Nanite Data is out of date"));
 
-static FAutoConsoleVariable CVarStripLayerTextureMipsOnLoad(
+FAutoConsoleVariable CVarStripLayerTextureMipsOnLoad(
 	TEXT("landscape.StripLayerMipsOnLoad"),
 	false,
 	TEXT("Remove (on load) the mip chain from textures used in layers which don't require them"));
diff --git a/Engine/Source/Runtime/Landscape/Private/LandscapeEdit.cpp b/Engine/Source/Runtime/Landscape/Private/LandscapeEdit.cpp
--- a/Engine/Source/Runtime/Landscape/Private/LandscapeEdit.cpp
+++ b/Engine/Source/Runtime/Landscape/Private/LandscapeEdit.cpp
@@ -109,6 +109,8 @@ static TAutoConsoleVariable<int32> CVarLandscapeApplyPhysicalMaterialChangesImme
 	1,
     TEXT("Applies physical material task changes immediately rather than during the next cook/PIE."));
 
+extern FAutoConsoleVariable CVarStripLayerTextureMipsOnLoad;
+
 #if WITH_EDITOR
 
 // Used to temporarily disable material instance updates (typically used for cases where multiple updates are called on sample component)
@@ -7205,7 +7207,7 @@ void ULandscapeComponent::ReallocateWeightmapsInternal(FLandscapeEditDataInterfa
 			}
 		}
 
-		if (CurrentWeightmapTexture != nullptr)
+		if ((CurrentWeightmapTexture != nullptr) && CVarStripLayerTextureMipsOnLoad->GetBool())
 		{
 			// sanity check - only final weightmaps are allowed more than one mip
 			ensure((CurrentWeightmapTexture->Source.GetNumMips() == 1) || bIsFinalWeightmap);
