* Fixed deprecation of target layers in landscape : the previous code did not dirty the main landscape actor, where the source list of target layers lies (all of the proxies' lists are a copy of the source list, because of the LandscapeInherited behavior), so all proxies had inconsistent list against the source landscape actor, and since deprecation had already been performed, they would not re-run the deprecation code that analyzes weightmap allocations and deduces the layer info pointers to update the main landscape actor's target layers list upon reload, effectively leaving null entries for layer info everywhere in the source landscape actor upon reload. The fix involves going through all proxies and resynchronizing landscape shared properties any time a proxy discovers an unregistered landscape layer info in one of its layer allocations, because this updates the TargetLayers list in the source actor, and this list needs to be re-synchronized with all proxies. Fortunately, this is a one-time thing (this deprecation step won't be performed afterwards) and only necessary when going from a target layer with a null layer info to a non-null one (i.e. once per target layer)
* Another issue was discovered when re-assigning a null target layer of an existing target layer : this would not call Modify on the main landscape actor, which 1) incorrectly prevents the operation from being undone 2) also prevents the actor to be dirty, and therefore, upon reload, the landscape would not retain the proper layer info assignment
* Also fixed a bunch of places where Modify was called before changing landscape properties

diff --git a/Engine/Source/Runtime/Core/Public/UObject/FortniteMainBranchObjectVersion.h b/Engine/Source/Runtime/Core/Public/UObject/FortniteMainBranchObjectVersion.h
--- a/Engine/Source/Runtime/Core/Public/UObject/FortniteMainBranchObjectVersion.h
+++ b/Engine/Source/Runtime/Core/Public/UObject/FortniteMainBranchObjectVersion.h
@@ -483,6 +483,9 @@ struct FFortniteMainBranchObjectVersion
 		// Serialize EditorBounds in WorldPartitionActorDesc
 		WorldPartitionActorDescSerializeEditorBounds,
 
+		// Fixup for the data that has been damaged by LandscapeTargetLayersInLandscapeActor (loss of landscape layer info object assignments)
+		FixupLandscapeTargetLayersInLandscapeActor,
+
 		// -----<new versions can be added above this line>-------------------------------------------------
 		VersionPlusOne,
 		LatestVersion = VersionPlusOne - 1
diff --git a/Engine/Source/Runtime/Landscape/Classes/LandscapeInfo.h b/Engine/Source/Runtime/Landscape/Classes/LandscapeInfo.h
--- a/Engine/Source/Runtime/Landscape/Classes/LandscapeInfo.h
+++ b/Engine/Source/Runtime/Landscape/Classes/LandscapeInfo.h
@@ -458,7 +458,7 @@ private:
 	bool ApplySplinesInternal(bool bOnlySelected, TScriptInterface<ILandscapeSplineInterface> SplineOwner, TSet<TObjectPtr<ULandscapeComponent>>* OutModifiedComponents, bool bMarkPackageDirty, int32 LandscapeMinX, int32 LandscapeMinY, int32 LandscapeMaxX, int32 LandscapeMaxY, TFunctionRef<TSharedPtr<FModulateAlpha>(ULandscapeLayerInfoObject*)> GetOrCreateModulate);
 	void MoveSegment(ULandscapeSplineSegment* InSegment, TScriptInterface<ILandscapeSplineInterface> From, TScriptInterface<ILandscapeSplineInterface> To);
 	void MoveControlPoint(ULandscapeSplineControlPoint* InControlPoint, TScriptInterface<ILandscapeSplineInterface> From, TScriptInterface<ILandscapeSplineInterface> To);
-	void UpdateLayerInfoMapInternal(ALandscapeProxy* Proxy);
+	bool UpdateLayerInfoMapInternal(ALandscapeProxy* Proxy);
 	bool TryAddToModifiedPackages(UPackage* InPackage, const ALandscape* InLandscapeOverride = nullptr);
 #endif
 
diff --git a/Engine/Source/Runtime/Landscape/Private/Landscape.cpp b/Engine/Source/Runtime/Landscape/Private/Landscape.cpp
--- a/Engine/Source/Runtime/Landscape/Private/Landscape.cpp
+++ b/Engine/Source/Runtime/Landscape/Private/Landscape.cpp
@@ -3592,11 +3592,11 @@ void ALandscape::EnableNaniteSkirts(bool bInEnable, float InSkirtDepth, bool bIn
 		{
 			if (Proxy != nullptr)
 			{
+				Proxy->Modify(bInShouldDirtyPackage);
 				Proxy->SynchronizeSharedProperties(this);
 				Proxy->InvalidateOrUpdateNaniteRepresentation(/*bInCheckContentId*/true, /*InTargetPlatform*/nullptr);
 				Proxy->UpdateRenderingMethod();
 				Proxy->MarkComponentsRenderStateDirty();
-				Proxy->Modify(bInShouldDirtyPackage);
 			}
 			return true;
 		});
@@ -3856,10 +3856,15 @@ void ULandscapeInfo::CreateTargetLayerSettingsFor(ULandscapeLayerInfoObject* Lay
 {
 	ForEachLandscapeProxy([LayerInfo](ALandscapeProxy* Proxy)
 	{
-		if (!Proxy->HasTargetLayer(LayerInfo->LayerName))
+		if (Proxy->HasTargetLayer(LayerInfo->LayerName))
+		{
+			Proxy->UpdateTargetLayer(LayerInfo->LayerName, FLandscapeTargetLayerSettings(LayerInfo));
+		}
+		else
 		{
 			Proxy->AddTargetLayer(LayerInfo->LayerName, FLandscapeTargetLayerSettings(LayerInfo));
 		}
+
 		return true;
 	});
 }
@@ -3907,41 +3912,49 @@ int32 ULandscapeInfo::GetLayerInfoIndex(FName LayerName, ALandscapeProxy* Owner
 }
 
 
-void ULandscapeInfo::UpdateLayerInfoMapInternal(ALandscapeProxy* Proxy)
+bool ULandscapeInfo::UpdateLayerInfoMapInternal(ALandscapeProxy* Proxy)
 {
 	TRACE_CPUPROFILER_EVENT_SCOPE(ULandscapeInfo::UpdateLayerInfoMapInternal);
 
+	bool bLayerInfoMapChanged = false;
 	if( !LandscapeActor.IsValid())
 	{
-		return;
+		return false;
 	}
 
+	// Perform a delayed (see where bInitWeightLayersFromMaterial is assigned), one-time deprecation of the landscape layer data based on the content of the materials in the components (see FixupLandscapeTargetLayersInLandscapeActor)
 	if (GIsEditor && Proxy && Proxy->bInitWeightLayersFromMaterial)
 	{
 		Proxy->bInitWeightLayersFromMaterial = false;
 		// add the landscape layers from the material on the proxy & LandscapeLayerInfos from the WeightMap Allocations
-		if (LandscapeActor.IsValid())
+		TArray<FName> LayerNames = Proxy->RetrieveAllLayerNamesFromMaterials();
+		const TMap<FName, ULandscapeLayerInfoObject*> InfoObjects = Proxy->RetrieveAllocationInfos();
+		for(const FName& LayerName : LayerNames)
 		{
-			TArray<FName> LayerNames = Proxy->RetrieveAllLayerNamesFromMaterials();
-			const TMap<FName, ULandscapeLayerInfoObject*> InfoObjects = Proxy->RetrieveAllocationInfos();
-			for(const FName& LayerName : LayerNames)
+			ULandscapeLayerInfoObject* const * LayerInfoObject = InfoObjects.Find(LayerName);
+			if (!LandscapeActor->HasTargetLayer(LayerName))
 			{
-				ULandscapeLayerInfoObject* const * LayerInfoObject = InfoObjects.Find(LayerName);
-				if (!LandscapeActor->HasTargetLayer(LayerName))
-				{
-					LandscapeActor->AddTargetLayer(LayerName, FLandscapeTargetLayerSettings(LayerInfoObject ? *LayerInfoObject : nullptr), false);
-				}
+				// Mark the parent landscape actor dirty with bInForceResave == true so that the parent actor is put into the list of files to save even if we do this fixup on load :
+				MarkObjectDirty(/*InObject = */LandscapeActor.Get(), /*bInForceResave = */true);
+
+				LandscapeActor->AddTargetLayer(LayerName, FLandscapeTargetLayerSettings(LayerInfoObject ? *LayerInfoObject : nullptr), false);
+				bLayerInfoMapChanged = true;
 			}
+		}
 			
-			for (const auto& Layer : LandscapeActor->GetTargetLayers())
+		for (const auto& Layer : LandscapeActor->GetTargetLayers())
+		{
+			// Fixup the layers that point to an invalid layer info found in the proxy if possible
+			if (Layer.Value.LayerInfoObj == nullptr)
 			{
-				if (!Layer.Value.LayerInfoObj)
+				ULandscapeLayerInfoObject* const * LayerInfoObject = InfoObjects.Find(Layer.Key);
+				if (LayerInfoObject && LandscapeActor->HasTargetLayer(Layer.Key))
 				{
-					ULandscapeLayerInfoObject* const * LayerInfoObject = InfoObjects.Find(Layer.Key);
-					if (LayerInfoObject && LandscapeActor->HasTargetLayer(Layer.Key))
-					{
-						LandscapeActor->UpdateTargetLayer(Layer.Key, FLandscapeTargetLayerSettings(*LayerInfoObject), false);
-					}
+					// Mark the parent landscape actor dirty with bInForceResave == true so that the parent actor is put into the list of files to save even if we do this fixup on load :
+					MarkObjectDirty(/*InObject = */LandscapeActor.Get(), /*bInForceResave = */true);
+
+					LandscapeActor->UpdateTargetLayer(Layer.Key, FLandscapeTargetLayerSettings(*LayerInfoObject), false);
+					bLayerInfoMapChanged = true;
 				}
 			}
 		}
@@ -3967,11 +3980,13 @@ void ULandscapeInfo::UpdateLayerInfoMapInternal(ALandscapeProxy* Proxy)
 			Layers[LayerInfoIndex].LayerInfoObj = ALandscapeProxy::VisibilityLayer;
 		}
 	}
+
+	return bLayerInfoMapChanged;
 }
 
 bool ULandscapeInfo::UpdateLayerInfoMap(ALandscapeProxy* Proxy /*= nullptr*/, bool bInvalidate /*= false*/)
 {
-	UpdateLayerInfoMapInternal(Proxy);
+	bool bLayerInfoMapChanged = UpdateLayerInfoMapInternal(Proxy);
 
 	if (GIsEditor)
 	{
@@ -3981,8 +3996,7 @@ bool ULandscapeInfo::UpdateLayerInfoMap(ALandscapeProxy* Proxy /*= nullptr*/, bo
 			Landscape->RequestLayersInitialization(/*bInRequestContentUpdate*/false);
 		}
 	}
-	// return value is unused and has no meaning.
-	return true;
+	return bLayerInfoMapChanged;
 
 }
 
@@ -4101,52 +4115,54 @@ void ALandscapeProxy::PostLoad()
 
 #if WITH_EDITOR
 	const int32 LinkerVersion = GetLinkerCustomVersion(FFortniteMainBranchObjectVersion::GUID);
-	bInitWeightLayersFromMaterial = LinkerVersion < FFortniteMainBranchObjectVersion::LandscapeTargetLayersInLandscapeActor;
+	// With and before LandscapeTargetLayersInLandscapeActor and until FixupLandscapeTargetLayersInLandscapeActor, some layer info objects have been incorrectly unassigned and we now have to go through 
+	//  all materials of all streaming proxies to gather the missing landscape layer info objects and this can only be done after proxies are united with their parent landscape by their ULandscapeInfo 
+	//  (since TargetLayers is a LandscapeInherited property of the parent landscape, propagated to the child proxies), so we set a boolean to delay this operation until then :
+	bInitWeightLayersFromMaterial = (LinkerVersion < FFortniteMainBranchObjectVersion::FixupLandscapeTargetLayersInLandscapeActor);
 	
-	PRAGMA_DISABLE_DEPRECATION_WARNINGS
-		if (!LandscapeMaterialsOverride_DEPRECATED.IsEmpty())
+PRAGMA_DISABLE_DEPRECATION_WARNINGS
+	if (!LandscapeMaterialsOverride_DEPRECATED.IsEmpty())
+	{
+		PerLODOverrideMaterials.Reserve(LandscapeMaterialsOverride_DEPRECATED.Num());
+		for (const FLandscapeProxyMaterialOverride& LocalMaterialOverride : LandscapeMaterialsOverride_DEPRECATED)
 		{
-			PerLODOverrideMaterials.Reserve(LandscapeMaterialsOverride_DEPRECATED.Num());
-			for (const FLandscapeProxyMaterialOverride& LocalMaterialOverride : LandscapeMaterialsOverride_DEPRECATED)
-			{
-				PerLODOverrideMaterials.Add({ LocalMaterialOverride.LODIndex.Default, LocalMaterialOverride.Material });
-			}
-			LandscapeMaterialsOverride_DEPRECATED.Reset();
+			PerLODOverrideMaterials.Add({ LocalMaterialOverride.LODIndex.Default, LocalMaterialOverride.Material });
 		}
-	PRAGMA_ENABLE_DEPRECATION_WARNINGS
+		LandscapeMaterialsOverride_DEPRECATED.Reset();
+	}
+PRAGMA_ENABLE_DEPRECATION_WARNINGS
 
-		if (GIsEditor)
+	if (GIsEditor)
+	{
+		// We may not have run PostLoad on LandscapeComponents yet
+		for (TObjectPtr<ULandscapeComponent>& LandscapeComponent : LandscapeComponents)
 		{
-			// We may not have run PostLoad on LandscapeComponents yet
-			for (TObjectPtr<ULandscapeComponent>& LandscapeComponent : LandscapeComponents)
+			if (LandscapeComponent)
 			{
-				if (LandscapeComponent)
-				{
-					LandscapeComponent->ConditionalPostLoad();
-				}
+				LandscapeComponent->ConditionalPostLoad();
 			}
+		}
 
-			// We may not have run PostLoad on CollisionComponent yet
-			for (TObjectPtr<ULandscapeHeightfieldCollisionComponent>& CollisionComponent : CollisionComponents)
+		// We may not have run PostLoad on CollisionComponent yet
+		for (TObjectPtr<ULandscapeHeightfieldCollisionComponent>& CollisionComponent : CollisionComponents)
+		{
+			if (CollisionComponent)
 			{
-				if (CollisionComponent)
-				{
-					CollisionComponent->ConditionalPostLoad();
-				}
+				CollisionComponent->ConditionalPostLoad();
 			}
+		}
 
-			if ((GetLinker() && (GetLinker()->UEVer() < VER_UE4_LANDSCAPE_COMPONENT_LAZY_REFERENCES)) ||
-				LandscapeComponents.Num() != CollisionComponents.Num() ||
-				LandscapeComponents.ContainsByPredicate([](ULandscapeComponent* Comp) { return ((Comp != nullptr) && (Comp->GetCollisionComponent() == nullptr)); }) ||
-				CollisionComponents.ContainsByPredicate([](ULandscapeHeightfieldCollisionComponent* Comp) { return ((Comp != nullptr) && (Comp->GetRenderComponent() == nullptr)); }))
-			{
-				// Need to clean up invalid collision and render components
-				RecreateCollisionComponents();
-			}
+		if ((GetLinker() && (GetLinker()->UEVer() < VER_UE4_LANDSCAPE_COMPONENT_LAZY_REFERENCES)) ||
+			LandscapeComponents.Num() != CollisionComponents.Num() ||
+			LandscapeComponents.ContainsByPredicate([](ULandscapeComponent* Comp) { return ((Comp != nullptr) && (Comp->GetCollisionComponent() == nullptr)); }) ||
+			CollisionComponents.ContainsByPredicate([](ULandscapeHeightfieldCollisionComponent* Comp) { return ((Comp != nullptr) && (Comp->GetRenderComponent() == nullptr)); }))
+		{
+			// Need to clean up invalid collision and render components
+			RecreateCollisionComponents();
 		}
+	}
 	
 PRAGMA_DISABLE_DEPRECATION_WARNINGS
-
 	if (EditorCachedLayerInfos_DEPRECATED.Num() > 0)
 	{
 		for (int32 i = 0; i < EditorCachedLayerInfos_DEPRECATED.Num(); i++)
@@ -4157,9 +4173,7 @@ PRAGMA_DISABLE_DEPRECATION_WARNINGS
 	}
 PRAGMA_ENABLE_DEPRECATION_WARNINGS
 
-		bool bFixedUpInvalidMaterialInstances = false;
-
-
+	bool bFixedUpInvalidMaterialInstances = false;
 	for (ULandscapeComponent* Comp : LandscapeComponents)
 	{
 		if (Comp == nullptr)
@@ -4235,7 +4249,7 @@ PRAGMA_ENABLE_DEPRECATION_WARNINGS
 
 		FMessageLog MessageLog("MapCheck");
 		CreateMapCheckMessage(MessageLog)
-			->AddToken(FTextToken::Create(FText::Format(LOCTEXT("MapCheck_Message_LandscapeRebuildNanite", "{LandscapeProxyName} : Landscape Nanite is enabled but saved mesh data is out of date. "), Arguments)))
+			->AddToken(FTextToken::Create(FText::Format(LOCTEXT("MapCheck_Message_LandscapeRebuildNanite", "{LandscapeProxyName} : Landscape Nanite is enabled but the saved mesh data is out of date. "), Arguments)))
 			->AddToken(FActionToken::Create(LOCTEXT("MapCheck_SaveFixedUpData", "Save Modified Landscapes"), LOCTEXT("MapCheck_SaveFixedUpData_Desc", "Saves the modified landscape proxy actors"),
 				FOnActionTokenExecuted::CreateLambda([WeakLandscapeProxy]()
 		{
@@ -4789,7 +4803,7 @@ void ALandscapeProxy::FixupSharedData(ALandscape* Landscape, const bool bMapChec
 			bool bNeedsManualResave = LandscapeInfo->MarkObjectDirty(/*InObject = */this, /*bInForceResave = */true, LandscapeActor);
 
 			if (bMapCheck && bNeedsManualResave)
-	{
+			{
 				UE::Landscape::Private::DisplaySynchronizedPropertiesMapcheckWarning(SynchronizedProperties, *this, *Landscape);
 			}
 		}
@@ -5437,13 +5451,20 @@ void ULandscapeInfo::RegisterLandscapeActorWithProxyInternal(ALandscapeProxy* Pr
 #if WITH_EDITOR
 		StreamingProxy->FixupSharedData(Landscape, bMapCheck);
 #endif // WITH_EDITOR
-
-		StreamingProxy->SetLODGroupKeyInternal(Landscape->LODGroupKey);
 	}
 
 #if WITH_EDITOR
 	// generic proxy setup (that requires ALandscape actor) here
-	UpdateLayerInfoMap(Proxy);
+	if (bool bLayerInfoMapChanged = UpdateLayerInfoMap(Proxy))
+	{
+		// The layer info map is part of the main landscape so if it has changed, we needto do another round of shared data fixup on all proxies, so all proxies have their TargetLayers list synchronized. 
+		//  This is a one-time thing because at some point during development, the target layer data was deprecated and the deprecation turned somewhat sour :S
+		ForEachLandscapeProxy([this, Landscape, bMapCheck](ALandscapeProxy* Proxy)
+		{
+			Proxy->FixupSharedData(Landscape, bMapCheck);
+			return true;
+		});
+	}
 
 	if (GIsEditor)
 	{
@@ -5474,9 +5495,6 @@ void ULandscapeInfo::RegisterLandscapeActorWithProxyInternal(ALandscapeProxy* Pr
 		}
 	}
 #endif // WITH_EDITOR
-
-
-
 }
 
 void ULandscapeInfo::RegisterActor(ALandscapeProxy* Proxy, bool bMapCheck, bool bUpdateAllAddCollisions)
@@ -6470,7 +6488,7 @@ bool ALandscapeProxy::HasLayer(ULandscapeLayerInfoObject* LayerInfoObject) const
 	return TargetLayers.FindKey(FLandscapeTargetLayerSettings(LayerInfoObject)) == nullptr;
 }
 
-bool  ALandscapeProxy::RemoveTargetLayer(const FName& Name,  bool bPostEditChange)
+bool ALandscapeProxy::RemoveTargetLayer(const FName& Name,  bool bPostEditChange)
 {
 	Modify();
 	
@@ -6520,6 +6538,7 @@ bool ALandscapeProxy::UpdateTargetLayer(const FName& Name, const FLandscapeTarge
 	check(TargetLayerSettings);
 	if (TargetLayerSettings)
 	{
+		Modify();
 		*TargetLayerSettings = InTargetLayerSettings;
 		
 		if (FProperty* Property = StaticClass()->FindPropertyByName(GET_MEMBER_NAME_CHECKED(ALandscapeProxy, TargetLayers));  Property != nullptr && bPostEditChange)
@@ -6528,7 +6547,6 @@ bool ALandscapeProxy::UpdateTargetLayer(const FName& Name, const FLandscapeTarge
 			PostEditChangeProperty(PropertyChangedEvent);
 		}
 
-		Modify();
 		return true;
 	}
 
