Fix for rare runtime crashes in FAsyncGrassBuilder at world shutdown

diff --git a/Engine/Source/Runtime/Landscape/Private/Landscape.cpp b/Engine/Source/Runtime/Landscape/Private/Landscape.cpp
--- a/Engine/Source/Runtime/Landscape/Private/Landscape.cpp
+++ b/Engine/Source/Runtime/Landscape/Private/Landscape.cpp
@@ -272,7 +272,7 @@ ULandscapeComponent::ULandscapeComponent(const FObjectInitializer& ObjectInitial
 	, SplineHash(0)
 	, PhysicalMaterialHash(0)
 #endif
-	, GrassData(MakeShareable(new FLandscapeComponentGrassData()))
+	, GrassData(MakeShared<FLandscapeComponentGrassData>())
 	, ChangeTag(0)
 {
 	SetCollisionProfileName(UCollisionProfile::NoCollision_ProfileName);
@@ -967,8 +967,10 @@ PRAGMA_ENABLE_DEPRECATION_WARNINGS
 			}
 			else
 			{
-			Ar << GrassData.Get();
-		}
+				// technically on load this is doing a thread-unsafe operation by stomping the data in the existing ref
+				// but we're assuming there are no async threads using this pointer yet at load...
+				Ar << GrassData.Get();
+			}
 		}
 
 		// When loading or saving a component, validate that grass data is valid : 
@@ -4106,8 +4108,7 @@ void ALandscapeProxy::PostLoad()
 					*GetName());
 
 				// Free the memory, so at least we will save the space at runtime.
-				TUniquePtr<FLandscapeComponentGrassData> NewGrassData = MakeUnique<FLandscapeComponentGrassData>();
-				Comp->GrassData = MakeShareable(NewGrassData.Release());
+				Comp->GrassData = MakeShared<FLandscapeComponentGrassData>();
 			}
 		}
 #endif // !WITH_EDITOR
diff --git a/Engine/Source/Runtime/Landscape/Private/LandscapeGrass.cpp b/Engine/Source/Runtime/Landscape/Private/LandscapeGrass.cpp
--- a/Engine/Source/Runtime/Landscape/Private/LandscapeGrass.cpp
+++ b/Engine/Source/Runtime/Landscape/Private/LandscapeGrass.cpp
@@ -996,7 +996,10 @@ TArray<uint16> ULandscapeComponent::RenderWPOHeightmap(int32 LOD)
 
 void ULandscapeComponent::RemoveGrassMap()
 {
-	*GrassData = FLandscapeComponentGrassData();
+	// this does a thread safe replacement of the existing grassdata with a newly allocated empty (invalid) one
+	// this ensures if anyone else is accessing the old grassdata (while holding a shared ref to it), it won't be modified or deleted
+	// this is also important for PIE, which can share grassdatas with editor
+	GrassData = MakeShared<FLandscapeComponentGrassData>();
 
 #if WITH_EDITOR
 	GrassData->bIsDirty = true;
@@ -1665,6 +1668,8 @@ void FLandscapeComponentGrassData::ConditionalDiscardDataOnLoad()
 		// If all grass types have been removed, discard the height data too.
 		if (WeightOffsets.Num() == 0)
 		{
+			// NOTE: this is NOT thread safe, as we are overwriting the existing grassdata here
+			// However if it is only performed on load, then we can assume no async tasks are holding shared refs to this yet.
 			*this = FLandscapeComponentGrassData();
 			NumElements = 0;
 		}
