Removed invalid fix that prevents edit layers from working
#rnx

diff --git a/Engine/Source/Runtime/Landscape/Private/LandscapeEditLayers.cpp b/Engine/Source/Runtime/Landscape/Private/LandscapeEditLayers.cpp
--- a/Engine/Source/Runtime/Landscape/Private/LandscapeEditLayers.cpp
+++ b/Engine/Source/Runtime/Landscape/Private/LandscapeEditLayers.cpp
@@ -5967,8 +5967,6 @@ void ALandscape::PrepareComponentDataToExtractMaterialLayersCS(const TArray<ULan
 									break;
 								}
 							}
-							// didn't find a matching layer
-							continue;
 						}
 
 						OutComponentData.Add(Data);
