Prevented weightmap textures generated by turning edit layers on on a landscape that has it off to have a mip chain, which causes an assert later down the line

diff --git a/Engine/Source/Runtime/Landscape/Private/LandscapeEditLayers.cpp b/Engine/Source/Runtime/Landscape/Private/LandscapeEditLayers.cpp
--- a/Engine/Source/Runtime/Landscape/Private/LandscapeEditLayers.cpp
+++ b/Engine/Source/Runtime/Landscape/Private/LandscapeEditLayers.cpp
@@ -2252,7 +2252,9 @@ void ALandscape::CopyOldDataToDefaultLayer(ALandscapeProxy* InProxy)
 				}
 				else
 				{
-					UTexture2D* NewLayerWeightmapTexture = InProxy->CreateLandscapeTexture(ComponentWeightmap->Source.GetSizeX(), ComponentWeightmap->Source.GetSizeY(), TEXTUREGROUP_Terrain_Weightmap, ComponentWeightmap->Source.GetFormat());
+					// No need for mip chain on edit layers : 
+					UTexture2D* NewLayerWeightmapTexture = InProxy->CreateLandscapeTexture(ComponentWeightmap->Source.GetSizeX(), ComponentWeightmap->Source.GetSizeY(), TEXTUREGROUP_Terrain_Weightmap, 
+						ComponentWeightmap->Source.GetFormat(), /*OptionalOverrideOuter = */ nullptr, /*bCompress = */false, /*bMipChain = */false);
 
 					// Only copy Mip0 as other mips will get regenerated
 					TArray64<uint8> ExistingMip0Data;
