Simplify disabling runtime grass generation per-proxy at runtime.

Doubly ensure disabled grass generation components never leave the NotReady state in game (and a bit of optimization by completely removing the check in non-editor builds).

diff --git a/Engine/Source/Runtime/Landscape/Classes/LandscapeProxy.h b/Engine/Source/Runtime/Landscape/Classes/LandscapeProxy.h
--- a/Engine/Source/Runtime/Landscape/Classes/LandscapeProxy.h
+++ b/Engine/Source/Runtime/Landscape/Classes/LandscapeProxy.h
@@ -695,6 +695,7 @@ public:
 	void UpdateGrassTypeSummary();
 
 	inline bool GetDisableRuntimeGrassMapGeneration() { return bDisableRuntimeGrassMapGeneration; }
+	inline void SetDisableRuntimeGrassMapGenerationProxyOnly(bool bInDisableRuntimeGrassMapGeneration) { bDisableRuntimeGrassMapGeneration = bInDisableRuntimeGrassMapGeneration; }
 
 	/**
 	 *	The resolution to cache lighting at, in texels/quad in one axis
diff --git a/Engine/Source/Runtime/Landscape/Private/Landscape.cpp b/Engine/Source/Runtime/Landscape/Private/Landscape.cpp
--- a/Engine/Source/Runtime/Landscape/Private/Landscape.cpp
+++ b/Engine/Source/Runtime/Landscape/Private/Landscape.cpp
@@ -3612,10 +3612,7 @@ void ALandscape::SetDisableRuntimeGrassMapGeneration(bool bInDisableRuntimeGrass
 	{
 		LandscapeInfo->ForEachLandscapeProxy([bInDisableRuntimeGrassMapGeneration](ALandscapeProxy* Proxy) -> bool
 		{
-			if (Proxy != nullptr)
-			{
-				Proxy->bDisableRuntimeGrassMapGeneration = bInDisableRuntimeGrassMapGeneration;
-			}
+			Proxy->SetDisableRuntimeGrassMapGenerationProxyOnly(bInDisableRuntimeGrassMapGeneration);
 			return true;
 		});
 	}
diff --git a/Engine/Source/Runtime/Landscape/Private/LandscapeGrassMapsBuilder.cpp b/Engine/Source/Runtime/Landscape/Private/LandscapeGrassMapsBuilder.cpp
--- a/Engine/Source/Runtime/Landscape/Private/LandscapeGrassMapsBuilder.cpp
+++ b/Engine/Source/Runtime/Landscape/Private/LandscapeGrassMapsBuilder.cpp
@@ -435,6 +435,7 @@ bool FLandscapeGrassMapsBuilder::UpdateTrackedComponents(const TArray<FVector>&
 			case EComponentStage::NotReady:
 				// in game, any not ready component will never become ready.
 				// in editor, check to see if the conditions changed.
+#if WITH_EDITOR
 				if (!bIsGameWorld && AmortizedUpdate.ShouldUpdate(ComponentStateIndex))
 				{
 					if (UE::Landscape::CanRenderGrassMap(Component))
@@ -442,6 +443,7 @@ bool FLandscapeGrassMapsBuilder::UpdateTrackedComponents(const TArray<FVector>&
 						break; // cancel and evict to restart build process
 					}
 				}
+#endif // WITH_EDITOR
 				continue; // next!
 			case EComponentStage::TextureStreaming:
 				// don't process streaming states yet -- first process rendering states to free up slots
