Fix for crash painting landscape introduced in 33667964

diff --git a/Engine/Source/Runtime/Landscape/Private/Landscape.cpp b/Engine/Source/Runtime/Landscape/Private/Landscape.cpp
--- a/Engine/Source/Runtime/Landscape/Private/Landscape.cpp
+++ b/Engine/Source/Runtime/Landscape/Private/Landscape.cpp
@@ -5605,6 +5605,9 @@ void ULandscapeInfo::RegisterActor(ALandscapeProxy* Proxy, bool bMapCheck, bool
 		{
 			UE_LOG(LogLandscape, Warning, TEXT("Multiple landscape actors with the same GUID detected: %s vs %s"), *LandscapeActor->GetPathName(), *Landscape->GetPathName());
 		}
+#if WITH_EDITORONLY_DATA
+		Landscape->bIsRegisteredWithLandscapeInfo = true;
+#endif // WITH_EDITORONLY_DATA
 	}
 	else
 	{
