Avoid caching landscape grass shaders for platforms which have grass.enabled = 0 (unless GGrassMapAlwaysBuildRuntimeGenerationResources  is also set, which will always cache the shadermaps).

diff --git a/Engine/Source/Runtime/Landscape/Private/LandscapeGrassWeightExporter.cpp b/Engine/Source/Runtime/Landscape/Private/LandscapeGrassWeightExporter.cpp
--- a/Engine/Source/Runtime/Landscape/Private/LandscapeGrassWeightExporter.cpp
+++ b/Engine/Source/Runtime/Landscape/Private/LandscapeGrassWeightExporter.cpp
@@ -59,10 +59,15 @@ static bool ShouldCacheLandscapeGrassShaders(const FMeshMaterialShaderPermutatio
 	const bool bPlatformUsesRuntimeGen = (GGrassMapUseRuntimeGeneration != 0);
 #endif // WITH_EDITOR
 
-	const bool bShouldBuildForPlatform = 
-		bIsEditorPlatform ||
-		GGrassMapAlwaysBuildRuntimeGenerationResources ||
-		bPlatformUsesRuntimeGen;
+#if WITH_EDITOR
+	static FShaderPlatformCachedIniValue<int32> GrassEnabledPerPlatform(TEXT("grass.Enable"));
+	const bool bGrassEnable = (GrassEnabledPerPlatform.Get(Parameters.Platform) != 0);
+#else 
+	const bool bGrassEnable = GGrassEnable != 0;
+#endif // WITH_EDITOR
+
+	const bool bShouldBuildForPlatform = GGrassMapAlwaysBuildRuntimeGenerationResources || (
+		bGrassEnable && (bIsEditorPlatform || bPlatformUsesRuntimeGen));
 
 	const bool bIsFixedGridVertexFactory =
 		Parameters.VertexFactoryType == FindVertexFactoryType(FName(TEXT("FLandscapeFixedGridVertexFactory"), FNAME_Find));
