Fix for nanite landscape UV seams between streaming proxies
(the nanite mesh was using proxy relative coords instead of global coords, and was also applying LOD scale incorrectly)

diff --git a/Engine/Source/Runtime/Landscape/Private/Landscape.cpp b/Engine/Source/Runtime/Landscape/Private/Landscape.cpp
--- a/Engine/Source/Runtime/Landscape/Private/Landscape.cpp
+++ b/Engine/Source/Runtime/Landscape/Private/Landscape.cpp
@@ -580,7 +580,7 @@ FGuid ALandscapeProxy::GetNaniteContentId() const
 		}
 
 		// Bump if changes to ULandscapeNaniteComponent::InitializeForLandscape() need to be enforced.
-		static FGuid ExportRawMeshGuid("36208D9A475B4D93B33BF84FFEDA1536");
+		static FGuid ExportRawMeshGuid("BF9FCA7C7A6948659695641CA727E9F2");
 		ContentStateAr << ExportRawMeshGuid;
 
 		FGuid HeightmapGuid = Component->GetHeightmap()->Source.GetId();
diff --git a/Engine/Source/Runtime/Landscape/Private/LandscapeEdit.cpp b/Engine/Source/Runtime/Landscape/Private/LandscapeEdit.cpp
--- a/Engine/Source/Runtime/Landscape/Private/LandscapeEdit.cpp
+++ b/Engine/Source/Runtime/Landscape/Private/LandscapeEdit.cpp
@@ -1158,7 +1158,7 @@ void ULandscapeComponent::FixupWeightmaps(const FGuid& InEditLayerGuid)
 						FFormatNamedArguments Arguments;
 						Arguments.Add(TEXT("LandscapeName"), FText::FromString(GetPathName()));
 						Arguments.Add(TEXT("TargetLayerName"), FText::FromString(Allocation.LayerInfo->LayerName.ToString()));
-						Arguments.Add(TEXT("EditLayerName"), FText::FromString(InEditLayerGuid.IsValid() ? LayersData[InEditLayerGuid].DebugName.ToString() : FString(TEXT("unknown"))));
+						Arguments.Add(TEXT("EditLayerName"), FText::FromString(InEditLayerGuid.IsValid() ? LayersData[InEditLayerGuid].DebugName.ToString() : FString(TEXT("<Runtime>"))));
 						FMessageLog("MapCheck").Info()
 							->AddToken(FTextToken::Create(FText::Format(LOCTEXT("MapCheck_Message_FixedUpDeletedLayerWeightmap", "{LandscapeName} : Removed unused layer weightmap Edit Layer: '{EditLayerName}' Target Layer: '{TargetLayerName}'"), Arguments)))
 							->AddToken(FMapErrorToken::Create(FMapErrors::FixedUpDeletedLayerWeightmap));
@@ -4278,6 +4278,9 @@ bool ALandscapeProxy::ExportToRawMeshDataCopy(const FRawMeshExportParams& InExpo
 
 		const FIntPoint ComponentOffsetRelativeToProxyBoundsQuads = Component->GetSectionBase() - LandscapeSectionOffset - LandscapeProxyBoundsRect.Min;
 		const FVector2f ComponentOffsetRelativeToProxyBoundsQuadsLOD = FVector2f(ComponentOffsetRelativeToProxyBoundsQuads) * LODScale;
+
+		const FVector2f GlobalQuadCoordsOffset = FVector2f(Component->GetSectionBase());
+
 		const FVector2f ComponentUVScaleRelativeToProxyBoundsLOD = LandscapeProxyBoundsRectUVScale * InvLODScale;
 
 		const FVector2f ComponentHeightmapUVBias = FVector2f(static_cast<float>(Component->HeightmapScaleBias.Z), static_cast<float>(Component->HeightmapScaleBias.W));
@@ -4547,7 +4550,7 @@ bool ALandscapeProxy::ExportToRawMeshDataCopy(const FRawMeshExportParams& InExpo
 							case FRawMeshExportParams::EUVMappingType::TerrainCoordMapping_XZ:
 							case FRawMeshExportParams::EUVMappingType::TerrainCoordMapping_YZ:
 							{
-								FVector2f UV = (ComponentOffsetRelativeToProxyBoundsQuadsLOD + QuadCoords * InvLODScale);
+								FVector2f UV = (GlobalQuadCoordsOffset + QuadCoords * InvLODScale);
 								if (UVMappingType == FRawMeshExportParams::EUVMappingType::TerrainCoordMapping_XY)
 								{
 									// Already in XY mapping
