Fix non-unity build

diff --git a/Engine/Source/Runtime/Landscape/Private/LandscapeGrassWeightExporter.cpp b/Engine/Source/Runtime/Landscape/Private/LandscapeGrassWeightExporter.cpp
--- a/Engine/Source/Runtime/Landscape/Private/LandscapeGrassWeightExporter.cpp
+++ b/Engine/Source/Runtime/Landscape/Private/LandscapeGrassWeightExporter.cpp
@@ -46,6 +46,8 @@ public:
 	FVector2f RenderOffset;
 };
 
+extern int32 GGrassEnable;
+
 static bool ShouldCacheLandscapeGrassShaders(const FMeshMaterialShaderPermutationParameters& Parameters)
 {
 	const bool bIsEditorPlatform = 
