World Partition - HLOD: When processing landscape proxies, avoid crash when dealing with a proxy that has no geometry

diff --git a/Engine/Source/Runtime/Landscape/Private/LandscapeHLODBuilder.cpp b/Engine/Source/Runtime/Landscape/Private/LandscapeHLODBuilder.cpp
--- a/Engine/Source/Runtime/Landscape/Private/LandscapeHLODBuilder.cpp
+++ b/Engine/Source/Runtime/Landscape/Private/LandscapeHLODBuilder.cpp
@@ -400,7 +400,12 @@ TArray<UActorComponent*> ULandscapeHLODBuilder::Build(const FHLODBuildContext& I
 				ExportParams.SkirtDepth = ScaleFactor;
 			}
 
-			LandscapeProxy->ExportToRawMesh(ExportParams, *MeshDescription);
+			// It's possible for landscape proxies to have no mesh data, in this case the export will fail...
+			if (!LandscapeProxy->ExportToRawMesh(ExportParams, *MeshDescription))
+			{
+				UE_LOG(LogHLODBuilder, Display, TEXT("Skipping HLOD builder for landscape proxy '%s' as it failed to export a mesh!"), *LandscapeProxy->GetFullName());
+				continue;
+			}
 
 			StaticMesh->CommitMeshDescription(0);
 
