Revert landscape nanite content ID to avoid cook time hit
This means that the fix requires users to manually rebuild nanite on the affected landscapes, they will not automatically rebuild to correct the issue.

diff --git a/Engine/Source/Runtime/Landscape/Private/Landscape.cpp b/Engine/Source/Runtime/Landscape/Private/Landscape.cpp
--- a/Engine/Source/Runtime/Landscape/Private/Landscape.cpp
+++ b/Engine/Source/Runtime/Landscape/Private/Landscape.cpp
@@ -580,7 +580,7 @@ FGuid ALandscapeProxy::GetNaniteContentId() const
 		}
 
 		// Bump if changes to ULandscapeNaniteComponent::InitializeForLandscape() need to be enforced.
-		static FGuid ExportRawMeshGuid("BF9FCA7C7A6948659695641CA727E9F2");
+		static FGuid ExportRawMeshGuid("36208D9A475B4D93B33BF84FFEDA1536");
 		ContentStateAr << ExportRawMeshGuid;
 
 		FGuid HeightmapGuid = Component->GetHeightmap()->Source.GetId();
