Fix for landscape texture streaming getting disabled after garbage collection
Here we just reset the flag post-garbage collection to ensure streaming requests aren't lost

diff --git a/Engine/Source/Runtime/Landscape/Private/LandscapeSubsystem.cpp b/Engine/Source/Runtime/Landscape/Private/LandscapeSubsystem.cpp
--- a/Engine/Source/Runtime/Landscape/Private/LandscapeSubsystem.cpp
+++ b/Engine/Source/Runtime/Landscape/Private/LandscapeSubsystem.cpp
@@ -259,7 +259,7 @@ void ULandscapeSubsystem::Deinitialize()
 void ULandscapeSubsystem::HandlePostGarbageCollect()
 {
 	ALandscapeProxy::RemoveInvalidExclusionBoxes();
-	GetTextureStreamingManager()->CleanupInvalidEntries();
+	GetTextureStreamingManager()->CleanupPostGarbageCollect();
 }
 
 TStatId ULandscapeSubsystem::GetStatId() const
diff --git a/Engine/Source/Runtime/Landscape/Private/LandscapeTextureStreamingManager.cpp b/Engine/Source/Runtime/Landscape/Private/LandscapeTextureStreamingManager.cpp
--- a/Engine/Source/Runtime/Landscape/Private/LandscapeTextureStreamingManager.cpp
+++ b/Engine/Source/Runtime/Landscape/Private/LandscapeTextureStreamingManager.cpp
@@ -128,15 +128,24 @@ bool FLandscapeTextureStreamingManager::WaitForTextureStreaming()
 	return bFullyStreamed;
 }
 
-void FLandscapeTextureStreamingManager::CleanupInvalidEntries()
+void FLandscapeTextureStreamingManager::CleanupPostGarbageCollect()
 {
 	for (auto It = TextureStates.CreateIterator(); It; ++It)
 	{
-		TWeakObjectPtr<UTexture>& TexPtr = It.Key();
-		if (!TexPtr.IsValid())
+		UTexture* Texture = It.Key().Get();
+		if (Texture == nullptr)
 		{
 			It.RemoveCurrent();
 		}
+		else
+		{
+			// reset the texture force resident after garbage collection (which clears it sometimes)
+			FTextureState& State = It.Value();
+			if (State.bForever || State.RequestCount > 0)
+			{
+				Texture->bForceMiplevelsToBeResident = true;
+			}
+		}
 	}
 }
 
diff --git a/Engine/Source/Runtime/Landscape/Private/LandscapeTextureStreamingManager.h b/Engine/Source/Runtime/Landscape/Private/LandscapeTextureStreamingManager.h
--- a/Engine/Source/Runtime/Landscape/Private/LandscapeTextureStreamingManager.h
+++ b/Engine/Source/Runtime/Landscape/Private/LandscapeTextureStreamingManager.h
@@ -57,8 +57,8 @@ public:
 	// (normally they don't complete streaming until a certain point in the update tick)
 	bool WaitForTextureStreaming();
 
-	// Call this clean up any old entries in tracked TextureStates for textures that have been unloaded without first being Unrequested.
-	void CleanupInvalidEntries();
+	// Call this post garbage collect to clean up any old entries in tracked TextureStates for textures that have been unloaded without first being Unrequested, and make sure the streaming flags weren't cleared
+	void CleanupPostGarbageCollect();
 
 	// Check that all requested textures are still requested.
 	void CheckRequestedTextures();
