Optimize of grass instance generation, replace TWeakObjectPtr with TObjectPtr
This saves about 9-15% of the update cost at runtime

diff --git a/Engine/Source/Runtime/Landscape/Private/LandscapeSubsystem.cpp b/Engine/Source/Runtime/Landscape/Private/LandscapeSubsystem.cpp
--- a/Engine/Source/Runtime/Landscape/Private/LandscapeSubsystem.cpp
+++ b/Engine/Source/Runtime/Landscape/Private/LandscapeSubsystem.cpp
@@ -139,21 +139,30 @@ ULandscapeSubsystem::~ULandscapeSubsystem()
 
 void ULandscapeSubsystem::RegisterActor(ALandscapeProxy* Proxy)
 {
-	Proxies.AddUnique(TWeakObjectPtr<ALandscapeProxy>(Proxy));
+	check(Proxy != nullptr);
+
+ 	TObjectPtr<ALandscapeProxy> ProxyPtr(Proxy);
+
+	// editor can get multiple registration calls, ensure we don't register more than once
+ 	Proxies.AddUnique(ProxyPtr);
 	
 	if (ALandscape* LandscapeActor = Cast<ALandscape>(Proxy))
 	{
-		LandscapeActors.AddUnique(TWeakObjectPtr<ALandscape>(LandscapeActor));
+		TObjectPtr<ALandscape> LandscapeActorPtr(LandscapeActor);
+		LandscapeActors.AddUnique(LandscapeActorPtr);
 	}
 }
 
 void ULandscapeSubsystem::UnregisterActor(ALandscapeProxy* Proxy)
 {
-	Proxies.Remove(TWeakObjectPtr<ALandscapeProxy>(Proxy));
+	check(Proxy != nullptr);
+	TObjectPtr<ALandscapeProxy> ProxyPtr(Proxy);
+	Proxies.Remove(ProxyPtr);
 
 	if (ALandscape* LandscapeActor = Cast<ALandscape>(Proxy))
 	{
-		LandscapeActors.Remove(TWeakObjectPtr<ALandscape>(LandscapeActor));
+		TObjectPtr<ALandscape> LandscapeActorPtr(LandscapeActor);
+		LandscapeActors.Remove(LandscapeActorPtr);
 	}
 }
 
@@ -281,12 +290,10 @@ void ULandscapeSubsystem::UnregisterComponent(ULandscapeComponent* Component)
 void ULandscapeSubsystem::RemoveGrassInstances(const TSet<ULandscapeComponent*>* ComponentsToRemoveGrassInstances)
 {
 	TRACE_CPUPROFILER_EVENT_SCOPE(ULandscapeSubsystem::RemoveGrassInstances);
-	for (TWeakObjectPtr<ALandscapeProxy> ProxyPtr : Proxies)
+	for (TObjectPtr<ALandscapeProxy> ProxyPtr : Proxies)
 	{
-		if (ALandscapeProxy* Proxy = ProxyPtr.Get())
-		{
-			Proxy->FlushGrassComponents(ComponentsToRemoveGrassInstances, /*bFlushGrassMaps = */false);
-		}
+		ALandscapeProxy* Proxy = ProxyPtr.Get();
+		Proxy->FlushGrassComponents(ComponentsToRemoveGrassInstances, /*bFlushGrassMaps = */false);
 	}
 }
 
@@ -332,12 +339,10 @@ void ULandscapeSubsystem::RegenerateGrass(bool bInFlushGrass, bool bInForceSync,
 		}
 
 		// Update the grass near the specified location(s) : 
-		for (TWeakObjectPtr<ALandscapeProxy> ProxyPtr : Proxies)
+		for (TObjectPtr<ALandscapeProxy> ProxyPtr : Proxies)
 		{
-			if (ALandscapeProxy* Proxy = ProxyPtr.Get())
-			{
-				Proxy->UpdateGrass(CameraLocations, bInForceSync);
-			}
+			ALandscapeProxy* Proxy = ProxyPtr.Get();
+			Proxy->UpdateGrass(CameraLocations, bInForceSync);
 		}
 	}
 }
@@ -431,9 +436,9 @@ void ULandscapeSubsystem::Tick(float DeltaTime)
 	bool bAllProxiesRuntimeGrassMapsDisabled = true;
 
 #if WITH_EDITOR
-	for (TWeakObjectPtr<ALandscape> ActorPtr : LandscapeActors)
+	for (TObjectPtr<ALandscape> ActorPtr : LandscapeActors)
 	{
-		if (ALandscape* Landscape = ActorPtr.Get())
+		ALandscape* Landscape = ActorPtr.Get();
 		{
 			// if either of these things are true, then we wait for them to complete before running ANY grass map updates..
 			bool bLandscapeToolIsModifyingLandscape = !Landscape->bGrassUpdateEnabled;
@@ -447,35 +452,38 @@ void ULandscapeSubsystem::Tick(float DeltaTime)
 #endif // WITH_EDITOR
 
 	static TArray<ALandscapeProxy*> ActiveProxies;
-	ActiveProxies.Reset(Proxies.Num());
-	for (TWeakObjectPtr<ALandscapeProxy> ProxyPtr : Proxies)
 	{
-		if (ALandscapeProxy* Proxy = ProxyPtr.Get())
+		ActiveProxies.Reset(Proxies.Num());
+
+		for (TObjectPtr<ALandscapeProxy> ProxyPtr : Proxies)
 		{
-			ActiveProxies.Add(Proxy);
-			
-			// Update the proxies proxy
+			ALandscapeProxy* Proxy = ProxyPtr.Get();
 			{
-				if (!Proxy->GetDisableRuntimeGrassMapGeneration())
+				ActiveProxies.Add(Proxy);
+			
+				// Update the proxies proxy
 				{
-					bAllProxiesRuntimeGrassMapsDisabled = false;
-				}
+					if (!Proxy->GetDisableRuntimeGrassMapGeneration())
+					{
+						bAllProxiesRuntimeGrassMapsDisabled = false;
+					}
 
-#if WITH_EDITOR
-				if (!bIsGameWorld)
-				{
-					// in editor, automatically update component grass types if the material changes
-					for (ULandscapeComponent* Component : Proxy->LandscapeComponents)
+	#if WITH_EDITOR
+					if (!bIsGameWorld)
 					{
-						Component->UpdateGrassTypes();
+						// in editor, automatically update component grass types if the material changes
+						for (ULandscapeComponent* Component : Proxy->LandscapeComponents)
+						{
+							Component->UpdateGrassTypes();
+						}
 					}
-				}
-#endif // WITH_EDITOR
+	#endif // WITH_EDITOR
 
-				// Update the grass type summary if necessary
-				if (!Proxy->IsGrassTypeSummaryValid())
-				{
-					Proxy->UpdateGrassTypeSummary();
+					// Update the grass type summary if necessary
+					if (!Proxy->IsGrassTypeSummaryValid())
+					{
+						Proxy->UpdateGrassTypeSummary();
+					}
 				}
 			}
 		}
@@ -605,12 +613,10 @@ void ULandscapeSubsystem::OnNaniteEnabledChanged(IConsoleVariable*)
 {
 	QUICK_SCOPE_CYCLE_COUNTER(STAT_Landscape_OnNaniteEnabledChanged);
 
-	for (TWeakObjectPtr<ALandscapeProxy>& ProxyPtr : Proxies)
+	for (TObjectPtr<ALandscapeProxy>& ProxyPtr : Proxies)
 	{
-		if (ALandscapeProxy* Proxy = ProxyPtr.Get())
-		{
-			Proxy->UpdateRenderingMethod();
-		}
+		ALandscapeProxy* Proxy = ProxyPtr.Get();
+		Proxy->UpdateRenderingMethod();
 	}
 }
 
@@ -648,7 +654,7 @@ TArray<ALandscapeProxy*> ULandscapeSubsystem::GetOutdatedProxies(UE::Landscape::
 
 	TArray<ALandscapeProxy*> FinalProxiesToBuild;
 	Algo::TransformIf(Proxies, FinalProxiesToBuild, 
-		[InMatchingOutdatedDataFlags, bInMustMatchAllFlags](const TWeakObjectPtr<ALandscapeProxy>& InProxyPtr)
+		[InMatchingOutdatedDataFlags, bInMustMatchAllFlags](const TObjectPtr<ALandscapeProxy>& InProxyPtr)
 		{ 
 			UE::Landscape::EOutdatedDataFlags ProxyOutdatedDataFlags = InProxyPtr->GetOutdatedDataFlags();
 			return bInMustMatchAllFlags
@@ -672,9 +678,9 @@ TArray<TTuple<ALandscapeProxy*, UE::Landscape::EOutdatedDataFlags>> ULandscapeSu
 
 	TArray<TTuple<ALandscapeProxy*, UE::Landscape::EOutdatedDataFlags>> OutdatedProxies;
 	OutdatedProxies.Reserve(Proxies.Num());
-	for (TWeakObjectPtr<ALandscapeProxy> Proxy : Proxies)
+	for (TObjectPtr<ALandscapeProxy> Proxy : Proxies)
 	{
-		if (ALandscapeProxy* ValidProxy = Proxy.Get())
+		ALandscapeProxy* ValidProxy = Proxy.Get();
 		{
 			const UE::Landscape::EOutdatedDataFlags ProxyOutdatedDataFlags = ValidProxy->GetOutdatedDataFlags();
 			if ((bInMustMatchAllFlags && EnumHasAllFlags(ProxyOutdatedDataFlags, InMatchingOutdatedDataFlags))
@@ -705,7 +711,7 @@ void ULandscapeSubsystem::BuildNanite(TArrayView<ALandscapeProxy*> InProxiesToBu
 	TArray<ALandscapeProxy*> FinalProxiesToBuild;
 	if (InProxiesToBuild.IsEmpty())
 	{
-		Algo::Transform(Proxies, FinalProxiesToBuild, [](const TWeakObjectPtr<ALandscapeProxy>& InProxyPtr) { return InProxyPtr.Get(); });
+		Algo::Transform(Proxies, FinalProxiesToBuild, [](const TObjectPtr<ALandscapeProxy>& InProxyPtr) { return InProxyPtr.Get(); });
 	}
 	else 
 	{
diff --git a/Engine/Source/Runtime/Landscape/Public/LandscapeSubsystem.h b/Engine/Source/Runtime/Landscape/Public/LandscapeSubsystem.h
--- a/Engine/Source/Runtime/Landscape/Public/LandscapeSubsystem.h
+++ b/Engine/Source/Runtime/Landscape/Public/LandscapeSubsystem.h
@@ -157,8 +157,15 @@ private:
 	void HandlePostGarbageCollect();
 
 	bool bIsGrassCreationPrioritized = false;
-	TArray<TWeakObjectPtr<ALandscape>> LandscapeActors;
-	TArray<TWeakObjectPtr<ALandscapeProxy>> Proxies;
+
+	// UPROPERTY ensures these objects are not deleted before being unregistered
+	// (technically not necessary, as actors should always unregister prior to deletion)
+	UPROPERTY()
+	TArray<TObjectPtr<ALandscape>> LandscapeActors;
+
+	UPROPERTY()
+	TArray<TObjectPtr<ALandscapeProxy>> Proxies;
+
 	FDelegateHandle OnNaniteWorldSettingsChangedHandle;
 
 	FLandscapeTextureStreamingManager* TextureStreamingManager = nullptr;
