Change default LODBlendRange to 0.1

This feature has been running with 0.1 for several releases, and this seems to be a better default value, especially in aggressive LOD situations.
Because of how defaults are serialized, this change should also update all saved landscapes that were using the previous default (1.0) to instead use the new default (0.1).

diff --git a/Engine/Source/Runtime/Landscape/Classes/LandscapeProxy.h b/Engine/Source/Runtime/Landscape/Classes/LandscapeProxy.h
--- a/Engine/Source/Runtime/Landscape/Classes/LandscapeProxy.h
+++ b/Engine/Source/Runtime/Landscape/Classes/LandscapeProxy.h
@@ -529,7 +529,7 @@ public:
 
 	/** This controls the area that blends LOD between neighboring sections. At 1.0 it blends across the entire section, and lower numbers reduce the blend region to be closer to the boundary. */
 	UPROPERTY(EditAnywhere, Category = "LOD Distribution", meta = (DisplayName = "Blend Range", ClampMin = "0.01", ClampMax = "1.0", UIMin = "0.01", UIMax = "1.0", LandscapeInherited))
-	float LODBlendRange = 1.0f;
+	float LODBlendRange = 0.1f;
 
 #if WITH_EDITORONLY_DATA
 	/** LOD level to use when exporting the landscape to obj or FBX */
