Optimizations for landscape flatten/smooth tool :
* Don't resolve both heightmaps and weightmaps while we are only interested in one of the 2 in this case
* Perform all texture readback flushes in a single render command to avoid time lost flushing rendering commands repeatedly
* Keep a certain amount of readback staging textures instead of depleting the pool entirely, which removes hitching when new textures are suddenly requested (controlled by CVar : landscape.ReadbackPoolSizeMB, default : 256MB)
#tests UEENGQA-119816

diff --git a/Engine/Source/Runtime/Landscape/Classes/Landscape.h b/Engine/Source/Runtime/Landscape/Classes/Landscape.h
--- a/Engine/Source/Runtime/Landscape/Classes/Landscape.h
+++ b/Engine/Source/Runtime/Landscape/Classes/Landscape.h
@@ -479,15 +479,15 @@ private:
 	int32 RegenerateLayersHeightmaps(const FUpdateLayersContentContext& InUpdateLayersContentContext);
 	int32 PerformLayersHeightmapsLocalMerge(const FUpdateLayersContentContext& InUpdateLayersContentContext, const FEditLayersHeightmapMergeParams& InMergeParams);
 	int32 PerformLayersHeightmapsGlobalMerge(const FUpdateLayersContentContext& InUpdateLayersContentContext, const FEditLayersHeightmapMergeParams& InMergeParams);
-	void ResolveLayersHeightmapTexture(const FTextureToComponentHelper& MapHelper, const TSet<UTexture2D*>& HeightmapsToResolve, bool bIntermediateRender, bool bFlushRender, TArray<FLandscapeEditLayerComponentReadbackResult>& InOutComponentReadbackResults);
+	void ResolveLayersHeightmapTexture(const FTextureToComponentHelper& MapHelper, const TSet<UTexture2D*>& HeightmapsToResolve, bool bIntermediateRender, TArray<FLandscapeEditLayerComponentReadbackResult>& InOutComponentReadbackResults);
 
 	int32 RegenerateLayersWeightmaps(FUpdateLayersContentContext& InUpdateLayersContentContext);
 	int32 PerformLayersWeightmapsLocalMerge(FUpdateLayersContentContext& InUpdateLayersContentContext, const FEditLayersWeightmapMergeParams& InMergeParams);
 	int32 PerformLayersWeightmapsGlobalMerge(FUpdateLayersContentContext& InUpdateLayersContentContext, const FEditLayersWeightmapMergeParams& InMergeParams);
-	void ResolveLayersWeightmapTexture(const FTextureToComponentHelper& MapHelper, const TSet<UTexture2D*>& WeightmapsToResolve, bool bIntermediateRender, bool bFlushRender, TArray<FLandscapeEditLayerComponentReadbackResult>& InOutComponentReadbackResults);
+	void ResolveLayersWeightmapTexture(const FTextureToComponentHelper& MapHelper, const TSet<UTexture2D*>& WeightmapsToResolve, bool bIntermediateRender, TArray<FLandscapeEditLayerComponentReadbackResult>& InOutComponentReadbackResults);
 
 	bool HasTextureDataChanged(TArrayView<const FColor> InOldData, TArrayView<const FColor> InNewData, bool bInIsWeightmap, uint64 InPreviousHash, uint64& OutNewHash, TOptional<uint8>& OutChangedWeightmapChannelsMasks) const;
-	bool ResolveLayersTexture(FTextureToComponentHelper const& MapHelper, FLandscapeEditLayerReadback* InCPUReadBack, UTexture2D* InOutputTexture, bool bIntermediateRender,	bool bFlushRender,
+	bool ResolveLayersTexture(FTextureToComponentHelper const& MapHelper, FLandscapeEditLayerReadback* InCPUReadBack, UTexture2D* InOutputTexture, bool bIntermediateRender,
 		TArray<FLandscapeEditLayerComponentReadbackResult>& InOutComponentReadbackResults, bool bIsWeightmap);
 
 	static bool IsUpdateFlagEnabledForModes(ELandscapeComponentUpdateFlag InFlag, uint32 InUpdateModes);
diff --git a/Engine/Source/Runtime/Landscape/Classes/LandscapeComponent.h b/Engine/Source/Runtime/Landscape/Classes/LandscapeComponent.h
--- a/Engine/Source/Runtime/Landscape/Classes/LandscapeComponent.h
+++ b/Engine/Source/Runtime/Landscape/Classes/LandscapeComponent.h
@@ -378,6 +378,8 @@ enum ELandscapeLayerUpdateMode : uint32
 	Update_All = Update_Weightmap_All | Update_Heightmap_All,
 	Update_All_Editing = Update_Weightmap_Editing | Update_Heightmap_Editing,
 	Update_All_Editing_NoCollision = Update_Weightmap_Editing_NoCollision | Update_Heightmap_Editing_NoCollision,
+	Update_Heightmap_Types = Update_Heightmap_All | Update_Heightmap_Editing | Update_Heightmap_Editing_NoCollision,
+	Update_Weightmap_Types = Update_Weightmap_All | Update_Weightmap_Editing | Update_Weightmap_Editing_NoCollision,
 	// In cases where we couldn't update the clients right away this flag will be set in RegenerateLayersContent
 	Update_Client_Deferred = 1 << 6,
 	// Update landscape component clients while editing
diff --git a/Engine/Source/Runtime/Landscape/Private/LandscapeEditLayers.cpp b/Engine/Source/Runtime/Landscape/Private/LandscapeEditLayers.cpp
--- a/Engine/Source/Runtime/Landscape/Private/LandscapeEditLayers.cpp
+++ b/Engine/Source/Runtime/Landscape/Private/LandscapeEditLayers.cpp
@@ -298,8 +298,9 @@ struct FTextureToComponentHelper
 	};
 	FRIEND_ENUM_CLASS_FLAGS(ERefreshFlags);
 
-	FTextureToComponentHelper(const ULandscapeInfo& InLandscapeInfo)
+	FTextureToComponentHelper(const ULandscapeInfo& InLandscapeInfo, uint32 InLayerContentUpdateModes)
 		: LandscapeInfo(&InLandscapeInfo)
+		, LayerContentUpdateModes(InLayerContentUpdateModes)
 	{
 		Refresh(ERefreshFlags::RefreshAll);
 	}
@@ -320,6 +321,16 @@ struct FTextureToComponentHelper
 			});
 		}
 
+		// Restrict the refresh flags to only the type we're trying to update : heightmaps or weightmaps :
+		if ((LayerContentUpdateModes & ELandscapeLayerUpdateMode::Update_Heightmap_Types) == 0)
+		{
+			EnumRemoveFlags(InRefreshFlags, ERefreshFlags::RefreshHeightmaps);
+		}
+		if ((LayerContentUpdateModes & ELandscapeLayerUpdateMode::Update_Weightmap_Types) == 0)
+		{
+			EnumRemoveFlags(InRefreshFlags, ERefreshFlags::RefreshWeightmaps);
+		}
+
 		if (EnumHasAnyFlags(InRefreshFlags, ERefreshFlags::RefreshHeightmaps | ERefreshFlags::RefreshWeightmaps))
 		{
 			// Cleanup our heightmap/weightmap info:
@@ -369,6 +380,7 @@ struct FTextureToComponentHelper
 	}
 
 	const ULandscapeInfo* LandscapeInfo = nullptr;
+	uint32 LayerContentUpdateModes = 0;
 	TArray< ULandscapeComponent* > LandscapeComponents;
 	TSet< UTexture2D* > Heightmaps;
 	TMap< UTexture2D*, TArray<ULandscapeComponent*> > HeightmapToComponents;
@@ -3799,8 +3811,9 @@ struct FUpdateLayersContentContext
 	};
 	FRIEND_ENUM_CLASS_FLAGS(ERefreshFlags);
 
-	FUpdateLayersContentContext(const FTextureToComponentHelper& InMapHelper, bool bInPartialUpdate)
+	FUpdateLayersContentContext(const FTextureToComponentHelper& InMapHelper, bool bInPartialUpdate, uint32 InLayerContentUpdateModes)
 		: bPartialUpdate(bInPartialUpdate)
+		, LayerContentUpdateModes(InLayerContentUpdateModes)
 		, MapHelper(InMapHelper)
 	{
 		// No need to update the map helper, it's assumed to be already ready in the constructor
@@ -3828,6 +3841,7 @@ struct FUpdateLayersContentContext
 	void Refresh(ERefreshFlags InRefreshFlags)
 	{
 		TRACE_CPUPROFILER_EVENT_SCOPE(UpdateLayersContentContext_Refresh);
+
 		// Start by updating the map helper if necessary (keep track of components/heightmaps/weightmaps relationship) :
 		if (EnumHasAnyFlags(InRefreshFlags, ERefreshFlags::RefreshMapHelper))
 		{
@@ -3855,6 +3869,16 @@ struct FUpdateLayersContentContext
 			}
 		}
 
+		// At this point, restrict the refresh flags to only the type we're trying to update : heightmaps or weightmaps :
+		if ((LayerContentUpdateModes & ELandscapeLayerUpdateMode::Update_Heightmap_Types) == 0)
+		{
+			EnumRemoveFlags(InRefreshFlags, ERefreshFlags::RefreshHeightmapInfos);
+		}
+		if ((LayerContentUpdateModes & ELandscapeLayerUpdateMode::Update_Weightmap_Types) == 0)
+		{
+			EnumRemoveFlags(InRefreshFlags, ERefreshFlags::RefreshWeightmapInfos);
+		}
+
 		if (EnumHasAnyFlags(InRefreshFlags, ERefreshFlags::RefreshHeightmapInfos | ERefreshFlags::RefreshWeightmapInfos))
 		{
 			TSet<UTexture2D*> HeightmapsToRender;
@@ -4007,6 +4031,8 @@ struct FUpdateLayersContentContext
 
 	// Indicates whether all components of the landscape are marked dirty :
 	const bool bPartialUpdate = false;
+	// Layer content update flags : helps us track what we really want to update during the run :
+	uint32 LayerContentUpdateModes = 0;
 	// Helper to gather mappings between heightmaps/weightmaps and components :
 	FTextureToComponentHelper MapHelper;
 	// List of landscape components that have been made dirty and need to be updated : 
@@ -5142,10 +5168,9 @@ int32 ALandscape::RegenerateLayersHeightmaps(const FUpdateLayersContentContext&
 	TRACE_CPUPROFILER_EVENT_SCOPE(LandscapeLayers_RegenerateLayersHeightmaps);
 	ULandscapeInfo* Info = GetLandscapeInfo();
 
-	const int32 AllHeightmapUpdateModes = (ELandscapeLayerUpdateMode::Update_Heightmap_All | ELandscapeLayerUpdateMode::Update_Heightmap_Editing | ELandscapeLayerUpdateMode::Update_Heightmap_Editing_NoCollision);
-	const int32 HeightmapUpdateModes = LayerContentUpdateModes & AllHeightmapUpdateModes;
+	const int32 HeightmapUpdateModes = LayerContentUpdateModes & ELandscapeLayerUpdateMode::Update_Heightmap_Types;
 	const bool bForceRender = CVarForceLayersUpdate.GetValueOnAnyThread() != 0;
-	const bool bSkipBrush = CVarLandscapeLayerBrushOptim.GetValueOnAnyThread() == 1 && ((HeightmapUpdateModes & AllHeightmapUpdateModes) == ELandscapeLayerUpdateMode::Update_Heightmap_Editing);
+	const bool bSkipBrush = CVarLandscapeLayerBrushOptim.GetValueOnAnyThread() == 1 && (HeightmapUpdateModes == ELandscapeLayerUpdateMode::Update_Heightmap_Editing);
 
 	if ((HeightmapUpdateModes == 0 && !bForceRender) || Info == nullptr)
 	{
@@ -5246,7 +5271,6 @@ void ALandscape::ResolveLayersHeightmapTexture(
 	FTextureToComponentHelper const& MapHelper,
 	TSet<UTexture2D*> const& HeightmapsToResolve,
 	bool bIntermediateRender,
-	bool bFlushRender,
 	TArray<FLandscapeEditLayerComponentReadbackResult>& InOutComponentReadbackResults)
 {
 	TRACE_CPUPROFILER_EVENT_SCOPE(LandscapeLayers_ResolveLayersHeightmapTexture);
@@ -5264,7 +5288,7 @@ void ALandscape::ResolveLayersHeightmapTexture(
 		check(LandscapeProxy);
 		if (FLandscapeEditLayerReadback** CPUReadback = LandscapeProxy->HeightmapsCPUReadback.Find(Heightmap))
 		{
-			const bool bChanged = ResolveLayersTexture(MapHelper, *CPUReadback, Heightmap, bIntermediateRender, bFlushRender, InOutComponentReadbackResults, /*bIsWeightmap = */false);
+			const bool bChanged = ResolveLayersTexture(MapHelper, *CPUReadback, Heightmap, bIntermediateRender, InOutComponentReadbackResults, /*bIsWeightmap = */false);
 			if (bChanged)
 			{
 				ChangedComponents.Append(MapHelper.HeightmapToComponents[Heightmap]);
@@ -5693,20 +5717,12 @@ bool ALandscape::ResolveLayersTexture(
 	FLandscapeEditLayerReadback* InCPUReadback,
 	UTexture2D* InOutputTexture,
 	bool bIntermediateRender,
-	bool bFlushRender,
 	TArray<FLandscapeEditLayerComponentReadbackResult>& InOutComponentReadbackResults,
 	bool bIsWeightmap)
 {
 	TRACE_CPUPROFILER_EVENT_SCOPE(LandscapeLayers_ResolveLayersTexture);
 
-	if (bFlushRender)
-	{
-		InCPUReadback->Flush();
-	}
-	else
-	{
-		InCPUReadback->Tick();
-	}
+	InCPUReadback->Tick();
 
 	const int32 CompletedReadbackNum = InCPUReadback->GetCompletedResultNum();
 
@@ -7734,9 +7750,8 @@ int32 ALandscape::PerformLayersWeightmapsGlobalMerge(FUpdateLayersContentContext
 int32 ALandscape::RegenerateLayersWeightmaps(FUpdateLayersContentContext& InUpdateLayersContentContext)
 {
 	TRACE_CPUPROFILER_EVENT_SCOPE(LandscapeLayers_RegenerateLayersWeightmaps);
-	const int32 AllWeightmapUpdateModes = (ELandscapeLayerUpdateMode::Update_Weightmap_All | ELandscapeLayerUpdateMode::Update_Weightmap_Editing | ELandscapeLayerUpdateMode::Update_Weightmap_Editing_NoCollision);
-	const int32 WeightmapUpdateModes = LayerContentUpdateModes & AllWeightmapUpdateModes;
-	const bool bSkipBrush = CVarLandscapeLayerBrushOptim.GetValueOnAnyThread() == 1 && ((WeightmapUpdateModes & AllWeightmapUpdateModes) == ELandscapeLayerUpdateMode::Update_Weightmap_Editing);
+	const int32 WeightmapUpdateModes = LayerContentUpdateModes & ELandscapeLayerUpdateMode::Update_Weightmap_Types;
+	const bool bSkipBrush = CVarLandscapeLayerBrushOptim.GetValueOnAnyThread() == 1 && (WeightmapUpdateModes == ELandscapeLayerUpdateMode::Update_Weightmap_Editing);
 	const bool bForceRender = CVarForceLayersUpdate.GetValueOnAnyThread() != 0;
 
 	ULandscapeInfo* Info = GetLandscapeInfo();
@@ -8024,7 +8039,6 @@ void ALandscape::ResolveLayersWeightmapTexture(
 	FTextureToComponentHelper const& MapHelper,
 	TSet<UTexture2D*> const& WeightmapsToResolve,
 	bool bIntermediateRender,
-	bool bFlushRender,
 	TArray<FLandscapeEditLayerComponentReadbackResult>& InOutComponentReadbackResults)
 {
 	TRACE_CPUPROFILER_EVENT_SCOPE(LandscapeLayers_ResolveLayersWeightmapTexture);
@@ -8042,7 +8056,7 @@ void ALandscape::ResolveLayersWeightmapTexture(
 		check(LandscapeProxy);
 		if (FLandscapeEditLayerReadback** CPUReadback = LandscapeProxy->WeightmapsCPUReadback.Find(Weightmap))
 		{
-			const bool bChanged = ResolveLayersTexture(MapHelper, *CPUReadback, Weightmap, bIntermediateRender, bFlushRender, InOutComponentReadbackResults, /*bIsWeightmap = */true);
+			const bool bChanged = ResolveLayersTexture(MapHelper, *CPUReadback, Weightmap, bIntermediateRender, InOutComponentReadbackResults, /*bIsWeightmap = */true);
 			if (bChanged)
 			{
 				ChangedComponents.Append(MapHelper.WeightmapToComponents[Weightmap]);
@@ -8589,7 +8603,7 @@ void ALandscape::UpdateLayersContent(bool bInWaitForStreaming, bool bInSkipMonit
 	}
 
 	// Gather mappings between heightmaps/weightmaps and components
-	FTextureToComponentHelper MapHelper(*LandscapeInfo);
+	FTextureToComponentHelper MapHelper(*LandscapeInfo, LayerContentUpdateModes);
 
 	// Poll and complete any outstanding resolve work
 	// If bIntermediateRender then we want to flush all work here before we do the intermediate render later on
@@ -8600,9 +8614,15 @@ void ALandscape::UpdateLayersContent(bool bInWaitForStreaming, bool bInSkipMonit
 		const bool bDoIntermediateRender = false; // bIntermediateRender flag is for the work queued up this frame not the delayed resolves
 		const bool bDoFlushRender = bIntermediateRender; // Flush before we do an intermediate render later in this frame
 
+		// Flushing once all readback tasks is much faster than asking each to do it so start by doing just this :
+		if (bDoFlushRender)
+		{
+			FLandscapeEditLayerReadback::FlushAllReadbackTasks();
+		}
+
 		TArray<FLandscapeEditLayerComponentReadbackResult> ComponentReadbackResults;
-		ResolveLayersHeightmapTexture(MapHelper, MapHelper.Heightmaps, bDoIntermediateRender, bDoFlushRender, ComponentReadbackResults);
-		ResolveLayersWeightmapTexture(MapHelper, MapHelper.Weightmaps, bDoIntermediateRender, bDoFlushRender, ComponentReadbackResults);
+		ResolveLayersHeightmapTexture(MapHelper, MapHelper.Heightmaps, bDoIntermediateRender, ComponentReadbackResults);
+		ResolveLayersWeightmapTexture(MapHelper, MapHelper.Weightmaps, bDoIntermediateRender, ComponentReadbackResults);
 		LayerContentUpdateModes |= UpdateAfterReadbackResolves(ComponentReadbackResults);
 	}
 
@@ -8614,7 +8634,7 @@ void ALandscape::UpdateLayersContent(bool bInWaitForStreaming, bool bInSkipMonit
 	bool bUpdateAll = LayerContentUpdateModes & Update_All;
 	bool bPartialUpdate = !bForceRender && !bUpdateAll && CVarLandscapeLayerOptim.GetValueOnAnyThread() == 1;
 
-	FUpdateLayersContentContext UpdateLayersContentContext(MapHelper, bPartialUpdate);
+	FUpdateLayersContentContext UpdateLayersContentContext(MapHelper, bPartialUpdate, LayerContentUpdateModes);
 
 	// Regenerate any heightmaps and weightmaps
 	int32 ProcessedModes = 0;
@@ -8626,9 +8646,10 @@ void ALandscape::UpdateLayersContent(bool bInWaitForStreaming, bool bInSkipMonit
 	// If we are flushing then read back resolved textures immediately
 	if (bFlushRender || CVarLandscapeForceFlush.GetValueOnGameThread() != 0)
 	{
-		const bool bDoFlushRender = true;
-		ResolveLayersHeightmapTexture(UpdateLayersContentContext.MapHelper, UpdateLayersContentContext.HeightmapsToResolve, bIntermediateRender, bDoFlushRender, UpdateLayersContentContext.AllLandscapeComponentReadbackResults);
-		ResolveLayersWeightmapTexture(UpdateLayersContentContext.MapHelper, UpdateLayersContentContext.WeightmapsToResolve, bIntermediateRender, bDoFlushRender, UpdateLayersContentContext.AllLandscapeComponentReadbackResults);
+		// Flushing once all readback tasks is much faster than asking each to do it so start by doing just this :
+		FLandscapeEditLayerReadback::FlushAllReadbackTasks();
+		ResolveLayersHeightmapTexture(UpdateLayersContentContext.MapHelper, UpdateLayersContentContext.HeightmapsToResolve, bIntermediateRender, UpdateLayersContentContext.AllLandscapeComponentReadbackResults);
+		ResolveLayersWeightmapTexture(UpdateLayersContentContext.MapHelper, UpdateLayersContentContext.WeightmapsToResolve, bIntermediateRender, UpdateLayersContentContext.AllLandscapeComponentReadbackResults);
 	}
 
 	// Clear processed mode flags
diff --git a/Engine/Source/Runtime/Landscape/Private/LandscapeEditReadback.cpp b/Engine/Source/Runtime/Landscape/Private/LandscapeEditReadback.cpp
--- a/Engine/Source/Runtime/Landscape/Private/LandscapeEditReadback.cpp
+++ b/Engine/Source/Runtime/Landscape/Private/LandscapeEditReadback.cpp
@@ -5,8 +5,18 @@
 #include "Engine/Texture2D.h"
 #include "LandscapePrivate.h"
 #include "RenderingThread.h"
+#include "RenderUtils.h"
 #include "Hash/CityHashHelpers.h"
 
+namespace UE::Landscape::Private
+{
+static TAutoConsoleVariable<int32> CVarReadbackPoolSizeMB(
+	TEXT("landscape.ReadbackPoolSizeMB"),
+	256,
+	TEXT("Minimum pool size (in MB) for the editor readbacks. This ensures a minimum amount of readback textures are left in the pool when reclaiming memory, which avoids severe hiccups when reallocating a lot of resources."));
+
+static int32 TotalStagingTexturesAllocatedSize = 0;
+}
 
 /** Data for a read back task. */
 struct FLandscapeEditReadbackTaskImpl
@@ -26,6 +36,7 @@ struct FLandscapeEditReadbackTaskImpl
 	FIntPoint Size = FIntPoint(ForceInitToZero);
 	uint32 NumMips = 0;
 	EPixelFormat Format = PF_Unknown;
+	int32 StagingTexturesAllocatedSize = 0;
 
 	// Create on render thread
 	TArray<FTextureRHIRef> StagingTextures;
@@ -34,6 +45,7 @@ struct FLandscapeEditReadbackTaskImpl
 	// Result written on render thread and read on game thread
 	ECompletionState CompletionState = ECompletionState::None;
 	TArray<TArray<FColor>> Result;
+
 };
 
 /** Initialize the read back task data that is written by game thread. */
@@ -51,6 +63,8 @@ void InitTask_GameThread(FLandscapeEditReadbackTaskImpl& Task, UTexture2D const*
 /** Initialize the read back task resources. */
 bool InitTask_RenderThread(FLandscapeEditReadbackTaskImpl& Task)
 {
+	using namespace UE::Landscape::Private;
+
 	if (Task.StagingTextures.Num() == 0 || !Task.StagingTextures[0].IsValid() || Task.StagingTextures[0]->GetSizeXYZ() != FIntVector(Task.Size.X, Task.Size.Y, 1) || (Task.StagingTextures[0]->GetFormat() != Task.Format))
 	{
 		Task.StagingTextures.SetNum(Task.NumMips);
@@ -67,6 +81,8 @@ bool InitTask_RenderThread(FLandscapeEditReadbackTaskImpl& Task)
 			Task.StagingTextures[MipIndex] = RHICreateTexture(Desc);
 		}
 
+		Task.StagingTexturesAllocatedSize += CalcTextureSize(Task.Size.X, Task.Size.Y, Task.Format, Task.NumMips);
+		TotalStagingTexturesAllocatedSize += Task.StagingTexturesAllocatedSize;
 	}
 
 	if (!Task.ReadbackFence.IsValid())
@@ -235,6 +251,8 @@ public:
 	/** Free render resources that have been unused for long enough. */
 	void GarbageCollect()
 	{
+		using namespace UE::Landscape::Private;
+		
 		const uint32 PoolSize = Pool.Num();
 		if (PoolSize > 0)
 		{
@@ -253,8 +271,14 @@ public:
 					Task->ReadbackContext.Empty();
 					Task->Result.Empty();
 
-					if (!Task->StagingTextures.IsEmpty() || Task->ReadbackFence.IsValid())
+					const int32 MinPoolSize = CVarReadbackPoolSizeMB.GetValueOnGameThread();
+					const int32 MinPoolSizeBytes = MinPoolSize * 1024 * 1024;
+					if ((!Task->StagingTextures.IsEmpty() || Task->ReadbackFence.IsValid()) 
+						&& ((TotalStagingTexturesAllocatedSize - Task->StagingTexturesAllocatedSize) > MinPoolSizeBytes)) // Don't deplete the pool under the minimum limit
 					{
+						TotalStagingTexturesAllocatedSize -= Task->StagingTexturesAllocatedSize;
+						check(TotalStagingTexturesAllocatedSize >= 0);
+
 						// Release the render resources (which may already be released)
 						ENQUEUE_RENDER_COMMAND(FLandscapeEditLayerReadback_Release)([Task](FRHICommandListImmediate& RHICmdList)
 						{
@@ -268,6 +292,27 @@ public:
 
 		FrameCount++;
 	}
+
+	void FlushAll()
+	{
+		// Flush all pending tasks in a single command
+		ENQUEUE_RENDER_COMMAND(FLandscapeEditLayerReadback_FlushAll)([this](FRHICommandListImmediate& RHICmdList)
+		{
+			auto ItEnd = Pool.end();
+			for (auto It = Pool.begin(); It != ItEnd; ++It)
+			{
+				FLandscapeEditReadbackTaskImpl& Task = *It;
+				if (Task.TextureResource != nullptr)
+				{
+					bool bTaskComplete = UpdateTask_RenderThread(RHICmdList, Task, /*bFlush = */true);
+					check(bTaskComplete); // Flush should never fail to complete
+				}
+			}
+		});
+
+		TRACE_CPUPROFILER_EVENT_SCOPE(LandscapeLayers_ReadbackFlushAll);
+		FlushRenderingCommands();
+	}
 };
 
 /** Static global pool object. */
@@ -406,3 +451,8 @@ void FLandscapeEditLayerReadback::GarbageCollectTasks()
 {
 	GReadbackTaskPool.GarbageCollect();
 }
+
+void FLandscapeEditLayerReadback::FlushAllReadbackTasks()
+{
+	GReadbackTaskPool.FlushAll();
+}
\ No newline at end of file
diff --git a/Engine/Source/Runtime/Landscape/Private/LandscapeEditReadback.h b/Engine/Source/Runtime/Landscape/Private/LandscapeEditReadback.h
--- a/Engine/Source/Runtime/Landscape/Private/LandscapeEditReadback.h
+++ b/Engine/Source/Runtime/Landscape/Private/LandscapeEditReadback.h
@@ -67,6 +67,10 @@ public:
 	/** Optional call to allow garbage collection of any task resources and check for leaks. Call once per frame. */
 	static void GarbageCollectTasks();
 
+	/** Brute-force flush of all pending readback tasks. Use this when you want to make sure all readbacks are performed. In case of multiple readbacks, this is much more efficient 
+	 than individually flushing each task. */
+	static void FlushAllReadbackTasks();
+
 private:
 	uint64 Hash;
 	TArray<int32> TaskHandles;
