Made the landscape UpToDate predicate a bit more robust to landscape patches changes

#rnx

diff --git a/Engine/Source/Runtime/Landscape/Private/LandscapeEditLayers.cpp b/Engine/Source/Runtime/Landscape/Private/LandscapeEditLayers.cpp
--- a/Engine/Source/Runtime/Landscape/Private/LandscapeEditLayers.cpp
+++ b/Engine/Source/Runtime/Landscape/Private/LandscapeEditLayers.cpp
@@ -8449,7 +8449,7 @@ void ALandscape::UpdateLayersContent(bool bInWaitForStreaming, bool bInSkipMonit
 		// If nothing to do, let's do some garbage collecting on async readback tasks so that we slowly get rid of staging textures 
 		//  (don't do it while waiting for read backs because something might prevent us from updating the readbacks (e.g. waiting for resources to compiling...), which would 
 		//  lead to FLandscapeEditReadbackTaskPool's frame count increasing while readback tasks don't have the chance to complete, leading to the "readback leak" warning to incorrectly be triggered) :
-		if ((LayerContentUpdateModes == 0) && !FLandscapeEditLayerReadback::HasWork())
+		if(IsUpToDate())
 		{
 			FLandscapeEditLayerReadback::GarbageCollectTasks();
 		}
@@ -8535,7 +8535,7 @@ void ALandscape::UpdateLayersContent(bool bInWaitForStreaming, bool bInSkipMonit
 	const bool bForceRender = CVarForceLayersUpdate.GetValueOnAnyThread() != 0;
 
 	// User triggered change has been completely processed, resetting user triggered flag on all components.
-	if (!bProcessReadbacks && (LayerContentUpdateModes == 0))
+	if(IsUpToDate())
 	{
 		GetLandscapeInfo()->ForAllLandscapeComponents(
 			[this](ULandscapeComponent* Component) -> void
@@ -8546,12 +8546,12 @@ void ALandscape::UpdateLayersContent(bool bInWaitForStreaming, bool bInSkipMonit
 					Component->SetUserTriggeredChangeRequested(/* bInUserTriggered = */false);	
 				}
 			}
-		);	
-	}
-	
-	if (LayerContentUpdateModes == 0 && !bForceRender && !bProcessReadbacks)
-	{
-		return;
+		);
+
+		if (!bForceRender)
+		{
+			return;
+		}
 	}
 
 	// The Edit layers shaders only work on SM5 : cancel any update that might happen when SM5+ shading model is not active :
@@ -9285,7 +9285,7 @@ bool ALandscape::IsUpToDate() const
 #if WITH_EDITORONLY_DATA
 	if (CanHaveLayersContent() && GetWorld() != nullptr && !GetWorld()->IsGameWorld())
 	{
-		return LayerContentUpdateModes == 0;
+		return LayerContentUpdateModes == 0 && !FLandscapeEditLayerReadback::HasWork();
 	}
 #endif
 
