Fix crash when building static lighting with a Landscape using WPO by waiting for the async readback to be complete before fetching the results.

diff --git a/Engine/Source/Runtime/Landscape/Private/LandscapeGrass.cpp b/Engine/Source/Runtime/Landscape/Private/LandscapeGrass.cpp
--- a/Engine/Source/Runtime/Landscape/Private/LandscapeGrass.cpp
+++ b/Engine/Source/Runtime/Landscape/Private/LandscapeGrass.cpp
@@ -918,6 +918,7 @@ bool ULandscapeComponent::UpdateGrassTypes(bool bForceUpdate)
 namespace UE::Landscape
 {
 	extern uint32 ComputeGrassMapGenerationHash(const ULandscapeComponent* Component, UMaterialInterface* Material);
+	extern void SubmitGPUCommands(bool bBlockUntilRTComplete, bool bBlockRTUntilGPUComplete);
 }
 
 uint32 ULandscapeComponent::ComputeGrassMapGenerationHash() const
@@ -965,16 +966,30 @@ TArray<uint16> ULandscapeComponent::RenderWPOHeightmap(int32 LOD)
 {
 	TArray<uint16> Results;
 
-	if (!CanRenderGrassMap())
-	{
-		GetMaterialInstance(0)->GetMaterialResource(GetWorld()->GetFeatureLevel())->FinishCompilation();
-	}
-
 	if (ensure(SceneProxy))
 	{
+		if (!CanRenderGrassMap())
+		{
+			GetMaterialInstance(0)->GetMaterialResource(GetWorld()->GetFeatureLevel())->FinishCompilation();
+
+			if (!CanRenderGrassMap())
+			{
+				UE_LOG(LogGrass, Verbose, TEXT("Failed to calculate Landscape WPO height for static lighting. Grass map generation shader could not be compiled."));
+				return Results;
+			}
+		}
+
 		if (LOD == 0)
 		{
 			FLandscapeGrassWeightExporter Exporter(GetLandscapeProxy(), { this }, /*bInNeedsGrassmap = */ false, /*bInNeedsHeightmap =*/ true, {});
+
+			while (!Exporter.IsAsyncReadbackComplete())
+			{
+				UE::Landscape::SubmitGPUCommands(/* bBlockUntilRTComplete =  */ true, /* bBlockRTUntilGPUComplete =  */ true);
+				bool bOutRenderCommandsQueued;
+				Exporter.CheckAndUpdateAsyncReadback(bOutRenderCommandsQueued, true);
+			}
+
 			TMap<ULandscapeComponent*, TUniquePtr<FLandscapeComponentGrassData>, TInlineSetAllocator<1>> TempGrassData;
 			TempGrassData = Exporter.FetchResults(/* bFreeAsyncReadback= */ true);
 			Results = TArray<uint16>(TempGrassData[this]->GetHeightData());
@@ -984,6 +999,14 @@ TArray<uint16> ULandscapeComponent::RenderWPOHeightmap(int32 LOD)
 			TArray<int32> HeightMips;
 			HeightMips.Add(LOD);
 			FLandscapeGrassWeightExporter Exporter(GetLandscapeProxy(), { this }, /*bInNeedsGrassmap = */ false, /*bInNeedsHeightmap =*/ false, MoveTemp(HeightMips));
+
+			while (!Exporter.IsAsyncReadbackComplete())
+			{
+				UE::Landscape::SubmitGPUCommands(/* bBlockUntilRTComplete =  */ true, /* bBlockRTUntilGPUComplete =  */ true);
+				bool bOutRenderCommandsQueued;
+				Exporter.CheckAndUpdateAsyncReadback(bOutRenderCommandsQueued, true);
+			}
+
 			TMap<ULandscapeComponent*, TUniquePtr<FLandscapeComponentGrassData>, TInlineSetAllocator<1>> TempGrassData;
 			TempGrassData = Exporter.FetchResults(/* bFreeAsyncReadback= */ true);
 			Results = MoveTemp(TempGrassData[this]->HeightMipData[LOD]);
diff --git a/Engine/Source/Runtime/Landscape/Private/LandscapeGrassMapsBuilder.cpp b/Engine/Source/Runtime/Landscape/Private/LandscapeGrassMapsBuilder.cpp
--- a/Engine/Source/Runtime/Landscape/Private/LandscapeGrassMapsBuilder.cpp
+++ b/Engine/Source/Runtime/Landscape/Private/LandscapeGrassMapsBuilder.cpp
@@ -248,7 +248,7 @@ namespace UE::Landscape
 		return MinSqrDistance;
 	}
 
-	static void SubmitGPUCommands(bool bBlockUntilRTComplete, bool bBlockRTUntilGPUComplete)
+	void SubmitGPUCommands(bool bBlockUntilRTComplete, bool bBlockRTUntilGPUComplete)
 	{
 		FEvent* ResultsReadyEvent = nullptr;
 		if (bBlockUntilRTComplete)
