Fix the index that the last row of verts in the collision mesh uses for their color.

diff --git a/Engine/Source/Runtime/Landscape/Private/LandscapeCollision.cpp b/Engine/Source/Runtime/Landscape/Private/LandscapeCollision.cpp
--- a/Engine/Source/Runtime/Landscape/Private/LandscapeCollision.cpp
+++ b/Engine/Source/Runtime/Landscape/Private/LandscapeCollision.cpp
@@ -557,9 +557,11 @@ FPrimitiveSceneProxy* ULandscapeHeightfieldCollisionComponent::CreateSceneProxy(
 
 			const Chaos::FHeightField::FData<uint16>& GeomData = InHeightfield.GeomData; 
 			const int32 NumRows = InHeightfield.GetNumRows();
+			const int32 RowBounds = NumRows - 1;
 			const int32 NumCols = InHeightfield.GetNumCols();
+			const int32 ColBounds = NumCols - 1;
 			const int32 NumVerts = NumRows * NumCols;
-			const int32 NumTris = (NumRows - 1) * (NumCols - 1) * 2;
+			const int32 NumTris = RowBounds * ColBounds * 2;
 			Vertices.SetNumUninitialized(NumVerts);
 
 			TArray<FColor, TInlineAllocator<16>> MaterialIndexColors;
@@ -576,18 +578,20 @@ FPrimitiveSceneProxy* ULandscapeHeightfieldCollisionComponent::CreateSceneProxy(
 			for (int32 I = 0; I < NumVerts; I++)
 			{
 				const Chaos::FVec3 Point = GeomData.GetPointScaled(I);
-				uint8 MaterialIndex = InHeightfield.GetMaterialIndex(I % NumCols, I / NumCols);
+				const int32 CurrentCol = I % NumCols;
+				const int32 CurrentRow = I / NumCols;
+				uint8 MaterialIndex = InHeightfield.GetMaterialIndex(CurrentCol, CurrentRow);
 				Vertices[I].Position = FVector3f(static_cast<float>(Point.X), static_cast<float>(Point.Y), static_cast<float>(Point.Z));
 
 				// Material indices are not defined for the last row/column in each component since they are per-triangle and not per-vertex.
-				// To show something intuitive for the user, we simply extend the previous vertex onto the final one.
-				if (I % NumCols == NumCols - 1)
+				// To show something intuitive for the user, we simply extend the previous vertices.
+				if (CurrentCol == ColBounds)
 				{
 					Vertices[I].Color = Vertices[I - 1].Color;
 				}
-				else if (I / NumCols == NumCols - 1)
+				else if (CurrentRow == ColBounds)
 				{
-					Vertices[I].Color = Vertices[I - NumRows - 1].Color;
+					Vertices[I].Color = Vertices[I - NumRows].Color;
 				}
 				else
 				{
@@ -599,12 +603,12 @@ FPrimitiveSceneProxy* ULandscapeHeightfieldCollisionComponent::CreateSceneProxy(
 			// Editor heightfields don't have material indices (hence, no holes), in which case InHeightfield.GeomData.MaterialIndices.Num() == 1 : 
 			const int32 NumMaterialIndices = InHeightfield.GeomData.MaterialIndices.Num();
 			const bool bHasMaterialIndices = (NumMaterialIndices > 1);
-			check(!bHasMaterialIndices || (NumMaterialIndices == ((NumRows - 1) * (NumCols - 1))));
+			check(!bHasMaterialIndices || (NumMaterialIndices == (RowBounds * ColBounds)));
 
 			int32 TriangleIdx = 0;
-			for (int32 Y = 0; Y < (NumRows - 1); Y++)
+			for (int32 Y = 0; Y < RowBounds; Y++)
 			{
-				for (int32 X = 0; X < (NumCols - 1); X++)
+				for (int32 X = 0; X < ColBounds; X++)
 				{
 					int32 DataIdx = X + Y * NumCols;
 					bool bHole = false;
@@ -612,7 +616,7 @@ FPrimitiveSceneProxy* ULandscapeHeightfieldCollisionComponent::CreateSceneProxy(
 					if (bHasMaterialIndices)
 					{
 						// Material indices don't have the final row/column : 
-						int32 MaterialIndicesDataIdx = X + Y * (NumCols - 1);
+						int32 MaterialIndicesDataIdx = X + Y * ColBounds;
 						uint8 LayerIdx = InHeightfield.GeomData.MaterialIndices[MaterialIndicesDataIdx];
 						bHole = (LayerIdx == TNumericLimits<uint8>::Max());
 					}
