Removed redundant calls to SupportsLandscapeEditing when ticking landscape proxies for grass. This is to avoid O(N^2) complexity when iterating on landscape proxies for ticking grass.

diff --git a/Engine/Source/Runtime/Landscape/Classes/LandscapeProxy.h b/Engine/Source/Runtime/Landscape/Classes/LandscapeProxy.h
--- a/Engine/Source/Runtime/Landscape/Classes/LandscapeProxy.h
+++ b/Engine/Source/Runtime/Landscape/Classes/LandscapeProxy.h
@@ -1119,7 +1119,6 @@ PRAGMA_ENABLE_DEPRECATION_WARNINGS
 
 		return true;
 	}
-	void TickGrass(const TArray<FVector>& Cameras, int32& InOutNumCompsCreated);
 	void ProcessAsyncGrassInstanceTasks(bool bWaitAsyncTasks, bool bForceSync, const TSet<UHierarchicalInstancedStaticMeshComponent*>& StillUsed);
 
 	/** Flush the grass cache, removing grass instances on the given components (or all proxy components if the component set is not specified).
diff --git a/Engine/Source/Runtime/Landscape/Private/LandscapeGrass.cpp b/Engine/Source/Runtime/Landscape/Private/LandscapeGrass.cpp
--- a/Engine/Source/Runtime/Landscape/Private/LandscapeGrass.cpp
+++ b/Engine/Source/Runtime/Landscape/Private/LandscapeGrass.cpp
@@ -1732,28 +1732,6 @@ void FLandscapeComponentGrassData::ConditionalDiscardDataOnLoad()
 // ALandscapeProxy grass-related functions
 //
 
-void ALandscapeProxy::TickGrass(const TArray<FVector>& Cameras, int32& InOutNumCompsCreated)
-{
-	TRACE_CPUPROFILER_EVENT_SCOPE(ALandscapeProxy::TickGrass);
-#if WITH_EDITORONLY_DATA
-	if (ALandscape* Landscape = GetLandscapeActor())
-	{
-		ULandscapeInfo* LandscapeInfo = GetLandscapeInfo();
-		
-		bool bLandscapeUpdateAllowed = (LandscapeInfo != nullptr) && (GetWorld()->GetFeatureLevel() >= ERHIFeatureLevel::SM5) && LandscapeInfo->SupportsLandscapeEditing();
-
-		// Don't allow grass to tick if landscape is not up to date -- unless landscape update is not possible (preview or level instanced modes)
-		bool bAllowGrassTick = Landscape->IsUpToDate() || !bLandscapeUpdateAllowed;
-		if (!bAllowGrassTick || !Landscape->bGrassUpdateEnabled)
-		{
-			return;
-		}
-	}
-#endif
-
-	UpdateGrass(Cameras, InOutNumCompsCreated);
-}
-
 struct FGrassBuilderBase
 {
 	bool bHaveValidData;
diff --git a/Engine/Source/Runtime/Landscape/Private/LandscapeSubsystem.cpp b/Engine/Source/Runtime/Landscape/Private/LandscapeSubsystem.cpp
--- a/Engine/Source/Runtime/Landscape/Private/LandscapeSubsystem.cpp
+++ b/Engine/Source/Runtime/Landscape/Private/LandscapeSubsystem.cpp
@@ -291,11 +291,11 @@ void ULandscapeSubsystem::RemoveGrassInstances(const TSet<ULandscapeComponent*>*
 {
 	TRACE_CPUPROFILER_EVENT_SCOPE(ULandscapeSubsystem::RemoveGrassInstances);
 	for (TObjectPtr<ALandscapeProxy> ProxyPtr : Proxies)
-	{
+		{
 		ALandscapeProxy* Proxy = ProxyPtr.Get();
-		Proxy->FlushGrassComponents(ComponentsToRemoveGrassInstances, /*bFlushGrassMaps = */false);
+			Proxy->FlushGrassComponents(ComponentsToRemoveGrassInstances, /*bFlushGrassMaps = */false);
+		}
 	}
-}
 
 void ULandscapeSubsystem::RegenerateGrass(bool bInFlushGrass, bool bInForceSync, TOptional<TArrayView<FVector>> InOptionalCameraLocations)
 {
@@ -340,12 +340,12 @@ void ULandscapeSubsystem::RegenerateGrass(bool bInFlushGrass, bool bInForceSync,
 
 		// Update the grass near the specified location(s) : 
 		for (TObjectPtr<ALandscapeProxy> ProxyPtr : Proxies)
-		{
+			{
 			ALandscapeProxy* Proxy = ProxyPtr.Get();
-			Proxy->UpdateGrass(CameraLocations, bInForceSync);
+				Proxy->UpdateGrass(CameraLocations, bInForceSync);
+			}
 		}
 	}
-}
 
 ETickableTickType ULandscapeSubsystem::GetTickableTickType() const
 {
@@ -436,16 +436,23 @@ void ULandscapeSubsystem::Tick(float DeltaTime)
 	bool bAllProxiesRuntimeGrassMapsDisabled = true;
 
 #if WITH_EDITOR
+	TSet<ALandscape*> DisallowedGrassTickLandscapes;
 	for (TObjectPtr<ALandscape> ActorPtr : LandscapeActors)
 	{
 		ALandscape* Landscape = ActorPtr.Get();
 		{
+			const bool bLandscapeIsUpToDate = Landscape->IsUpToDate();
+			const bool bLandscapeSupportsEditing = Landscape->GetLandscapeInfo()->SupportsLandscapeEditing();
+			const bool bLandscapeUpdateAllowed = bLandscapeSupportsEditing && (Landscape->GetWorld()->GetFeatureLevel() >= ERHIFeatureLevel::SM5);
+
 			// if either of these things are true, then we wait for them to complete before running ANY grass map updates..
 			bool bLandscapeToolIsModifyingLandscape = !Landscape->bGrassUpdateEnabled;
-			bool bLandscapeLayerMergeWillBePerformedSoon = !Landscape->IsUpToDate() && Landscape->GetLandscapeInfo()->SupportsLandscapeEditing();
-			if (bLandscapeToolIsModifyingLandscape || bLandscapeLayerMergeWillBePerformedSoon)
+			// Don't allow grass to tick if landscape is not up to date -- unless landscape update is not possible (preview or level instanced modes)
+			bool bAllowGrassTick = bLandscapeIsUpToDate || !bLandscapeUpdateAllowed;
+			if (bLandscapeToolIsModifyingLandscape || !bAllowGrassTick)
 			{
 				bAllProxiesReadyForGrassMapGeneration = false;
+				DisallowedGrassTickLandscapes.Add(Landscape);
 			}
 		}
 	}
@@ -453,41 +460,41 @@ void ULandscapeSubsystem::Tick(float DeltaTime)
 
 	static TArray<ALandscapeProxy*> ActiveProxies;
 	{
-		ActiveProxies.Reset(Proxies.Num());
+	ActiveProxies.Reset(Proxies.Num());
 
 		for (TObjectPtr<ALandscapeProxy> ProxyPtr : Proxies)
-		{
+	{
 			ALandscapeProxy* Proxy = ProxyPtr.Get();
-			{
-				ActiveProxies.Add(Proxy);
+		{
+			ActiveProxies.Add(Proxy);
 			
-				// Update the proxies proxy
+			// Update the proxies proxy
+			{
+				if (!Proxy->GetDisableRuntimeGrassMapGeneration())
 				{
-					if (!Proxy->GetDisableRuntimeGrassMapGeneration())
-					{
-						bAllProxiesRuntimeGrassMapsDisabled = false;
-					}
+					bAllProxiesRuntimeGrassMapsDisabled = false;
+				}
 
-	#if WITH_EDITOR
-					if (!bIsGameWorld)
+#if WITH_EDITOR
+				if (!bIsGameWorld)
+				{
+					// in editor, automatically update component grass types if the material changes
+					for (ULandscapeComponent* Component : Proxy->LandscapeComponents)
 					{
-						// in editor, automatically update component grass types if the material changes
-						for (ULandscapeComponent* Component : Proxy->LandscapeComponents)
-						{
-							Component->UpdateGrassTypes();
-						}
+						Component->UpdateGrassTypes();
 					}
-	#endif // WITH_EDITOR
+				}
+#endif // WITH_EDITOR
 
-					// Update the grass type summary if necessary
-					if (!Proxy->IsGrassTypeSummaryValid())
-					{
-						Proxy->UpdateGrassTypeSummary();
-					}
+				// Update the grass type summary if necessary
+				if (!Proxy->IsGrassTypeSummaryValid())
+				{
+					Proxy->UpdateGrassTypeSummary();
 				}
 			}
 		}
 	}
+	}
 
 	bool bGrassMapGenerationDisabled = bAllProxiesRuntimeGrassMapsDisabled;
 #if WITH_EDITOR
@@ -501,7 +508,6 @@ void ULandscapeSubsystem::Tick(float DeltaTime)
 
 	GrassMapsBuilder->AmortizedUpdateGrassMaps(Cameras ? *Cameras : TArray<FVector>(), bIsGrassCreationPrioritized, bAllowStartGrassMapGeneration);
 
-	int32 InOutNumComponentsCreated = 0;
 #if WITH_EDITOR
 	int32 NumProxiesUpdated = 0;
 	int32 NumMeshesToUpdate = 0;
@@ -512,6 +518,7 @@ void ULandscapeSubsystem::Tick(float DeltaTime)
 		NumNaniteMeshUpdatesAvailable -= NumMeshesToUpdate;
 	}
 #endif // WITH_EDITOR
+
 	for (ALandscapeProxy* Proxy : ActiveProxies)
 	{
 #if WITH_EDITOR
@@ -540,10 +547,17 @@ void ULandscapeSubsystem::Tick(float DeltaTime)
 			}
 		}
 #endif //WITH_EDITOR
+
 		// TODO [chris.tchou] : this stops all async task processing if cameras go away, which might leave tasks dangling
-		if (Cameras && Proxy->ShouldTickGrass())
+		bool bShouldTickGrass = Proxy->ShouldTickGrass();
+#if WITH_EDITOR
+		bShouldTickGrass &= !DisallowedGrassTickLandscapes.Contains(Proxy->GetLandscapeActor());
+#endif // WITH_EDITOR
+
+		if (bShouldTickGrass && (Cameras != nullptr))
 		{
-			Proxy->TickGrass(*Cameras, InOutNumComponentsCreated);
+			int32 InOutNumComponentsCreated = 0;
+			Proxy->UpdateGrass(*Cameras, InOutNumComponentsCreated);
 		}
 
 #if !WITH_EDITOR
@@ -614,11 +628,11 @@ void ULandscapeSubsystem::OnNaniteEnabledChanged(IConsoleVariable*)
 	QUICK_SCOPE_CYCLE_COUNTER(STAT_Landscape_OnNaniteEnabledChanged);
 
 	for (TObjectPtr<ALandscapeProxy>& ProxyPtr : Proxies)
-	{
+		{
 		ALandscapeProxy* Proxy = ProxyPtr.Get();
-		Proxy->UpdateRenderingMethod();
+			Proxy->UpdateRenderingMethod();
+		}
 	}
-}
 
 #if WITH_EDITOR
 void ULandscapeSubsystem::BuildAll()
@@ -674,14 +688,14 @@ TArray<TTuple<ALandscapeProxy*, UE::Landscape::EOutdatedDataFlags>> ULandscapeSu
 	if (!World || World->IsGameWorld())
 	{
 		return {};
-	}
+}
 
 	TArray<TTuple<ALandscapeProxy*, UE::Landscape::EOutdatedDataFlags>> OutdatedProxies;
 	OutdatedProxies.Reserve(Proxies.Num());
 	for (TObjectPtr<ALandscapeProxy> Proxy : Proxies)
 	{
 		ALandscapeProxy* ValidProxy = Proxy.Get();
-		{
+{
 			const UE::Landscape::EOutdatedDataFlags ProxyOutdatedDataFlags = ValidProxy->GetOutdatedDataFlags();
 			if ((bInMustMatchAllFlags && EnumHasAllFlags(ProxyOutdatedDataFlags, InMatchingOutdatedDataFlags))
 				|| (!bInMustMatchAllFlags && EnumHasAnyFlags(ProxyOutdatedDataFlags, InMatchingOutdatedDataFlags)))
@@ -927,11 +941,11 @@ bool ULandscapeSubsystem::GetActionableMessage(FActionableMessage& OutActionable
 
 			uint32 RemainingFlags = static_cast<uint32>(ProxyAndFlag.Value);
 			while (RemainingFlags != 0)
-			{
+	{
 				uint32 FlagIndex = FBitSet::GetAndClearNextBit(RemainingFlags);
 				++NumOutdatedProxyPerFlag[FlagIndex];
 			}
-		}
+	}
 	});
 	
 	// If more than 1 action is required, go with a BuildAll action
@@ -944,7 +958,7 @@ bool ULandscapeSubsystem::GetActionableMessage(FActionableMessage& OutActionable
 
 		return true;
 	}
-
+	
 	if (int32 OutdatedProxiesCount = NumOutdatedProxyPerFlag[GetOutdatedDataFlagIndex(EOutdatedDataFlags::GrassMaps)])
 	{
 		OutActionableMessage.Message = FText::Format(LOCTEXT("GRASS_MAPS_NEED_TO_BE_REBUILT_FMT", "{0} Landscape {0}|plural(one=actor,other=actors) with grass maps {0}|plural(one=needs,other=need) to be rebuilt"), OutdatedProxiesCount);
