GPULM: Add missing Landscape GPULM shader types.

#rnx

diff --git a/Engine/Source/Runtime/Landscape/Private/LandscapeRender.cpp b/Engine/Source/Runtime/Landscape/Private/LandscapeRender.cpp
--- a/Engine/Source/Runtime/Landscape/Private/LandscapeRender.cpp
+++ b/Engine/Source/Runtime/Landscape/Private/LandscapeRender.cpp
@@ -4143,6 +4143,8 @@ public:
 			FName(TEXT("FVLMVoxelizationPS")),
 			FName(TEXT("FLightmapGBufferVS")),
 			FName(TEXT("FLightmapGBufferPS")),
+			FName(TEXT("FGPULightmassCHS")),
+			FName(TEXT("FGPULightmassCHS_AHS")),
 		};
 		return ShaderTypes;
 	}
