Fixed regression that made landscape brush disappear while in "Game mode" ('G' keybind). The intention of the original change is kept by making sure the landscape tools are not rendered in scene captures

#tests editor + scene captures

diff --git a/Engine/Source/Runtime/Landscape/Private/LandscapeRender.cpp b/Engine/Source/Runtime/Landscape/Private/LandscapeRender.cpp
--- a/Engine/Source/Runtime/Landscape/Private/LandscapeRender.cpp
+++ b/Engine/Source/Runtime/Landscape/Private/LandscapeRender.cpp
@@ -2840,7 +2840,7 @@ void FLandscapeComponentSceneProxy::GetDynamicMeshElements(const TArray<const FS
 
 #if WITH_EDITOR
 			  // Extra render passes for landscape tools
-			if (GLandscapeEditModeActive && ViewSpecificFamily.EngineShowFlags.Editor)
+			if (GLandscapeEditModeActive && !View->bIsSceneCapture)
 			{
 				// Region selection
 				if (EditToolRenderData.SelectedType)
