Remove VF_SUPPORTS_WORLD_POSITION_SHADER_OFFSETS=0 from a landscape vertex factory, as landscapes do support WPO

diff --git a/Engine/Source/Runtime/Landscape/Private/LandscapeRender.cpp b/Engine/Source/Runtime/Landscape/Private/LandscapeRender.cpp
--- a/Engine/Source/Runtime/Landscape/Private/LandscapeRender.cpp
+++ b/Engine/Source/Runtime/Landscape/Private/LandscapeRender.cpp
@@ -3623,8 +3623,6 @@ void FLandscapeVertexFactory::ModifyCompilationEnvironment(const FVertexFactoryS
 	FVertexFactory::ModifyCompilationEnvironment(Parameters, OutEnvironment);
 
 	OutEnvironment.SetDefine(TEXT("VF_SUPPORTS_PRIMITIVE_SCENE_DATA"), Parameters.VertexFactoryType->SupportsPrimitiveIdStream() && UseGPUScene(Parameters.Platform, GetMaxSupportedFeatureLevel(Parameters.Platform)));
-	// landscape does not support WPO, so we can remove one interpolator
-	OutEnvironment.SetDefine(TEXT("VF_SUPPORTS_WORLD_POSITION_SHADER_OFFSETS"), TEXT("0"));
 
 	// Make sure landscape vertices go back to local space so that we have consistency between the transform on normals and geometry
 	OutEnvironment.SetDefine(TEXT("RAY_TRACING_DYNAMIC_MESH_IN_LOCAL_SPACE"), TEXT("1"));
