Introducing landscape batched merge (activated by CVar : landscape.EditLayersLocalMerge.Enable 2): complete refactor of the landscape edit layers merge algorithm to solve the taxing GPU memory requirements that global merge has and serve as the basis for future features (new weight-blending methods, partial exports, ...) and performance and workflow improvements (deterministic editing of partially-loaded worlds, local updates, removal of useless weightmap allocations, ...)
Unlike local merge (landscape.EditLayersLocalMerge.Enable 1), batched merge is backwards-compatible with BP brushes and will therefore replace both global merge and local merge in the near future. However, in order to maintain full backwards-compatibility, existing BP brushes will still act as they do now wrt determinism, requiring the entire world to be loaded by default. This means that as long as a landscape uses existing and not-yet converted brushes like Landmass, Water, etc., it will effectively act as global merge, with the same memory constraints as the current ones.
It also supports both vertical (in between edit layers) and horizontal (in between paint layers) blending, which allows legacy weight-blending to still work but opens up several possibilities in term of paint layers blending improvements.

Details:
* ILandscapeEditLayerRenderer is an interface (UInterface, actually) to implement for any element that needs to write heightmaps/weightmaps/visibility maps. These can be anything (edit layers, BP brushes, actors, components, ...), which streamlines the rendering/blending capabilities while maintaining backwards-compatibility
* UE::Landscape::EditLayers::IEditLayerRendererProvider is an interface to implement in order to provide a sequence of ILandscapeEditLayerRenderer to the landscape. This will be called on-demand when merging the landscape, which means the landscape dynamically retrieves its active renderers instead of relying on the landscape actor's data to do so. This solves the issue of having to checkout the landscape when adding/removing an element (e.g. a BP brush) to it. BP brushes are still supported, but things like landscape patch components can now be trivially added to a sequence of render operations without checking out the landscape actor
* The merge operation is now generic and is split into several render batches (instead of one large render, the size of the entire loaded landscape), hence the name. Each batch is sized automatically so that it contains a minimal amount of weightmaps and a minimal amount of landscape components.
* A component dependency analysis is run in order to maximize the amount of components to render in a given batch while trying to respect the optimal batch size (CVar landscape.EditLayersLocalMerge.MaxResolutionPerRenderBatch, default value : 1024) If a given component depends on components outside the batch, though, it will either be rendered in another batch (if possible) or the batches will be made large enough to render these (which allows the system to be compatible with existing BP brushes such as Water, but also prevents batched merge to solve the GPU memory requirements that global merge has)
* The merge is split into the following phases :
1. Retrieval of edit layer renderers and their associated states for the various merge types (supported/enabled for heightmaps/weightmaps/visibility)
At this point, the user can selectively change the enabled state of the supported renderers, which allows to customize the merge operation without mutating the landscape state (e.g. only keeping the layers under edit layer N enabled)
2. Trim the unneeded renderers for the requested merge operation (disabled renderers, unneeded weightmaps, etc.)
3. (weightmap merge only) Analyze the renderers' weightmaps usage and divide the renders into several "render groups", that is, a set of weightmaps rendered together because they have a(n horizontal) dependency. Each render group will be rendered separately (i.e. reduces the memory requirements of weight-blending to the minimum)
4. Each renderer provides a list of "render items", which contain locality information (where does this render item outputs data? what input data does it require? does it output heightmap/weightmap/visibility? what weightmap(s) does it output?) and a full component dependency analysis is then run on all renderers/requested components
5. The work is split into several batches, each batch into several render steps : for each render group, there's one step per active renderer, and one last step : a simple callback, which the user can attach to, which allows to retrieve the data that has been rendered for this render group for this batch, i.e. a portion of the merge landscape data, before the render target it was using gets repurposed for another batch/render group.
Currently, this is where the readback of heightmaps/weightmaps to the final texture is done and eventually, the same merge can be repurposed by using another callback and, for example, copying into a transient render target (partial renders), or an image on disk (partial exports)
6. The list of render batches and render steps is executed. Some renderers are automatically added to mix (e.g. default data rendering, to provide the initial clear, at the beginning, and "legacy" weight-blending and normals rendering, both at the end of the merge process) and each batch also has a list of "validity" render targets (one per target layer), which is nothing else than a binary "stencil" buffer, masking the components that are *not* in a given batch. These are not implemented as stencil buffers, though, because we eventually might need access to them on the BP side, so they are plain UTextureRenderTarget2D.
* Several CVars (landscape.BatchedMerge.VisualLog.*) were added to help debug the merge process
* In order to improve rendering time, a new utility shader (CopyQuadsMultiSourcePS) was made  to copy several (up to 63) source textures into an atlas, which optimizes quite a bit the edit layer->atlas copy operation

Misc changes:
* Improvements/fixes to landscape-related GPU tags
* Fixed landscape.DumpWeightmapDiff which would leave txt/png files for weightmap allocations that haven't actually changed if one of the other channels of the weightmap had some data changes

Code details:
* OnFullHeightmapRenderDoneDelegate has been replaced with OnEditLayersMerged and is now called for each batch render
* Non-public API landscape utils have moved to LandscapeUtilsPrivate.h/.cpp
* Fixed edge case of a map that was saved with an invalid ULandscapeEditLayer (e.g. after having uninstalled a plugin where the base class of a layer was defined): we now clear the layer and make sure we don't fall into a situation where there's no edit layer altogether (which is not allowed and usually enforced by the UX)
* Added parameter to ReallocateWeightmaps in order to restrict the re-allocation of weightmaps to a given set of components. This is to avoid a situation where we're trying to add weightmap allocations to components that are not part of the merge. This will eventually go away when we refactor the weightmap allocation system in the editor
* Since the whole thing has to run on the game thread (because of BP calls), a new "scratch render target" system has been added (ULandscapeEditResourcesSubsystem / ULandscapeScratchRenderTarget) in order to help the various merge steps to request/release URenderTargets (2D and 2D array) and manage their RHI access. This would normally be done via RDG but we don't have access to it outside the render thread. This is shared across landscapes, which minimizes the amount of render targets needed when editing multiple landscapes at the same time.

#tests Editor, cooked game

diff --git a/Engine/Plugins/Experimental/LandscapePatch/Source/LandscapePatch/Private/LandscapePatchManager.cpp b/Engine/Plugins/Experimental/LandscapePatch/Source/LandscapePatch/Private/LandscapePatchManager.cpp
--- a/Engine/Plugins/Experimental/LandscapePatch/Source/LandscapePatch/Private/LandscapePatchManager.cpp
+++ b/Engine/Plugins/Experimental/LandscapePatch/Source/LandscapePatch/Private/LandscapePatchManager.cpp
@@ -46,7 +46,7 @@ namespace LandscapePatchManagerLocals
 		TEXT("LandscapePatch.MigrateLegacyListToPrioritySystem"),
 		TEXT("For all patch managers, make any patches in their patch list be directly bound to their edit layer, and "
 			"set the patch priorities according to their index."),
-		FConsoleCommandDelegate::CreateLambda([]() 
+		FConsoleCommandDelegate::CreateLambda([]()
 	{
 		const FScopedTransaction Transaction(MigratePatchesTransactionName);
 
@@ -76,7 +76,7 @@ namespace LandscapePatchManagerLocals
 #endif // WITH_EDITOR
 
 	// Removes invalid patches from the list. This happens automatically when applying patches.
-	void FilterLegacyRegisteredPatches(TArray<TSoftObjectPtr<ULandscapePatchComponent>>& PatchComponents, 
+	void FilterLegacyRegisteredPatches(TArray<TSoftObjectPtr<ULandscapePatchComponent>>& PatchComponents,
 		TMap<TSoftObjectPtr<ULandscapePatchComponent>, int32>& PatchToIndex, const ALandscapePatchManager* ThisPatchManager)
 	{
 		// Used for removing invalid brushes. We remove from the index map immediately but then remove
@@ -126,32 +126,32 @@ namespace LandscapePatchManagerLocals
 
 			// Make sure the patch has this manager set as its patch manager.
 			PRAGMA_DISABLE_DEPRECATION_WARNINGS
-			if (Component->GetPatchManager() != ThisPatchManager)
-			PRAGMA_ENABLE_DEPRECATION_WARNINGS
-			{
-				UE_LOG(LogLandscapePatch, Warning, TEXT("ALandscapePatchManager: Found a patch whose patch manager is not set "
-						"to a patch manager that contains it. It will be removed."));
-				RemoveComponentFromIndexMap(Component);
-				bHaveInvalidPatches = true;
-				continue;
-			}
+				if (Component->GetPatchManager() != ThisPatchManager)
+					PRAGMA_ENABLE_DEPRECATION_WARNINGS
+				{
+						UE_LOG(LogLandscapePatch, Warning, TEXT("ALandscapePatchManager: Found a patch whose patch manager is not set "
+								"to a patch manager that contains it. It will be removed."));
+						RemoveComponentFromIndexMap(Component);
+						bHaveInvalidPatches = true;
+					continue;
+				}
 
-			if (!Component->IsPatchInWorld())
-			{
-				UE_LOG(LogLandscapePatch, Warning, TEXT("ALandscapePatchManager: Found a non-world patch in patch manager. It will be removed."));
-				RemoveComponentFromIndexMap(Component);
-				bHaveInvalidPatches = true;
-				continue;
-			}
+					if (!Component->IsPatchInWorld())
+					{
+						UE_LOG(LogLandscapePatch, Warning, TEXT("ALandscapePatchManager: Found a non-world patch in patch manager. It will be removed."));
+						RemoveComponentFromIndexMap(Component);
+						bHaveInvalidPatches = true;
+						continue;
+					}
 		}
 
 		if (bHaveInvalidPatches)
 		{
 			PatchComponents.RemoveAll([ThisPatchManager](TSoftObjectPtr<ULandscapePatchComponent> Component) {
 				PRAGMA_DISABLE_DEPRECATION_WARNINGS
-				return !Component.IsValid() || Component->GetPatchManager() != ThisPatchManager || !Component->IsPatchInWorld();
+					return !Component.IsValid() || Component->GetPatchManager() != ThisPatchManager || !Component->IsPatchInWorld();
 				PRAGMA_ENABLE_DEPRECATION_WARNINGS
-				});
+			});
 			// Update forward indices
 			for (int32 i = MinRemovedIndex; i < PatchComponents.Num(); ++i)
 			{
@@ -193,7 +193,7 @@ ALandscapePatchManager::ALandscapePatchManager(const FObjectInitializer& ObjectI
 #endif
 }
 
-void ALandscapePatchManager::Initialize_Native(const FTransform & InLandscapeTransform,
+void ALandscapePatchManager::Initialize_Native(const FTransform& InLandscapeTransform,
 	const FIntPoint& InLandscapeSize,
 	const FIntPoint& InLandscapeRenderTargetSize)
 {
@@ -218,12 +218,27 @@ UTextureRenderTarget2D* ALandscapePatchManager::RenderLayer_Native(const FLandsc
 			continue;
 		}
 
-		BrushParameters.CombinedResult = Component->RenderLayer_Native(BrushParameters, GetHeightmapCoordsToWorld()); 
+		BrushParameters.CombinedResult = Component->RenderLayer_Native(BrushParameters, GetHeightmapCoordsToWorld());
 	}
 
 	return BrushParameters.CombinedResult;
 }
 
+#if WITH_EDITOR
+TArray<UE::Landscape::EditLayers::FEditLayerRenderItem> ALandscapePatchManager::GetRenderItems(const ULandscapeInfo* InLandscapeInfo) const
+{
+	using namespace UE::Landscape::EditLayers;
+
+	TArray<FEditLayerRenderItem> RenderItems = Super::GetRenderItems(InLandscapeInfo);
+	check(RenderItems.Num() == 1);
+
+	// For now, this brush requires the entire landscape to be loaded to work deterministically, so we force the input area to be "infinite" : 
+	RenderItems[0].SetInputWorldArea(FInputWorldArea::CreateInfinite());
+
+	return RenderItems;
+}
+#endif // WITH_EDITOR
+
 void ALandscapePatchManager::SetTargetLandscape(ALandscape* InTargetLandscape)
 {
 #if WITH_EDITOR
@@ -294,24 +309,24 @@ void ALandscapePatchManager::AddPatch(ULandscapePatchComponent* Patch)
 			Modify();
 			TSoftObjectPtr<ULandscapePatchComponent> PatchSoftPtr(Patch);
 			PatchComponents.Add(PatchSoftPtr);
-			PatchToIndex.Add(PatchSoftPtr, PatchComponents.Num()-1);
+			PatchToIndex.Add(PatchSoftPtr, PatchComponents.Num() - 1);
 		}
 
 		PRAGMA_DISABLE_DEPRECATION_WARNINGS
-		if (Patch->GetPatchManager() != this)
-		{
-			UE_LOG(LogLandscapePatch, Warning, TEXT("ALandscapePatchManager::AddPatch: Added patch does not have this manager set "
-				"as its manager. Patches are typically added to managers by setting the manager on the patch. "
-				"(Package: % s, Actor : % s)"), *GetNameSafe(Patch->GetPackage()), *GetNameSafe(Patch->GetAttachmentRootActor()));
-		}
+			if (Patch->GetPatchManager() != this)
+			{
+				UE_LOG(LogLandscapePatch, Warning, TEXT("ALandscapePatchManager::AddPatch: Added patch does not have this manager set "
+					"as its manager. Patches are typically added to managers by setting the manager on the patch. "
+					"(Package: % s, Actor : % s)"), *GetNameSafe(Patch->GetPackage()), *GetNameSafe(Patch->GetAttachmentRootActor()));
+			}
 		PRAGMA_ENABLE_DEPRECATION_WARNINGS
 
-		// No need to update if the patch is disabled. Important to avoid needlessly updating while dragging a blueprint with
-		// a disabled patch (since construction scripts constantly add and remove).
-		if (Patch->IsEnabled())
-		{
-			RequestLandscapeUpdate(!UE::GetIsEditorLoadingPackage());
-		}
+			// No need to update if the patch is disabled. Important to avoid needlessly updating while dragging a blueprint with
+			// a disabled patch (since construction scripts constantly add and remove).
+			if (Patch->IsEnabled())
+			{
+				RequestLandscapeUpdate(!UE::GetIsEditorLoadingPackage());
+			}
 	}
 }
 
@@ -334,7 +349,7 @@ bool ALandscapePatchManager::RemovePatch(ULandscapePatchComponent* Patch)
 					PatchToIndex.Add(PatchComponents[i], i);
 				}
 			}
-			
+
 		}
 
 		// No need to update if the patch was already disabled.Important to avoid needlessly updating while dragging 
@@ -434,10 +449,10 @@ void ALandscapePatchManager::MigrateToPrioritySystemAndDelete()
 			Priority += PriorityStep;
 
 			PRAGMA_DISABLE_DEPRECATION_WARNINGS
-			Patch->SetPatchManager(nullptr);
+				Patch->SetPatchManager(nullptr);
 			PRAGMA_ENABLE_DEPRECATION_WARNINGS
 
-			Patch->FixBindings();
+				Patch->FixBindings();
 		}
 
 		PatchComponents.Empty();
@@ -449,7 +464,7 @@ void ALandscapePatchManager::MigrateToPrioritySystemAndDelete()
 
 	// This will warn the user if there are still references to the manager
 	UTypedElementSelectionSet* SelectionSet = GEditor->GetEditorSubsystem<ULevelEditorSubsystem>()->GetSelectionSet();
-	GUnrealEd->DeleteActors({this}, GetWorld(), SelectionSet);
+	GUnrealEd->DeleteActors({ this }, GetWorld(), SelectionSet);
 }
 
 
@@ -477,7 +492,7 @@ bool ALandscapePatchManager::AffectsWeightmap() const
 
 	return AnyOfPatchComponents(PatchComponents,
 		[](ULandscapePatchComponent& InComponent) { return InComponent.AffectsWeightmap(); });
-		
+
 }
 
 bool ALandscapePatchManager::AffectsWeightmapLayer(const FName& InLayerName) const
@@ -488,7 +503,7 @@ bool ALandscapePatchManager::AffectsWeightmapLayer(const FName& InLayerName) con
 	{
 		return false;
 	}
-	
+
 	return AnyOfPatchComponents(PatchComponents,
 		[InLayerName](ULandscapePatchComponent& InComponent) { return InComponent.AffectsWeightmapLayer(InLayerName); });
 }
@@ -506,6 +521,19 @@ bool ALandscapePatchManager::AffectsVisibilityLayer() const
 		[](ULandscapePatchComponent& InComponent) { return InComponent.AffectsVisibilityLayer(); });
 }
 
+bool ALandscapePatchManager::CanAffectWeightmapLayer(const FName& InLayerName) const
+{
+	using namespace LandscapePatchManagerLocals;
+
+	if (!CanAffectWeightmap())
+	{
+		return false;
+	}
+
+	return AnyOfPatchComponents(PatchComponents,
+		[InLayerName](ULandscapePatchComponent& InComponent) { return InComponent.CanAffectWeightmapLayer(InLayerName); });
+}
+
 void ALandscapePatchManager::GetRenderDependencies(TSet<UObject*>& OutDependencies)
 {
 	for (TSoftObjectPtr<ULandscapePatchComponent>& Component : PatchComponents)
@@ -520,7 +548,6 @@ void ALandscapePatchManager::GetRenderDependencies(TSet<UObject*>& OutDependenci
 			Component->GetRenderDependencies(OutDependencies);
 		}
 	}
-
 }
 
 void ALandscapePatchManager::PostEditUndo()
@@ -583,13 +610,13 @@ void ALandscapePatchManager::CheckForErrors()
 	bool bHavePatchWithIncorrectManager = false;
 	for (TSoftObjectPtr<ULandscapePatchComponent>& Component : PatchComponents)
 	{
-		PRAGMA_DISABLE_DEPRECATION_WARNINGS
+PRAGMA_DISABLE_DEPRECATION_WARNINGS
 		if (Component.IsValid() && Component->GetPatchManager() != this)
 		{
 			bHavePatchWithIncorrectManager = true;
 			break;
 		}
-		PRAGMA_ENABLE_DEPRECATION_WARNINGS
+PRAGMA_ENABLE_DEPRECATION_WARNINGS
 	}
 
 	if (bHavePatchWithIncorrectManager)
@@ -600,18 +627,19 @@ void ALandscapePatchManager::CheckForErrors()
 				"least one patch whose patch manager pointer is set incorrectly. These patches should be removed from the manager."
 				"(Package: {Package}, Manager: {Actor})."), GetPackageAndActorArgs())))
 			->AddToken(FActionToken::Create(LOCTEXT("FixPatchesButton", "Fix patches"), FText(),
-				FOnActionTokenExecuted::CreateWeakLambda(this, [this]() 
-		{
-			// Hard to say whether this should be in a transaction, or even be an action, because this happens
-			// automatically on the next landscape update... We'll stick with having it be user triggerable but
-			// not undoable.
-
-			PatchComponents.RemoveAll([this](TSoftObjectPtr<ULandscapePatchComponent> Component) {
-				PRAGMA_DISABLE_DEPRECATION_WARNINGS
-				return Component.IsValid() && Component->GetPatchManager() != this;
-				PRAGMA_ENABLE_DEPRECATION_WARNINGS
-				});
-		})));
+				FOnActionTokenExecuted::CreateWeakLambda(this, [this]()
+				{
+					// Hard to say whether this should be in a transaction, or even be an action, because this happens
+					// automatically on the next landscape update... We'll stick with having it be user triggerable but
+					// not undoable.
+
+					PatchComponents.RemoveAll([this](TSoftObjectPtr<ULandscapePatchComponent> Component) 
+					{
+PRAGMA_DISABLE_DEPRECATION_WARNINGS
+						return Component.IsValid() && Component->GetPatchManager() != this;
+PRAGMA_ENABLE_DEPRECATION_WARNINGS
+					});
+				})));
 	}
 
 	if (PatchComponents.Num() > 0)
@@ -623,7 +651,7 @@ void ALandscapePatchManager::CheckForErrors()
 				"use Priority for ordering. You can use LandscapePatch.MigrateLegacyListToPrioritySystem to "
 				"fix this. (Package: {Package}, Manager: {Actor})."), GetPackageAndActorArgs())))
 			->AddToken(FActionToken::Create(LOCTEXT("MigrateToGuidsButton", "Migrate to guid system"), FText(),
-				FOnActionTokenExecuted::CreateWeakLambda(this, [this]() 
+				FOnActionTokenExecuted::CreateWeakLambda(this, [this]()
 		{
 			const FScopedTransaction Transaction(MigratePatchesTransactionName);
 			MigrateToPrioritySystemAndDelete();
diff --git a/Engine/Plugins/Experimental/LandscapePatch/Source/LandscapePatch/Public/LandscapePatchEditLayer.h b/Engine/Plugins/Experimental/LandscapePatch/Source/LandscapePatch/Public/LandscapePatchEditLayer.h
--- a/Engine/Plugins/Experimental/LandscapePatch/Source/LandscapePatch/Public/LandscapePatchEditLayer.h
+++ b/Engine/Plugins/Experimental/LandscapePatch/Source/LandscapePatch/Public/LandscapePatchEditLayer.h
@@ -45,7 +45,7 @@ public:
 	virtual double GetHighestPatchPriority();
 
 	void RequestLandscapeUpdate(bool bInUserTriggered = false);
-#endif
+#endif // WITH_EDITOR
 
 	inline static const double PATCH_PRIORITY_BASE = 1000;
 
@@ -58,8 +58,16 @@ public:
 	virtual void OnLayerRemoved() override;
 	// TODO: Remove once our guid gets initialized somewhere else
 	virtual void OnLayerCreated(FLandscapeLayer& Layer);
-	
+
 #if WITH_EDITOR
+	// TODO [jonathan.bard] : The following implementation is completely bogus, for now, until this supports batched merge : 
+	//~ Begin ILandscapeEditLayerRenderer implementation
+	virtual void GetRendererStateInfo(const ULandscapeInfo* InLandscapeInfo,
+		UE::Landscape::EditLayers::FEditLayerTargetTypeState& OutSupportedTargetTypeState, UE::Landscape::EditLayers::FEditLayerTargetTypeState& OutEnabledTargetTypeState, TArray<TSet<FName>>& OutRenderGroups) const override {}
+	virtual TArray<UE::Landscape::EditLayers::FEditLayerRenderItem> GetRenderItems(const ULandscapeInfo* InLandscapeInfo) const override { return {}; }
+	virtual FString GetEditLayerRendererDebugName() const override { return TEXT("LandscapePatchEditLayer"); }
+	//~ End ILandscapeEditLayerRenderer implementation
+	
 	// UObject
 	virtual void PostLoad() override;
 	virtual void PostEditUndo() override;
@@ -74,11 +82,11 @@ protected:
 		const FIntPoint& InLandscapeSize,
 		const FIntPoint& InLandscapeRenderTargetSize) override;
 	virtual UTextureRenderTarget2D* RenderLayerAsBlueprintBrush(const FLandscapeBrushParameters& InParameters) override;
+#endif // WITH_EDITOR
 
-private:
-#endif
 #if WITH_EDITORONLY_DATA
 
+private:
 	// This is the transient list of patches that are bound to this edit layer, usually sorted by priority.
 	//  TSoftObjectPtr is used because it is robust across blueprint actor construction script reruns.
 	UPROPERTY(Transient)
@@ -93,7 +101,8 @@ private:
 	//  is used as a safety in case some lack of notifications puts us in a situation where our patches
 	//  are not all valid and sorted.
 	bool bPatchListDirty = false;
-#endif
+#endif // WITH_EDITORONLY_DATA
+
 #if WITH_EDITOR
 	void UpdatePatchList();
 	void UpdatePatchListIfDirty();
@@ -108,10 +117,11 @@ private:
 
 	// Used in legacy global merge path to pass transform information from Initialize to RenderLayer
 	FTransform HeightmapCoordsToWorld;
-#endif
+#endif // WITH_EDITOR
+
 #if WITH_EDITORONLY_DATA
 	// TODO: Remove this once guid is stored in the base edit layer UObject
 	UPROPERTY()
 	FGuid EditLayerGuid;
-#endif
+#endif // WITH_EDITORONLY_DATA
 };
diff --git a/Engine/Plugins/Experimental/LandscapePatch/Source/LandscapePatch/Public/LandscapePatchManager.h b/Engine/Plugins/Experimental/LandscapePatch/Source/LandscapePatch/Public/LandscapePatchManager.h
--- a/Engine/Plugins/Experimental/LandscapePatch/Source/LandscapePatch/Public/LandscapePatchManager.h
+++ b/Engine/Plugins/Experimental/LandscapePatch/Source/LandscapePatch/Public/LandscapePatchManager.h
@@ -32,6 +32,12 @@ public:
 		const FIntPoint& InLandscapeRenderTargetSize) override;
 	virtual UTextureRenderTarget2D* RenderLayer_Native(const FLandscapeBrushParameters& InParameters) override;
 
+#if WITH_EDITOR
+	//~ Begin ILandscapeEditLayerRenderer implementation
+	virtual TArray<UE::Landscape::EditLayers::FEditLayerRenderItem> GetRenderItems(const ULandscapeInfo* InLandscapeInfo) const override;
+	//~ End ILandscapeEditLayerRenderer implementation
+#endif // WITH_EDITOR
+
 	// Adds the brush to the given landscape, removing it from any previous one. This differs from SetOwningLandscape
 	// in that SetOwningLandscape is called by the landscape itself from AddBrushToLayer to update the manager.
 	UFUNCTION(BlueprintCallable, Category = LandscapeManager)
@@ -79,6 +85,7 @@ public:
 	void MigrateToPrioritySystemAndDelete();
 
 	// ALandscapeBlueprintBrushBase
+	virtual bool CanAffectWeightmapLayer(const FName& InLayerName) const override;
 	virtual bool AffectsHeightmap() const override;
 	virtual bool AffectsWeightmap() const override;
 	virtual bool AffectsWeightmapLayer(const FName& InLayerName) const override;
diff --git a/Engine/Plugins/Experimental/Water/Source/Editor/Private/WaterLandscapeBrush.cpp b/Engine/Plugins/Experimental/Water/Source/Editor/Private/WaterLandscapeBrush.cpp
--- a/Engine/Plugins/Experimental/Water/Source/Editor/Private/WaterLandscapeBrush.cpp
+++ b/Engine/Plugins/Experimental/Water/Source/Editor/Private/WaterLandscapeBrush.cpp
@@ -404,6 +404,21 @@ void AWaterLandscapeBrush::AddReferencedObjects(UObject* InThis, FReferenceColle
 	}
 }
 
+#if WITH_EDITOR
+TArray<UE::Landscape::EditLayers::FEditLayerRenderItem> AWaterLandscapeBrush::GetRenderItems(const ULandscapeInfo* InLandscapeInfo) const
+{
+	using namespace UE::Landscape::EditLayers; 
+
+	TArray<FEditLayerRenderItem> RenderItems = Super::GetRenderItems(InLandscapeInfo);
+	check(RenderItems.Num() == 1);
+
+	// For now, this brush requires the entire landscape to be loaded to work deterministically, so we force the input area to be "infinite" : 
+	RenderItems[0].SetInputWorldArea(FInputWorldArea::CreateInfinite());
+
+	return RenderItems;
+}
+#endif // WITH_EDITOR
+
 void AWaterLandscapeBrush::GetWaterBodies(TSubclassOf<AWaterBody> WaterBodyClass, TArray<AWaterBody*>& OutWaterBodies) const
 {
 	FGetActorsOfType<AWaterBody>()(this, WaterBodyClass, OutWaterBodies);
@@ -484,8 +499,13 @@ void AWaterLandscapeBrush::SetTargetLandscape(ALandscape* InTargetLandscape)
 #endif // WITH_EDITOR
 }
 
-void AWaterLandscapeBrush::OnFullHeightmapRenderDone(UTextureRenderTarget2D* InHeightmapRenderTarget)
+void AWaterLandscapeBrush::OnEditLayersMerged(const FOnLandscapeEditLayersMergedParams& InParams)
 {
+	if (!InParams.bIsHeightmapMerge)
+	{
+		return;
+	}
+
 	// #todo_water [roey]: This needs to be changed when the WaterZone can maintain it's own list of "ground actors" so that we don't needlessly update all water zones.
 	if (UWaterSubsystem* WaterSubsystem = UWaterSubsystem::GetWaterSubsystem(GetWorld()))
 	{
@@ -497,14 +517,14 @@ void AWaterLandscapeBrush::SetOwningLandscape(ALandscape* InOwningLandscape)
 {
 	if (OwningLandscape != nullptr)
 	{
-		OwningLandscape->OnFullHeightmapRenderDoneDelegate().RemoveAll(this);
+		OwningLandscape->OnEditLayersMerged().RemoveAll(this);
 	}
 
 	Super::SetOwningLandscape(InOwningLandscape);
 
 	if (OwningLandscape != nullptr)
 	{
-		OwningLandscape->OnFullHeightmapRenderDoneDelegate().AddUObject(this, &AWaterLandscapeBrush::OnFullHeightmapRenderDone);
+		OwningLandscape->OnEditLayersMerged().AddUObject(this, &AWaterLandscapeBrush::OnEditLayersMerged);
 	}
 }
 
diff --git a/Engine/Plugins/Experimental/Water/Source/Editor/Private/WaterLandscapeBrush.h b/Engine/Plugins/Experimental/Water/Source/Editor/Private/WaterLandscapeBrush.h
--- a/Engine/Plugins/Experimental/Water/Source/Editor/Private/WaterLandscapeBrush.h
+++ b/Engine/Plugins/Experimental/Water/Source/Editor/Private/WaterLandscapeBrush.h
@@ -29,6 +29,12 @@ public:
 	virtual void UnregisterAllComponents(bool bForReregister) override;
 	static void AddReferencedObjects(UObject* InThis, FReferenceCollector& Collector);
 
+#if WITH_EDITOR
+	//~ Begin ILandscapeEditLayerRenderer implementation
+	virtual TArray<UE::Landscape::EditLayers::FEditLayerRenderItem> GetRenderItems(const ULandscapeInfo* InLandscapeInfo) const override;
+	//~ End ILandscapeEditLayerRenderer implementation
+#endif // WITH_EDITOR
+
 	UFUNCTION(BlueprintCallable, Category = "Water", meta = (DeterminesOutputType = "WaterBodyClass", DynamicOutputParam = "OutWaterBodies"))
 	void GetWaterBodies(TSubclassOf<AWaterBody> WaterBodyClass, TArray<AWaterBody*>& OutWaterBodies) const;
 
@@ -69,18 +75,18 @@ public:
 	UFUNCTION(BlueprintCallable, Category = "Cache")
 	void ClearActorCache(AActor* InActor);
 
-	UFUNCTION(BlueprintNativeEvent, Category = "Cache", meta = (CallInEditor = "true", DeprecatedFunction, DeprecationMessage = "This event isn't called anymore, the WaterVelocityTexture is now regenerated at runtime (WaterInfoTexture in AWaterZone)."))
+	UFUNCTION(BlueprintNativeEvent, meta = (DeprecatedFunction, DeprecationMessage = "This event isn't called anymore, the WaterVelocityTexture is now regenerated at runtime (WaterInfoTexture in AWaterZone)."))
 	void BlueprintGetRenderTargets(UTextureRenderTarget2D* InHeightRenderTarget, UTextureRenderTarget2D*& OutVelocityRenderTarget);
 	UE_DEPRECATED(all, "This function isn't called anymore, the WaterVelocityTexture is now regenerated at runtime (WaterInfoTexture in AWaterZone)")
 	virtual void BlueprintGetRenderTargets_Native(UTextureRenderTarget2D* InHeightRenderTarget, UTextureRenderTarget2D*& OutVelocityRenderTarget) {}
 
-	UFUNCTION(BlueprintNativeEvent, Category = "Cache", meta = (CallInEditor = "true", DeprecatedFunction, DeprecationMessage = "This event isn't called anymore, the WaterVelocityTexture is now regenerated at runtime (WaterInfoTexture in AWaterZone)."))
+	UFUNCTION(BlueprintNativeEvent, meta = (DeprecatedFunction, DeprecationMessage = "This event isn't called anymore, the WaterVelocityTexture is now regenerated at runtime (WaterInfoTexture in AWaterZone)."))
 	void BlueprintOnRenderTargetTexturesUpdated(UTexture2D* VelocityTexture);
 	UE_DEPRECATED(all, "This function isn't called anymore, the WaterVelocityTexture is now regenerated at runtime (WaterInfoTexture in AWaterZone)")
 	virtual void BlueprintOnRenderTargetTexturesUpdated_Native(UTexture2D* VelocityTexture) {}
 
 	UE_DEPRECATED(all, "This function is now useless, the WaterVelocityTexture is now regenerated at runtime (WaterInfoTexture in AWaterZone).")
-	UFUNCTION(BlueprintCallable, Category = "Cache", meta = (CallInEditor = "true", DeprecatedFunction, DeprecationMessage = "This function is now useless, the WaterVelocityTexture is now regenerated at runtime (WaterInfoTexture in AWaterZone)."))
+	UFUNCTION(BlueprintCallable, meta = (DeprecatedFunction, DeprecationMessage = "This function is now useless, the WaterVelocityTexture is now regenerated at runtime (WaterInfoTexture in AWaterZone)."))
 	void ForceWaterTextureUpdate() {}
 
 	void SetTargetLandscape(ALandscape* InTargetLandscape);
@@ -122,7 +128,7 @@ private:
 	void ClearActors();
 	bool IsActorAffectingLandscape(AActor* Actor) const;
 	void RegisterDelegates();
-	void OnFullHeightmapRenderDone(UTextureRenderTarget2D* HeightmapRenderTarget);
+	void OnEditLayersMerged(const struct FOnLandscapeEditLayersMergedParams& InParams);
 	void OnWaterBrushActorChanged(const IWaterBrushActorInterface::FWaterBrushActorChangedEventParams& InParams);
 	void OnActorsAffectingLandscapeChanged();
 	void OnLevelActorAdded(AActor* InActor);
diff --git a/Engine/Shaders/Private/Landscape/LandscapeEditLayersUtils.usf b/Engine/Shaders/Private/Landscape/LandscapeEditLayersUtils.usf
new file mode 100644
--- /dev/null
+++ b/Engine/Shaders/Private/Landscape/LandscapeEditLayersUtils.usf
@@ -0,0 +1,163 @@
+// Copyright Epic Games, Inc. All Rights Reserved.
+
+#include "../Common.ush"
+#include "LandscapeCommon.ush"
+
+// ----------------------------------------------------------------------------------
+
+#if defined (__INTELLISENSE__)
+// Uncomment the appropriate define for enabling syntax highlighting with HLSL Tools for Visual Studio : 
+//#define MARK_VALIDITY 1
+//#define COPY_QUADS 1
+//#define COPY_QUADS_MULTISOURCE 1
+//#define USE_SWIZZLE_MASK 1
+#endif // defined (__INTELLISENSE__)
+
+
+// ----------------------------------------------------------------------------------
+
+#if COPY_QUADS
+// CopyQuadsPS inputs/outputs :
+
+Texture2D<float4> InSourceTexture;
+#endif // COPY_QUADS
+
+#if COPY_QUADS_MULTISOURCE
+// CopyQuadsMultiSourcePS inputs/outputs :
+
+// Helper macro for handling an array of N textures (currently using 63 because 64 is the maximum supported number of SRV on all platforms needing to run this) :
+#define FOREACH_SOURCETEXTURE(op) \
+	op(0) \
+	op(1) \
+	op(2) \
+	op(3) \
+	op(4) \
+	op(5) \
+	op(6) \
+	op(7) \
+	op(8) \
+	op(9) \
+	op(10) \
+	op(11) \
+	op(12) \
+	op(13) \
+	op(14) \
+	op(15) \
+	op(16) \
+	op(17) \
+	op(18) \
+	op(19) \
+	op(20) \
+	op(21) \
+	op(22) \
+	op(23) \
+	op(24) \
+	op(25) \
+	op(26) \
+	op(27) \
+	op(28) \
+	op(29) \
+	op(30) \
+	op(31) \
+	op(32) \
+	op(33) \
+	op(34) \
+	op(35) \
+	op(36) \
+	op(37) \
+	op(38) \
+	op(39) \
+	op(40) \
+	op(41) \
+	op(42) \
+	op(43) \
+	op(44) \
+	op(45) \
+	op(46) \
+	op(47) \
+	op(48) \
+	op(49) \
+	op(50) \
+	op(51) \
+	op(52) \
+	op(53) \
+	op(54) \
+	op(55) \
+	op(56) \
+	op(57) \
+	op(58) \
+	op(59) \
+	op(60) \
+	op(61) \
+	op(62) 
+
+Buffer<uint2> InQuadInfos; // .x = source texture index, .y = swizzle mask
+
+// Define N input textures :
+#define DECLARE_SOURCE_TEXTURE(n) Texture2D<float4> InSourceTexture_##n;
+FOREACH_SOURCETEXTURE(DECLARE_SOURCE_TEXTURE)
+#undef DECLARE_SOURCE_TEXTURE
+#endif // COPY_QUADS_MULTISOURCE
+
+// ----------------------------------------------------------------------------------
+// Pixel shaders : 
+
+#if MARK_VALIDITY
+
+void MarkValidityPS(out float OutValidity : SV_Target0)
+{
+	OutValidity = 1.0f;
+}
+
+#endif // MARK_VALIDITY
+
+
+// ----------------------------------------------------------------------------------
+
+#if COPY_QUADS
+
+void CopyQuadsPS(float4 SvPosition : SV_POSITION, float2 QuadUV : TEXCOORD0, out float4 OutColor : SV_Target0)
+{
+	OutColor = InSourceTexture.SampleLevel(GlobalPointClampedSampler, QuadUV, 0);
+}
+
+#endif // COPY_QUADS
+
+
+// ----------------------------------------------------------------------------------
+
+#if COPY_QUADS_MULTISOURCE
+
+void CopyQuadsMultiSourcePS(float4 SvPosition : SV_POSITION, float2 SourceUV : TEXCOORD0, float2 UnusedQuadUV : TEXCOORD1, float InRectIndex : RECT_INDEX, out float4 OutColor : SV_Target0)
+{
+	OutColor = 0.0f;
+
+	uint QuadIndex = (uint)InRectIndex;
+	int SourceTextureIndex = InQuadInfos[QuadIndex].x;
+
+	BRANCH
+	switch (SourceTextureIndex)
+	{
+#define IMPLEMENT_SAMPLE_SWITCH_CASE(n) case n: OutColor = InSourceTexture_##n.SampleLevel(GlobalPointClampedSampler, SourceUV, 0); break; 
+	FOREACH_SOURCETEXTURE(IMPLEMENT_SAMPLE_SWITCH_CASE);
+#undef IMPLEMENT_SAMPLE_SWITCH_CASE
+
+	default: 
+		break;
+	}
+
+	// Swizzle channels as requested by ChannelSwizzleMask:
+#if USE_SWIZZLE_MASK
+	uint ChannelSwizzleMask = (uint)InQuadInfos[QuadIndex].y;
+	float4 SwizzledColor = 0.0f;
+	[unroll]
+	for (uint DestinationChannelIndex = 0; DestinationChannelIndex < 4; ++DestinationChannelIndex)
+	{
+		uint SourceChannelIndex = ((3u << (DestinationChannelIndex * 2u)) & ChannelSwizzleMask) >> (DestinationChannelIndex * 2u);
+		SwizzledColor[DestinationChannelIndex] = OutColor[SourceChannelIndex];
+	}
+	OutColor = SwizzledColor;
+#endif // USE_SWIZZLE_MASK
+}
+
+#endif // COPY_QUADS_MULTISOURCE
\ No newline at end of file
diff --git a/Engine/Shaders/Private/Landscape/LandscapeEditLayersWeightmaps.usf b/Engine/Shaders/Private/Landscape/LandscapeEditLayersWeightmaps.usf
new file mode 100644
--- /dev/null
+++ b/Engine/Shaders/Private/Landscape/LandscapeEditLayersWeightmaps.usf
@@ -0,0 +1,242 @@
+// Copyright Epic Games, Inc. All Rights Reserved.
+
+#include "../Common.ush"
+
+// ----------------------------------------------------------------------------------
+
+#if defined (__INTELLISENSE__)
+// Uncomment the appropriate define for enabling syntax highlighting with HLSL Tools for Visual Studio : 
+//#define MERGE_EDIT_LAYER 1
+//#define PERFORM_LEGACY_WEIGHT_BLENDING 1
+#define PACK_WEIGHTMAP 1
+//#define GENERATE_MIPS 1
+#endif // defined (__INTELLISENSE__)
+
+
+// ----------------------------------------------------------------------------------
+
+/** EEditLayerWeightmapBlendMode enum */
+#define EEDITLAYERWEIGHTMAPBLENDMODE_ADDITIVE 0
+#define EEDITLAYERWEIGHTMAPBLENDMODE_SUBTRACTIVE 1
+
+/** EWeightmapPaintLayerFlags enum */
+#define EWEIGHTMAPPAINTLAYERFLAGS_ISVISIBILITYLAYER (1 << 0)
+#define EWEIGHTMAPPAINTLAYERFLAGS_ISWEIGHTBLENDED (1 << 1)
+
+#if MERGE_EDIT_LAYER
+// MergeEditLayerPS inputs/outputs :
+
+uint InPaintLayerIndex; // Index of the paint layer being processed
+uint InEditLayerPaintLayerBlendMode; // See EEditLayerWeightmapBlendMode. Paint layer's blend mode, which is per-edit layer and per-paint layer and can therefore differ from one edit layer to another
+float InEditLayerAlpha; // Global alpha value of the edit layer currently being merged
+
+Texture2DArray<float> InCurrentEditLayerWeightmaps; // The paint layers of the current edit layer to merge (1 slice per weightmap in the render group)
+Texture2DArray<float> InPreviousEditLayersWeightmaps; // The result from the merge of all prior edit layers (1 slice per weightmap in the render group)
+// Array that contains per-edit layer / per-paint layer information (e.g. the paint layer's blend mode, which is per-edit layer and can therefore differ from one edit layer to another)
+#endif // MERGE_EDIT_LAYER
+
+#if PERFORM_LEGACY_WEIGHT_BLENDING
+// PerformLegacyWeightBlendingPS inputs/outputs :
+
+// Per-paint layer information : 
+struct FWeightmapPaintLayerInfo
+{
+	uint Flags; // See EWeightmapPaintLayerFlags
+};
+
+uint InPaintLayerIndex; // Index of the paint layer being processed
+uint InNumPaintLayers; // Number of paint layers (i.e. == InPaintLayerInfos.Num())
+// Array that contains per-paint layer information (e.g. paint layer's flags)
+StructuredBuffer<FWeightmapPaintLayerInfo> InPaintLayerInfos;
+Texture2DArray<float> InCurrentEditLayerWeightmaps; // The texture to horizontally blend (1 slice per weightmap in the render group)
+#endif // PERFORM_LEGACY_WEIGHT_BLENDING
+
+#if PACK_WEIGHTMAP
+// PackWeightmapPS inputs/outputs :
+
+int4 InSourceSliceIndices; // For each channel of the output texture (rgba), this indicates the slice index in InSourceWeightmaps where to read from
+uint4 InSourcePixelOffsets[4];  // For each channel, offset to add to the pixel's coordinates to load the proper sample in InSourceWeightmaps (.xy, .zw is unused, only there for alignment purposes)
+uint2 InSubsectionPixelOffset; // Offset of the subsection currently being rendered
+uint InIsAdditive; // = 1 if some channels of this texture have been packed in a previous draw, 0 otherwise
+
+Texture2DArray<float> InSourceWeightmaps; // Source, single-channel, texture to pack
+Texture2D<float4> InWeightmapBeingPacked; // The weightmap being packed, in case the packing operation is additive, i.e. it occurs across multiple draws (contains the channels that have previously been packed)
+#endif // PACK_WEIGHTMAP
+
+#if GENERATE_MIPS
+// GenerateMipsPS inputs/outputs :
+
+uint2 InCurrentMipSubsectionSize; // Size of the the subsection at the currently generated mip level
+
+Texture2D<float4> InSourceWeightmap; // Source weightmap (containing the current mip level - 1)
+#endif // GENERATE_MIPS
+
+
+// ----------------------------------------------------------------------------------
+// Util functions : 
+
+#if PERFORM_LEGACY_WEIGHT_BLENDING
+bool IsWeightBlendedPaintLayer(FWeightmapPaintLayerInfo InPaintLayerInfo)
+{
+	// For the visibility layer, deactivate weight blending altogether : 
+	return (((InPaintLayerInfo.Flags & EWEIGHTMAPPAINTLAYERFLAGS_ISVISIBILITYLAYER) == 0)
+		&& (InPaintLayerInfo.Flags & EWEIGHTMAPPAINTLAYERFLAGS_ISWEIGHTBLENDED));
+}
+#endif // PERFORM_LEGACY_WEIGHT_BLENDING
+
+
+// ----------------------------------------------------------------------------------
+// Pixel shaders : 
+
+#if MERGE_EDIT_LAYER
+void MergeEditLayerPS(in float4 SVPos : SV_POSITION, out float OutWeight : SV_Target0)
+{
+	uint2 TextureCoordinates = floor(SVPos.xy);
+		
+	float CurrentLayerWeight = InCurrentEditLayerWeightmaps.Load(int4(TextureCoordinates, InPaintLayerIndex, 0)); // xy = relative coordinates, z = index in texture array, w = mip level
+	float PreviousLayersWeight = InPreviousEditLayersWeightmaps.Load(int4(TextureCoordinates, InPaintLayerIndex, 0)); // xy = relative coordinates, z = index in texture array, w = mip level
+	OutWeight = PreviousLayersWeight;
+	
+	if (InEditLayerPaintLayerBlendMode == EEDITLAYERWEIGHTMAPBLENDMODE_ADDITIVE)
+	{
+		OutWeight += CurrentLayerWeight;
+	}
+	else if (InEditLayerPaintLayerBlendMode == EEDITLAYERWEIGHTMAPBLENDMODE_SUBTRACTIVE)
+	{
+		OutWeight -= CurrentLayerWeight;
+	}
+}
+#endif // MERGE_EDIT_LAYER
+
+#if PERFORM_LEGACY_WEIGHT_BLENDING
+void PerformLegacyWeightBlendingPS(in float4 SVPos : SV_POSITION, out float OutColor : SV_Target0)
+{
+	uint2 TextureCoordinates = floor(SVPos.xy);
+	OutColor = 0.0f;
+	
+	FWeightmapPaintLayerInfo ActivePaintLayerInfo = InPaintLayerInfos[InPaintLayerIndex];
+	if (!IsWeightBlendedPaintLayer(ActivePaintLayerInfo))
+	{
+		OutColor = InCurrentEditLayerWeightmaps.Load(int4(TextureCoordinates, InPaintLayerIndex, 0)); // xy = relative coordinates, z = index in texture array, w = mip level
+	}
+	else
+	{
+		float ActivePaintLayerWeight = 0.0f;
+		float BlendedWeightsSum = 0.0f;
+
+		LOOP
+		for (int i = 0; i < InNumPaintLayers; ++i)
+		{
+			bool bIsOutputPaintLayer = (i == InPaintLayerIndex);
+			FWeightmapPaintLayerInfo PaintLayerInfo = InPaintLayerInfos[i];
+
+			// Only weight blended (and non-visibility) paint layers participate to weight blending : 
+			if (IsWeightBlendedPaintLayer(PaintLayerInfo))
+			{
+				float Weight = InCurrentEditLayerWeightmaps.Load(int4(TextureCoordinates, i, 0)); // xy = relative coordinates, z = index in texture array, w = mip level
+				if (bIsOutputPaintLayer)
+				{
+					ActivePaintLayerWeight = Weight;
+				}
+				BlendedWeightsSum += Weight;
+			}
+		}
+
+		OutColor = (BlendedWeightsSum > 0.0f) ? saturate(ActivePaintLayerWeight / BlendedWeightsSum) : ActivePaintLayerWeight;
+	}
+}
+#endif // PERFORM_LEGACY_WEIGHT_BLENDING
+
+// ----------------------------------------------------------------------------------
+
+#if PACK_WEIGHTMAP
+void PackWeightmapPS(in float4 SVPos : SV_POSITION, out float4 OutColor : SV_Target0)
+{
+	uint2 LocalCoordinates = floor(SVPos.xy);
+	uint2 SubsectionRelativeTextureCoordinates = LocalCoordinates - InSubsectionPixelOffset;
+	OutColor = 0.0f;
+	
+	if (InIsAdditive != 0)
+	{
+		OutColor = InWeightmapBeingPacked.Load(int3(LocalCoordinates, 0)); // xy = relative coordinates, z = mip level
+	}
+	
+	UNROLL
+	for (uint i = 0; i < 4; ++i)
+	{
+		if (InSourceSliceIndices[i] >= 0)
+		{
+			uint2 TextureCoordinates = SubsectionRelativeTextureCoordinates + InSourcePixelOffsets[i].xy;
+			OutColor[i] = InSourceWeightmaps.Load(int4(TextureCoordinates, InSourceSliceIndices[i], 0)); // xy = relative coordinates, z = index in texture array, w = mip level
+		}
+	}
+}
+#endif // PACK_WEIGHTMAP
+
+
+// ----------------------------------------------------------------------------------
+#if GENERATE_MIPS
+void GenerateMipsPS(in float4 SVPos : SV_POSITION, out float4 OutColor : SV_Target0)
+{
+	uint2 TextureCoordinates = floor(SVPos.xy);
+	
+	float4 SourceSamples[4] =
+	{
+		InSourceWeightmap.Load(int3(2 * TextureCoordinates + int2(+0, +0), 0)),
+		InSourceWeightmap.Load(int3(2 * TextureCoordinates + int2(+1, +0), 0)),
+		InSourceWeightmap.Load(int3(2 * TextureCoordinates + int2(+0, +1), 0)),
+		InSourceWeightmap.Load(int3(2 * TextureCoordinates + int2(+1, +1), 0)),
+	};
+		
+	// Because the borders of each landscape subsection are shared between neighbors, we must ensure that the parent mip's inner row/column of pixels don't contribute, 
+	//  so that pixels on the subsection borders for neighboring subsections for mips have an equal value : 
+	//  9 possible cases (only the samples marked with a * must be kept): 
+	//      bIsMinBorder.x = true
+	//     |                        bIsMaxBorder.x = true
+	//     |                       |
+	//     v                       v
+	// +-------+   +-------+   +-------+
+	// | * :   |   | * : * |   |   : * |
+	// | - + - |...| - + - |...| - + - | <-- bIsMinBorder.y = true
+	// |   :   |   |   :   |   |   :   |
+	// +-------+   +-------+   +-------+
+	//    ...         ...	      ...
+	// +-------+   +-------+   +-------+
+	// | * :   |   | * : * |   |   : * |
+	// | - + - |...| - + - |...| - + - |
+	// | * :   |   | * : * |   |   : * |
+	// +-------+   +-------+   +-------+
+	//    ...         ...	      ...
+	// +-------+   +-------+   +-------+
+	// |   :   |   |   :   |   |   :   |
+	// | - + - |...| - + - |...| - + - | <-- bIsMaxBorder.y = true
+	// | * :   |   | * : * |   |   : * | 
+	// +-------+   +-------+   +-------+
+
+	bool bIsLastMip = all(InCurrentMipSubsectionSize == 1);
+
+	uint2 SubsectionRelativeTextureCoordinates = TextureCoordinates % InCurrentMipSubsectionSize;
+	bool2 bIsMinBorder = (SubsectionRelativeTextureCoordinates == 0);
+	bool2 bIsMaxBorder = (SubsectionRelativeTextureCoordinates == (InCurrentMipSubsectionSize - 1));
+		
+	float SampleWeights[4] =
+	{ 
+		// On the last mip, it's ok to keep all 4 samples : all neighbors components share them : 
+		((bIsMaxBorder.x || bIsMaxBorder.y) && !bIsLastMip) ? 0.0f : 1.0f,
+		((bIsMinBorder.x || bIsMaxBorder.y) && !bIsLastMip) ? 0.0f : 1.0f,
+		((bIsMaxBorder.x || bIsMinBorder.y) && !bIsLastMip) ? 0.0f : 1.0f,
+		((bIsMinBorder.x || bIsMinBorder.y) && !bIsLastMip) ? 0.0f : 1.0f,
+	};
+	
+	float TotalSampleWeight = 0.0f;
+	OutColor = 0.0f;
+	
+	UNROLL
+	for (int i = 0; i < 4; ++i)
+	{
+		OutColor += SourceSamples[i] * SampleWeights[i];
+		TotalSampleWeight += SampleWeights[i];
+	}
+	OutColor /= TotalSampleWeight;
+}
+#endif // GENERATE_MIPS
\ No newline at end of file
diff --git a/Engine/Shaders/Private/LandscapeLayersPS.usf b/Engine/Shaders/Private/LandscapeLayersPS.usf
--- a/Engine/Shaders/Private/LandscapeLayersPS.usf
+++ b/Engine/Shaders/Private/LandscapeLayersPS.usf
@@ -4,17 +4,18 @@
 #include "Landscape/LandscapeCommon.ush"
 
 Texture2D ReadTexture1;
-Texture2D ReadTexture2;							// This one is optional, it's valid when OutputConfig.z == 1
-SamplerState ReadTexture1Sampler;		
+Texture2D ReadTexture2; // This one is optional, it's valid when OutputConfig.z == 1
+SamplerState ReadTexture1Sampler;
 SamplerState ReadTexture2Sampler;
-float4 LayerInfo;								// x == Layer Alpha value, y == Layer Visibility flag, z == Layer Blend Mode (0: Additive Blend Mode, 1: Alpha Blend Mode), w == Unused
-float4 OutputConfig;							// x == ApplyLayerModifiers, y == Alpha Blend Mode or Output weightmap as substractive/additive, z == Use ReadTexture2 and out Delta with ReadTexture1, w == Output Normals or Output weightmap as normalized weight blend
-float2 TextureSize;								// x == ReadTexture1 width, y == ReadTexture1 height
-float3 LandscapeGridScale;						// x == LS Actor DrawScale.X, y == LS Actor DrawScale.y, z == LS Actor DrawScale.z / 128.0f (ZSCALE)
+float4 LayerInfo; // x == Layer Alpha value, y == Layer Visibility flag, z == Layer Blend Mode (0: Additive Blend Mode, 1: Alpha Blend Mode), w == Unused
+float4 OutputConfig; // x == ApplyLayerModifiers, y == Alpha Blend Mode or Output weightmap as substractive/additive, z == Use ReadTexture2 and out Delta with ReadTexture1, w == Output Normals (heightmaps) or Unused (weightmaps)
+float2 TextureSize; // x == ReadTexture1 width, y == ReadTexture1 height
+float3 LandscapeGridScale; // x == LS Actor DrawScale.X, y == LS Actor DrawScale.y, z == LS Actor DrawScale.z / 128.0f (ZSCALE)
 float CurrentMipComponentVertexCount;
 float2 CurrentMipTextureSize;
 float2 ParentMipTextureSize;
 float4 SourceOffsetAndSizeUV;
+uint ChannelSwizzleMask;
 
 float ExtractAndScaleHeight(float2 InPackedHeight)
 {
@@ -24,7 +25,7 @@ float ExtractAndScaleHeight(float2 InPackedHeight)
 
 float ExtractAlpha(float2 InPackedAlpha)
 {
-	int AlphaWithFlags = ((int)round(InPackedAlpha.r * 255.0) << 8) | (int)round(InPackedAlpha.g * 255.0);
+	int AlphaWithFlags = ((int) round(InPackedAlpha.r * 255.0) << 8) | (int) round(InPackedAlpha.g * 255.0);
 	int Alpha = AlphaWithFlags & 0xFFFC; // Last 2 bits are used to store lower/raise flags
 	return float(Alpha) / 65532.0; // 0xFFFC in float
 }
@@ -33,10 +34,10 @@ float4 CalculateNormalsFromHeights(float2 InTextureCoordinates, float4 OutColor)
 {
 	float2 TexelSize = 1.0 / TextureSize;
 
-	float MaskMinBorderX = (InTextureCoordinates.x <= TexelSize.x) ? 0.0f : 1.0f;			// first pixel X
-	float MaskMinBorderY = (InTextureCoordinates.y <= TexelSize.y) ? 0.0f : 1.0f;			// first pixel Y
-	float MaskMaxBorderX = (InTextureCoordinates.x >= 1.0f - TexelSize.x) ? 0.0f : 1.0f;	// last pixel X
-	float MaskMaxBorderY = (InTextureCoordinates.y >= 1.0f - TexelSize.y) ? 0.0f : 1.0f;	// last pixel Y
+	float MaskMinBorderX = (InTextureCoordinates.x <= TexelSize.x) ? 0.0f : 1.0f; // first pixel X
+	float MaskMinBorderY = (InTextureCoordinates.y <= TexelSize.y) ? 0.0f : 1.0f; // first pixel Y
+	float MaskMaxBorderX = (InTextureCoordinates.x >= 1.0f - TexelSize.x) ? 0.0f : 1.0f; // last pixel X
+	float MaskMaxBorderY = (InTextureCoordinates.y >= 1.0f - TexelSize.y) ? 0.0f : 1.0f; // last pixel Y
 
 	// The triangle topology is the following (where C = center, T = top, B = bottom, L = left, R = right and Nx the normals we need to interpolate):
 	// TL ------ TT
@@ -59,7 +60,7 @@ float4 CalculateNormalsFromHeights(float2 InTextureCoordinates, float4 OutColor)
 	// So we can retrieve all 7 samples using 4 gathers : 
 	// Make sure we gather the 4 pixels in the bottom right direction by sampling at (0.75, 0.75) of the pixel (gather retrieves the 4 samples that would be used for bilinear interpolation):
 
-	float2 GatherLocation = InTextureCoordinates + 0.25f * TexelSize; 
+	float2 GatherLocation = InTextureCoordinates + 0.25f * TexelSize;
 	float4 Red0 = ReadTexture1.GatherRed(ReadTexture1Sampler, GatherLocation - TexelSize);
 	float4 Green0 = ReadTexture1.GatherGreen(ReadTexture1Sampler, GatherLocation - TexelSize);
 	float4 Mask0 = ReadTexture1.GatherAlpha(ReadTexture1Sampler, GatherLocation - TexelSize);
@@ -78,12 +79,12 @@ float4 CalculateNormalsFromHeights(float2 InTextureCoordinates, float4 OutColor)
 
 	// mask out samples that are off the edge of the texture
 	TL.w *= MaskMinBorderX * MaskMinBorderY;
-	TT.w *=                  MaskMinBorderY;
+	TT.w *= MaskMinBorderY;
 	// CC should never be off the edge
 	LL.w *= MaskMinBorderX;
 	RR.w *= MaskMaxBorderX;
 	BR.w *= MaskMaxBorderX * MaskMaxBorderY;
-	BB.w *=                  MaskMaxBorderY;
+	BB.w *= MaskMaxBorderY;
 
 	// HACK [chris.tchou] remove once we have a better boundary normal solution
 	bool bApplyNormalBoundaryOverride = (OutputConfig.w >= 1.5);
@@ -209,7 +210,7 @@ void PSHeightmapMainMips(float2 InTextureCoordinates : TEXCOORD0, out float4 Out
 
 	// Special case of 1 texel size component
 	if (CurrentMipComponentVertexCount == 1)
-	{		
+	{
 		if (CurrentMipTextureSize.x >= CurrentMipTextureSize.y) // x size 1 texel
 		{
 			IsMinBorderTexelY = true;
@@ -276,7 +277,7 @@ void PSHeightmapMainMips(float2 InTextureCoordinates : TEXCOORD0, out float4 Out
 		{
 			HeightResult = RightHeight;
 			NormalResult = float2(RightSample.ba);
-		}		
+		}
 		else if (IsMaxBorderTexelY) // on bottom border
 		{
 			HeightResult = DownRightHeight;
@@ -322,7 +323,7 @@ void PSHeightmapMainMips(float2 InTextureCoordinates : TEXCOORD0, out float4 Out
 			NormalResult.x = lerp(lerp(SourceSample.b, RightSample.b, Alpha), lerp(DownSample.b, DownRightSample.b, Alpha), Alpha);
 			NormalResult.y = lerp(lerp(SourceSample.a, RightSample.a, Alpha), lerp(DownSample.a, DownRightSample.a, Alpha), Alpha);
 		}
-	}	
+	}
 
 	float2 PackedHeightResult = PackHeight(HeightResult);
 	OutColor = float4(PackedHeightResult.xy, NormalResult.xy);
@@ -342,7 +343,7 @@ void PSWeightmapMain(float2 InTextureCoordinates : TEXCOORD0, out float4 OutColo
 
 		// Apply addtive/substractive logic
 		SourceColor = OutputConfig.y == 0.0f ? clamp(BaseColor + SourceColor, 0.0, 1.0) : clamp(BaseColor - SourceColor, 0.0, 1.0);
-	}	
+	}
 
 	OutColor = SourceColor;
 }
@@ -356,7 +357,7 @@ void PSWeightmapMainMips(float2 InTextureCoordinates : TEXCOORD0, out float4 Out
 
 	// Special case of 1 texel size component
 	if (CurrentMipComponentVertexCount == 1)
-	{		
+	{
 		if (CurrentMipTextureSize.x >= CurrentMipTextureSize.y) // x size 1 texel
 		{
 			IsMinBorderTexelY = true;
@@ -391,7 +392,7 @@ void PSWeightmapMainMips(float2 InTextureCoordinates : TEXCOORD0, out float4 Out
 		float4 SourceSample = ReadTexture1.SampleLevel(ReadTexture1Sampler, SourceUV, 0);
 		float4 DownSample = ReadTexture1.SampleLevel(ReadTexture1Sampler, SourceUV + float2(0.0, ParentMipTexelSize.y), 0);
 
-		OutColor = IsMinBorderTexelY ? SourceSample : (IsMaxBorderTexelY ? DownSample : lerp(SourceSample, DownSample, Alpha));		
+		OutColor = IsMinBorderTexelY ? SourceSample : (IsMaxBorderTexelY ? DownSample : lerp(SourceSample, DownSample, Alpha));
 	}
 	else if (IsMaxBorderTexelX) // on right border
 	{
@@ -422,5 +423,13 @@ void CopyTexturePS(in float4 InPosition : SV_POSITION, in noperspective float2 I
 {
 	// InUV is in the [0,1] range but in the viewport space, whose size is CopySize, so we need to adjust it to sample the corresponding range in the source texture : 
 	float2 SampleUV = (InUV * SourceOffsetAndSizeUV.zw) + SourceOffsetAndSizeUV.xy;
-	OutColor = ReadTexture1.SampleLevel(ReadTexture1Sampler, SampleUV, 0);
-}
+	float4 SampleColor = ReadTexture1.SampleLevel(ReadTexture1Sampler, SampleUV, 0);
+	
+	// Swizzle channels as requested by ChannelSwizzleMask:
+	[unroll]
+	for (uint DestinationChannelIndex = 0; DestinationChannelIndex < 4; ++DestinationChannelIndex)
+	{
+		uint SourceChannelIndex = ((3u << (DestinationChannelIndex * 2u)) & ChannelSwizzleMask) >> (DestinationChannelIndex * 2u);
+		OutColor[DestinationChannelIndex] = SampleColor[SourceChannelIndex];
+	}
+}
\ No newline at end of file
diff --git a/Engine/Source/Runtime/Landscape/Classes/Landscape.h b/Engine/Source/Runtime/Landscape/Classes/Landscape.h
--- a/Engine/Source/Runtime/Landscape/Classes/Landscape.h
+++ b/Engine/Source/Runtime/Landscape/Classes/Landscape.h
@@ -8,6 +8,7 @@
 #include "LandscapeProxy.h"
 #include "LandscapeBlueprintBrushBase.h"
 #include "LandscapeEditTypes.h"
+#include "LandscapeEditLayerRenderer.h"
 #include "Delegates/DelegateCombinations.h"
 #include "Templates/SubclassOf.h"
 #include "Algo/Transform.h"
@@ -102,49 +103,43 @@ enum ELandscapeLayerUpdateMode : uint32;
 
 USTRUCT()
 struct FLandscapeLayerBrush
+#if CPP && WITH_EDITOR // UHT doesn't support inheriting from namespaced class
+	: public UE::Landscape::EditLayers::IEditLayerRendererProvider
+#endif // CPP && WITH_EDITOR
 {
 	GENERATED_USTRUCT_BODY()
 
-	FLandscapeLayerBrush()
-#if WITH_EDITORONLY_DATA
-		: FLandscapeLayerBrush(nullptr)
-#endif
-	{}
+	FLandscapeLayerBrush() = default;
 
 	FLandscapeLayerBrush(ALandscapeBlueprintBrushBase* InBlueprintBrush)
 #if WITH_EDITORONLY_DATA
 		: BlueprintBrush(InBlueprintBrush)
-		, LandscapeSize(MAX_int32, MAX_int32)
-		, LandscapeRenderTargetSize(MAX_int32, MAX_int32)
-#endif
+#endif // WITH_EDITORONLY_DATA
 	{}
 
+	virtual ~FLandscapeLayerBrush() = default;
+
 #if WITH_EDITOR
-	UTextureRenderTarget2D* Render(bool InIsHeightmap, const FIntRect& InLandscapeSize, UTextureRenderTarget2D* InLandscapeRenderTarget, const FName& InWeightmapLayerName = NAME_None);
-	
-	UTextureRenderTarget2D* RenderLayer(const FIntRect& InLandscapeSize, const FLandscapeBrushParameters& InParameters);
+	UTextureRenderTarget2D* RenderLayer(const FLandscapeBrushParameters& InParameters);
 
 	LANDSCAPE_API ALandscapeBlueprintBrushBase* GetBrush() const;
+	bool Affects(ELandscapeToolTargetType InTargetType, const FName& InWeightmapLayerName) const;
 	bool AffectsHeightmap() const;
 	bool AffectsWeightmapLayer(const FName& InWeightmapLayerName) const;
 	bool AffectsVisibilityLayer() const;
 	void SetOwner(ALandscape* InOwner);
-#endif
 
-private:
+	//~ Begin UE::Landscape::EditLayers::IEditLayerRendererProvider implementation
+	LANDSCAPE_API virtual TArray<UE::Landscape::EditLayers::FEditLayerRendererState> GetEditLayerRendererStates(const ULandscapeInfo* InLandscapeInfo, bool bInSkipBrush) override;
+	//~ End UE::Landscape::EditLayers::IEditLayerRendererProvider implementation
+#endif // WITH_EDITOR
 
-#if WITH_EDITOR
-	bool Initialize(const FIntRect& InLandscapeExtent, UTextureRenderTarget2D* InLandscapeRenderTarget);
-#endif
+private:
 
 #if WITH_EDITORONLY_DATA
 	UPROPERTY()
 	TObjectPtr<ALandscapeBlueprintBrushBase> BlueprintBrush;
-
-	FTransform LandscapeTransform;
-	FIntPoint LandscapeSize;
-	FIntPoint LandscapeRenderTargetSize;
-#endif
+#endif // WITH_EDITORONLY_DATA
 };
 
 UENUM()
@@ -157,6 +152,9 @@ enum ELandscapeBlendMode : int
 
 USTRUCT()
 struct FLandscapeLayer
+#if CPP && WITH_EDITOR // UHT doesn't support inheriting from namespaced class
+	: public UE::Landscape::EditLayers::IEditLayerRendererProvider
+#endif // CPP && WITH_EDITOR
 {
 	GENERATED_USTRUCT_BODY()
 
@@ -170,6 +168,16 @@ struct FLandscapeLayer
 		, BlendMode(LSBM_AdditiveBlend)
 	{}
 
+	virtual ~FLandscapeLayer() = default;
+
+#if WITH_EDITOR
+	//~ Begin UE::Landscape::EditLayers::IEditLayerRendererProvider implementation
+	LANDSCAPE_API virtual TArray<UE::Landscape::EditLayers::FEditLayerRendererState> GetEditLayerRendererStates(const ULandscapeInfo* InLandscapeInfo, bool bInSkipBrush) override;
+	//~ End UE::Landscape::EditLayers::IEditLayerRendererProvider implementation
+#endif // WITH_EDITOR
+
+	// TODO [jonathan.bard] : make those private, FLandscapeLayer is not LANDSCAPE_API so it can be done without deprecation
+public:
 	UPROPERTY(meta = (IgnoreForMemberInitializationTest))
 	FGuid Guid;
 
@@ -201,8 +209,42 @@ struct FLandscapeLayer
 	TObjectPtr<ULandscapeEditLayerBase> EditLayer;
 };
 
+enum class ELandscapeEditLayersMergeMode : uint8
+{
+	GlobalMerge = 0,
+	LocalMerge, 
+	BatchedMerge,
+	Invalid 
+};
+
+// TODO [jonathan.bard] : When global merge is a thing of the past, we should pass FOnRenderBatchTargetGroupDoneParams here, or extract some of its higher-level information, because there
+//  is a lot more useful information accessible in there (components being rendered, target layers, validity render targets, etc.)
+struct FOnLandscapeEditLayersMergedParams
+{
+	FOnLandscapeEditLayersMergedParams(UTextureRenderTarget* InRenderTarget, const FIntPoint& InRenderAreaResolution, bool bInIsHeightmapMerge)
+		: RenderTarget(InRenderTarget)
+		, RenderAreaResolution(InRenderAreaResolution)
+		, bIsHeightmapMerge(bInIsHeightmapMerge)
+	{}
+
+	/**
+	* Render target of the section of landscape that was rendered (important note: the render target's resolution can be larger than the actual landscape resolution, so RenderAreaResolution must be used.
+	* In the case of weightmaps, the render target will actually be a UTextureRenderTarget2DArray
+	*/
+	UTextureRenderTarget* RenderTarget = nullptr;
+
+	/** Actual resolution of this render : render targets are usually sized larger than the effective resolution at this point, so this needs to be used instead.
+	 *  Note : it's the true resolution of the landscape : at this point, there are no duplicate vertices in the render target. */
+	FIntPoint RenderAreaResolution = FIntPoint(ForceInit);
+
+	bool bIsHeightmapMerge = false;
+};
+
 UCLASS(MinimalAPI, showcategories=(Display, Movement, Collision, Lighting, LOD, Input), hidecategories=(Mobility))
 class ALandscape : public ALandscapeProxy
+#if CPP && WITH_EDITOR // UHT doesn't support inheriting from namespaced class
+	, public UE::Landscape::EditLayers::IEditLayerRendererProvider
+#endif // CPP && WITH_EDITOR
 {
 	GENERATED_BODY()
 
@@ -359,6 +401,7 @@ public:
 	LANDSCAPE_API const FLandscapeLayer* GetLayerConst(int32 InLayerIndex) const;
 	LANDSCAPE_API const FLandscapeLayer* GetLayerConst(const FGuid& InLayerGuid) const;
 	LANDSCAPE_API const FLandscapeLayer* GetLayerConst(const FName& InLayerName) const;
+	LANDSCAPE_API int32 GetLayerIndex(const FGuid& InLayerGuid) const;
 	LANDSCAPE_API int32 GetLayerIndex(FName InLayerName) const;
 
 	UE_DEPRECATED(5.5, "This has moved to private (GetLayerInternal), use the appropriate setters to mutate the edit layer")
@@ -392,7 +435,7 @@ public:
 	LANDSCAPE_API void ClearPaintLayer(const FGuid& InLayerGuid, ULandscapeLayerInfoObject* InLayerInfo);
 	LANDSCAPE_API void ClearLayer(int32 InLayerIndex, TSet<TObjectPtr<ULandscapeComponent>>* InComponents = nullptr, ELandscapeClearMode InClearMode = ELandscapeClearMode::Clear_All);
 	LANDSCAPE_API void ClearLayer(const FGuid& InLayerGuid, TSet<TObjectPtr<ULandscapeComponent>>* InComponents = nullptr, ELandscapeClearMode InClearMode = ELandscapeClearMode::Clear_All, bool bMarkPackageDirty = true);
-	LANDSCAPE_API void DeleteLayer(int32 InLayerIndex);
+	LANDSCAPE_API bool DeleteLayer(int32 InLayerIndex);
 	LANDSCAPE_API void CollapseLayer(int32 InLayerIndex);
 	LANDSCAPE_API void DeleteLayers();
 	LANDSCAPE_API void SetEditingLayer(const FGuid& InLayerGuid = FGuid());
@@ -416,7 +459,7 @@ public:
 	LANDSCAPE_API void SetLayerSubstractiveBlendStatus(int32 InLayerIndex, bool InStatus, const TWeakObjectPtr<ULandscapeLayerInfoObject>& InLayerInfoObj);
 	LANDSCAPE_API void ReplaceLayerSubstractiveBlendStatus(ULandscapeLayerInfoObject* InFromLayerInfo, ULandscapeLayerInfoObject* InToLayerInfo, bool bInShouldDirtyPackage);
 
-	LANDSCAPE_API int32 GetBrushLayer(class ALandscapeBlueprintBrushBase* InBrush) const;
+	LANDSCAPE_API int32 GetBrushLayer(const ALandscapeBlueprintBrushBase* InBrush) const;
 	LANDSCAPE_API void AddBrushToLayer(int32 InLayerIndex, class ALandscapeBlueprintBrushBase* InBrush);
 	LANDSCAPE_API void RemoveBrush(class ALandscapeBlueprintBrushBase* InBrush);
 	LANDSCAPE_API void RemoveBrushFromLayer(int32 InLayerIndex, class ALandscapeBlueprintBrushBase* InBrush);
@@ -468,22 +511,31 @@ private:
 	bool SupportsEditLayersLocalMerge();
 	bool HasNormalCaptureBPBrushLayer();
 
+	ELandscapeEditLayersMergeMode GetEditLayersMergeMode();
 	bool CreateLayersRenderingResource(bool bUseNormalCapture);
 	void PrepareEditLayersLocalMergeResources();
 	void UpdateLayersContent(bool bInWaitForStreaming = false, bool bInSkipMonitorLandscapeEdModeChanges = false, bool bIntermediateRender = false, bool bFlushRender = false);
 	void MonitorShaderCompilation();
 	void MonitorLandscapeEdModeChanges();
 	
+	//~ Begin UE::Landscape::EditLayers::IEditLayerRendererProvider implementation
+	virtual TArray<UE::Landscape::EditLayers::FEditLayerRendererState> GetEditLayerRendererStates(const ULandscapeInfo* InLandscapeInfo, bool bInSkipBrush) override;
+	//~ End UE::Landscape::EditLayers::IEditLayerRendererProvider implementation
+
+	UE::Landscape::EditLayers::FMergeRenderContext PrepareEditLayersMergeRenderContext(const UE::Landscape::EditLayers::FMergeRenderParams& InParams);
+
 	FLandscapeBrushParameters BuildLandscapeBrushParameters(bool bInIsHeightmapMerge, const FIntRect& InRenderAreaExtents, UTextureRenderTarget2D* InRenderTarget, FName InWeightmapLayerName = NAME_None);
 
 	int32 RegenerateLayersHeightmaps(const FUpdateLayersContentContext& InUpdateLayersContentContext);
 	int32 PerformLayersHeightmapsLocalMerge(const FUpdateLayersContentContext& InUpdateLayersContentContext, const FEditLayersHeightmapMergeParams& InMergeParams);
 	int32 PerformLayersHeightmapsGlobalMerge(const FUpdateLayersContentContext& InUpdateLayersContentContext, const FEditLayersHeightmapMergeParams& InMergeParams);
+	int32 PerformLayersHeightmapsBatchedMerge(const FUpdateLayersContentContext& InUpdateLayersContentContext, const FEditLayersHeightmapMergeParams& InMergeParams);
 	void ResolveLayersHeightmapTexture(const FTextureToComponentHelper& MapHelper, const TSet<UTexture2D*>& HeightmapsToResolve, bool bIntermediateRender, TArray<FLandscapeEditLayerComponentReadbackResult>& InOutComponentReadbackResults);
 
 	int32 RegenerateLayersWeightmaps(FUpdateLayersContentContext& InUpdateLayersContentContext);
 	int32 PerformLayersWeightmapsLocalMerge(FUpdateLayersContentContext& InUpdateLayersContentContext, const FEditLayersWeightmapMergeParams& InMergeParams);
 	int32 PerformLayersWeightmapsGlobalMerge(FUpdateLayersContentContext& InUpdateLayersContentContext, const FEditLayersWeightmapMergeParams& InMergeParams);
+	int32 PerformLayersWeightmapsBatchedMerge(FUpdateLayersContentContext& InUpdateLayersContentContext, const FEditLayersWeightmapMergeParams& InMergeParams);
 	void ResolveLayersWeightmapTexture(const FTextureToComponentHelper& MapHelper, const TSet<UTexture2D*>& WeightmapsToResolve, bool bIntermediateRender, TArray<FLandscapeEditLayerComponentReadbackResult>& InOutComponentReadbackResults);
 
 	bool HasTextureDataChanged(TArrayView<const FColor> InOldData, TArrayView<const FColor> InNewData, bool bInIsWeightmap, uint64 InPreviousHash, uint64& OutNewHash, TOptional<uint8>& OutChangedWeightmapChannelsMasks) const;
@@ -509,7 +561,8 @@ private:
 														  TArray<struct FLandscapeLayerWeightmapExtractMaterialLayersComponentData>& OutComponentData, TMap<ULandscapeLayerInfoObject*, int32>& OutLayerInfoObjects);
 	void PrepareComponentDataToPackMaterialLayersCS(int32 InCurrentWeightmapToProcessIndex, const FIntPoint& InLandscapeBase, const TArray<ULandscapeComponent*>& InAllLandscapeComponents, TArray<UTexture2D*>& InOutProcessedWeightmaps,
 													TArray<FLandscapeEditLayerReadback*>& OutProcessedCPUReadBacks, TArray<struct FLandscapeLayerWeightmapPackMaterialLayersComponentData>& OutComponentData);
-	void ReallocateLayersWeightmaps(FUpdateLayersContentContext& InUpdateLayersContentContext, const TArray<ULandscapeLayerInfoObject*>& InBrushRequiredAllocations);
+	void ReallocateLayersWeightmaps(FUpdateLayersContentContext& InUpdateLayersContentContext, const TArray<ULandscapeLayerInfoObject*>& InBrushRequiredAllocations, 
+		const TMap<ULandscapeComponent*, TArray<ULandscapeLayerInfoObject*>>* InPerComponentAllocations, bool bInRestrictTextureSharingToTheseComponents);
 	void InitializeLayersWeightmapResources();
 	bool GenerateZeroAllocationPerComponents(const TArray<ALandscapeProxy*>& InAllLandscape, const TMap<ULandscapeLayerInfoObject*, bool>& InWeightmapLayersBlendSubstractive);
 
@@ -578,8 +631,9 @@ public:
 	DECLARE_EVENT(ALandscape, FLandscapeBlueprintBrushChangedDelegate);
 	FLandscapeBlueprintBrushChangedDelegate& OnBlueprintBrushChangedDelegate() { return LandscapeBlueprintBrushChangedDelegate; }
 
-	DECLARE_EVENT_OneParam(ALandscape, FLandscapeFullHeightmapRenderDoneDelegate, UTextureRenderTarget2D*);
-	FLandscapeFullHeightmapRenderDoneDelegate& OnFullHeightmapRenderDoneDelegate() { return LandscapeFullHeightmapRenderDoneDelegate; }
+	/** Delegate that will be called whenever an edit layers merge is done */
+	DECLARE_MULTICAST_DELEGATE_OneParam(FOnEditLayersMergedDelegate, const FOnLandscapeEditLayersMergedParams& /*InParams*/);
+	FOnEditLayersMergedDelegate::RegistrationType& OnEditLayersMerged() const { return OnEditLayersMergedDelegate; }
 
 	/** Target Landscape Layer for Landscape Splines */
 	UE_DEPRECATED(5.5, "This has been refactored into the generic ULandscapeEditLayerBase system. Please check for the presence of a ULandscapeEditLayerSplines layer instead")
@@ -617,7 +671,7 @@ private:
 	TArray<FLandscapeLayer> LandscapeEditLayers;
 
 	FLandscapeBlueprintBrushChangedDelegate LandscapeBlueprintBrushChangedDelegate;
-	FLandscapeFullHeightmapRenderDoneDelegate LandscapeFullHeightmapRenderDoneDelegate;
+	mutable FOnEditLayersMergedDelegate OnEditLayersMergedDelegate;
 
 	/** Components affected by landscape splines (used to partially clear Layer Reserved for Splines) */
 	UPROPERTY(Transient)
@@ -657,6 +711,8 @@ private:
 	UPROPERTY(Transient)
 	bool bSplineLayerUpdateRequested;
 
+	ELandscapeEditLayersMergeMode CurrentEditLayersMergeMode = ELandscapeEditLayersMergeMode::GlobalMerge;
+
 	struct FWaitingForResourcesNotificationHelper
 	{
 		void Notify(ALandscape* InLandscape, class FLandscapeNotificationManager* InNotificationManager, ELandscapeNotificationType InNotificationType, const FText& InNotificationText);
diff --git a/Engine/Source/Runtime/Landscape/Classes/LandscapeComponent.h b/Engine/Source/Runtime/Landscape/Classes/LandscapeComponent.h
--- a/Engine/Source/Runtime/Landscape/Classes/LandscapeComponent.h
+++ b/Engine/Source/Runtime/Landscape/Classes/LandscapeComponent.h
@@ -965,6 +965,12 @@ public:
 		return FIntPoint(SectionBaseX, SectionBaseY);
 	}
 
+	/** @return the component's coordinates (aka index, aka component key) in the entire landscape */
+	FIntPoint GetComponentKey() const
+	{
+		return GetSectionBase() / ComponentSizeQuads;
+	}
+
 	/** @param InSectionBase new section base for a component */
 	void SetSectionBase(FIntPoint InSectionBase)
 	{
@@ -1127,12 +1133,13 @@ public:
 	/**
 	 * Create weightmaps for this component for the layers specified in the WeightmapLayerAllocations array, works in the landscape current edit layer when InCanUseEditingWeightmap is true
 	 */
+	UE_DEPRECATED(5.5, "Use the new prototype of ReallocateWeightmaps : note : this is very internal stuff, normally, you shouldn't have to use this function at all")
 	LANDSCAPE_API void ReallocateWeightmaps(FLandscapeEditDataInterface* DataInterface = nullptr, bool InCanUseEditingWeightmap = true, bool InSaveToTransactionBuffer = true, bool InForceReallocate = false, ALandscapeProxy* InTargetProxy = nullptr, TArray<UTexture*>* OutNewCreatedTextures = nullptr);
 
 	/**
 	 * Create weightmaps for this component for the layers specified in the WeightmapLayerAllocations array, works in the specified edit layer
 	 */
-	void ReallocateWeightmapsInternal(FLandscapeEditDataInterface* DataInterface = nullptr, const FGuid& InEditLayerGuid = FGuid(), bool InSaveToTransactionBuffer = true, bool InForceReallocate = false, ALandscapeProxy* InTargetProxy = nullptr, TArray<UTexture*>* OutNewCreatedTextures = nullptr);
+	LANDSCAPE_API TArray<UTexture*> ReallocateWeightmaps(FLandscapeEditDataInterface* DataInterface, const FGuid& InEditLayerGuid, bool bInSaveToTransactionBuffer, bool bInForceReallocate, ALandscapeProxy* InTargetProxy, TSet<ULandscapeComponent*>* InRestrictSharingToComponents);
 
 	/** Returns true if the component has a valid LandscapeHoleMaterial */
 	LANDSCAPE_API bool IsLandscapeHoleMaterialValid() const;
diff --git a/Engine/Source/Runtime/Landscape/Classes/LandscapeEditLayer.h b/Engine/Source/Runtime/Landscape/Classes/LandscapeEditLayer.h
--- a/Engine/Source/Runtime/Landscape/Classes/LandscapeEditLayer.h
+++ b/Engine/Source/Runtime/Landscape/Classes/LandscapeEditLayer.h
@@ -2,6 +2,8 @@
 
 #pragma once
 
+#include "LandscapeEditLayerRenderer.h"
+
 #include "LandscapeEditLayer.generated.h"
 
 enum class ELandscapeToolTargetType : uint8;
@@ -14,6 +16,9 @@ class ALandscape;
 */
 UCLASS(MinimalAPI, Abstract)
 class ULandscapeEditLayerBase : public UObject
+#if CPP && WITH_EDITOR
+	, public ILandscapeEditLayerRenderer
+#endif // CPP && WITH_EDITOR
 {
 	GENERATED_BODY()
 
@@ -149,8 +154,14 @@ public:
 
 	// UObject
 	LANDSCAPE_API virtual void PostLoad() override;
+
 protected:
+#if WITH_EDITOR
+	// TODO [jonathan.bard] remove this dependency to FLandscapeLayer ASAP (once all data from there has moved to the ULandscapeEditLayer class) :
+	const FLandscapeLayer* GetOwningLayer() const;
+#endif // WITH_EDITOR
 
+protected:
 	// TODO: This might be removed once more things are moved from FLandscapeLayer to ULandscapeLayer
 	UPROPERTY()
 	TWeakObjectPtr<ALandscape> OwningLandscape;
@@ -169,6 +180,16 @@ public:
 	virtual bool NeedsPersistentTextures() const override { return true; };
 	virtual bool SupportsCollapsingTo() const override { return true; } // If the layer has persistent textures, it can be collapsed to another layer (one that supports being collapsed away, that is)
 	// End ULandscapeEditLayerBase implementation
+
+#if WITH_EDITOR
+	//~ Begin ILandscapeEditLayerRenderer implementation
+	LANDSCAPE_API virtual void GetRendererStateInfo(const ULandscapeInfo* InLandscapeInfo,
+		UE::Landscape::EditLayers::FEditLayerTargetTypeState& OutSupportedTargetTypeState, UE::Landscape::EditLayers::FEditLayerTargetTypeState& OutEnabledTargetTypeState, TArray<TSet<FName>>& OutRenderGroups) const override;
+	LANDSCAPE_API virtual TArray<UE::Landscape::EditLayers::FEditLayerRenderItem> GetRenderItems(const ULandscapeInfo* InLandscapeInfo) const override;
+	LANDSCAPE_API virtual void RenderLayer(ILandscapeEditLayerRenderer::FRenderParams& InRenderParams) override;
+	LANDSCAPE_API virtual FString GetEditLayerRendererDebugName() const override;
+	//~ End ILandscapeEditLayerRenderer implementation
+#endif // WITH_EDITOR
 };
 
 /** 
diff --git a/Engine/Source/Runtime/Landscape/Classes/LandscapeInfo.h b/Engine/Source/Runtime/Landscape/Classes/LandscapeInfo.h
--- a/Engine/Source/Runtime/Landscape/Classes/LandscapeInfo.h
+++ b/Engine/Source/Runtime/Landscape/Classes/LandscapeInfo.h
@@ -432,6 +432,9 @@ public:
 	 * 
 	 * @return true if at least one overlapped component
 	 */
+	// TODO [jonathan.bard] : There should be multiple versions of this. 
+	//  1. This currently performs AABB to AABB checks and is therefore not optimal (but faster than OOBB to OOBB could ever be, which is why we might need 2 versions) :
+	//  2. Returning a TMap<> is a bit overkill, since the component key can be retrieved trivially from the component
 	LANDSCAPE_API bool GetOverlappedComponents(const FTransform& InAreaWorldTransform, const FBox2D& InAreaExtents, TMap<FIntPoint, ULandscapeComponent*>& OutOverlappedComponents, FIntRect& OutComponentIndicesBoundingRect);
 
 #if WITH_EDITOR
diff --git a/Engine/Source/Runtime/Landscape/Private/GlobalMergeLegacySupportUtil.cpp b/Engine/Source/Runtime/Landscape/Private/GlobalMergeLegacySupportUtil.cpp
--- a/Engine/Source/Runtime/Landscape/Private/GlobalMergeLegacySupportUtil.cpp
+++ b/Engine/Source/Runtime/Landscape/Private/GlobalMergeLegacySupportUtil.cpp
@@ -8,7 +8,7 @@
 
 #if WITH_EDITOR
 UTextureRenderTarget2D* ILandscapeBrushRenderCallAdapter_GlobalMergeLegacySupport::RenderAsBlueprintBrush(
-	const FLandscapeBrushParameters& InParameters, const FTransform& NewLandscapeTransform, const FIntPoint& RenderAreaSize)
+	const FLandscapeBrushParameters& InParameters, const FTransform& NewLandscapeTransform)
 {
 	TRACE_CPUPROFILER_EVENT_SCOPE(ILandscapeBrushRenderCallAdapter_GlobalMergeLegacySupport::RenderAsBlueprintBrush);
 
@@ -28,11 +28,11 @@ UTextureRenderTarget2D* ILandscapeBrushRenderCallAdapter_GlobalMergeLegacySuppor
 
 	const FIntPoint NewLandscapeRenderTargetSize = FIntPoint(InParameters.CombinedResult->SizeX, InParameters.CombinedResult->SizeY);
 	if (!CurrentRenderAreaWorldTransform.Equals(NewLandscapeTransform) 
-		|| (CurrentRenderAreaSize != RenderAreaSize)
+		|| (CurrentRenderAreaSize != InParameters.RenderAreaSize)
 		|| CurrentRenderTargetSize != NewLandscapeRenderTargetSize)
 	{
 		CurrentRenderAreaWorldTransform = NewLandscapeTransform;
-		CurrentRenderAreaSize = RenderAreaSize;
+		CurrentRenderAreaSize = InParameters.RenderAreaSize;
 		CurrentRenderTargetSize = NewLandscapeRenderTargetSize;
 
 		InitializeAsBlueprintBrush(CurrentRenderAreaWorldTransform, CurrentRenderAreaSize, CurrentRenderTargetSize);
diff --git a/Engine/Source/Runtime/Landscape/Private/Landscape.cpp b/Engine/Source/Runtime/Landscape/Private/Landscape.cpp
--- a/Engine/Source/Runtime/Landscape/Private/Landscape.cpp
+++ b/Engine/Source/Runtime/Landscape/Private/Landscape.cpp
@@ -81,6 +81,7 @@ Landscape.cpp: Terrain rendering
 #include "LandscapeEditLayer.h"
 #include "LandscapeTextureStorageProvider.h"
 #include "LandscapeUtils.h"
+#include "LandscapeUtilsPrivate.h"
 #include "LandscapeVersion.h"
 #include "UObject/FortniteMainBranchObjectVersion.h"
 #include "UObject/FortniteReleaseBranchCustomObjectVersion.h"
@@ -638,7 +639,7 @@ void ULandscapeComponent::CheckGenerateMobilePlatformData(bool bIsCooking, const
 	FString MobileVersion = FDevSystemGuids::GetSystemGuid(FDevSystemGuids::Get().LANDSCAPE_MOBILE_COOK_VERSION).ToString();
 	ComponentStateAr << MobileVersion;
 
-	bool IsTextureArrayEnabled = UE::Landscape::IsMobileWeightmapTextureArrayEnabled();
+	bool IsTextureArrayEnabled = UE::Landscape::Private::IsMobileWeightmapTextureArrayEnabled();
 	ComponentStateAr << IsTextureArrayEnabled;
 
 	uint32 Hash[5];
@@ -1677,7 +1678,6 @@ ALandscape::ALandscape(const FObjectInitializer& ObjectInitializer)
 	WeightmapScratchExtractLayerTextureResource = nullptr;
 	WeightmapScratchPackLayerTextureResource = nullptr;
 	bLandscapeLayersAreInitialized = false;
-	bLandscapeLayersAreUsingLocalMerge = false;
 	LandscapeEdMode = nullptr;
 	bGrassUpdateEnabled = true;
 	bIsSpatiallyLoaded = false;
@@ -2393,9 +2393,7 @@ void ULandscapeComponent::CopyFinalLayerIntoEditingLayer(FLandscapeEditDataInter
 		}
 	}
 
-	const bool bEditingWeighmaps = true;
-	const bool bSaveToTransactionBuffer = true;
-	ReallocateWeightmaps(&DataInterface, bEditingWeighmaps, bSaveToTransactionBuffer);
+	ReallocateWeightmaps(&DataInterface, GetEditingLayerGUID(), /*bInSaveToTransactionBuffer = */true, /*bool bInForceReallocate = */false, /*InTargetProxy = */nullptr, /*InRestrictSharingToComponents = */nullptr);
 
 	const TArray<TObjectPtr<UTexture2D>>& EditingWeightmapTextures = GetWeightmapTextures(true);
 	for (const FWeightmapLayerAllocationInfo& AllocInfo : EditingLayerWeightmapLayerAllocations)
@@ -2850,7 +2848,6 @@ PRAGMA_DISABLE_DEPRECATION_WARNINGS
 			}
 			check(Layer.EditLayer == nullptr);
 			Layer.EditLayer = NewObject<ULandscapeEditLayerBase>(this, EditLayerClass, MakeUniqueObjectName(this, EditLayerClass));
-			Layer.EditLayer->SetBackPointer(this);
 			Layer.EditLayer->OnLayerCreated(Layer);
 		}
 
@@ -2872,19 +2869,21 @@ PRAGMA_ENABLE_DEPRECATION_WARNINGS
 				Brush.SetOwner(this);
 			}
 
+			Layer.EditLayer->SetBackPointer(this);
 		}
 		else
 		{
 			UE_LOG(LogLandscape, Error, TEXT("Couldn't load edit layer object associated with layer %s for landscape %s. This may happen when the edit layer class cannot be found (for example, when a plugin is removed from the project). The layer will be deleted."), *Layer.Name.ToString(), *GetFullName());
-			// Don't use DeleteLayer because it relies on the ULandscapeInfo object to be valid, which is not the case on PostLoad. 
-			// Simply remove the layer from list and the proxies will eventually tell the user to remove their data associated with this layer upon registration (in order to let them avoid data loss) :
-			LandscapeEditLayers.RemoveAt(LayerIndex);
-			--LayerIndex;
-
-			// Request Update
-			RequestLayersContentUpdateForceAll();
+			ensure(DeleteLayer(LayerIndex));
 		}
 	}
+
+	// In case we're a landscape with edit layers but we actually lack a layer (e.g. it was removed by the test above, because its edit layer class is unknown), let's create one all the same 
+	//  because we're always supposed to have at least 1 : 
+	if (CanHaveLayersContent() && LandscapeEditLayers.IsEmpty())
+	{
+		ensure(CreateLayer() != INDEX_NONE);
+	}
 #endif // WITH_EDITOR
 
 	Super::PostLoad();
@@ -5907,6 +5906,7 @@ void ULandscapeInfo::UnregisterCollisionComponent(ULandscapeHeightfieldCollision
 	}
 }
 
+// TODO [jonathan.bard] : improve this function or create another one to take into account unloaded proxies : 
 bool ULandscapeInfo::GetOverlappedComponents(const FTransform& InAreaWorldTransform, const FBox2D& InAreaExtents, 
 	TMap<FIntPoint, ULandscapeComponent*>& OutOverlappedComponents, FIntRect& OutComponentIndicesBoundingRect)
 {
@@ -5915,33 +5915,46 @@ bool ULandscapeInfo::GetOverlappedComponents(const FTransform& InAreaWorldTransf
 		return false;
 	}
 
-	// Compute the AABB for this area in landscape space to find which of the landscape components are overlapping :
-	FVector Extremas[4];
-	const FTransform& LandscapeTransform = LandscapeActor->GetTransform();
-	Extremas[0] = LandscapeTransform.InverseTransformPosition(InAreaWorldTransform.TransformPosition(FVector(InAreaExtents.Min.X, InAreaExtents.Min.Y, 0.0)));
-	Extremas[1] = LandscapeTransform.InverseTransformPosition(InAreaWorldTransform.TransformPosition(FVector(InAreaExtents.Min.X, InAreaExtents.Max.Y, 0.0)));
-	Extremas[2] = LandscapeTransform.InverseTransformPosition(InAreaWorldTransform.TransformPosition(FVector(InAreaExtents.Max.X, InAreaExtents.Min.Y, 0.0)));
-	Extremas[3] = LandscapeTransform.InverseTransformPosition(InAreaWorldTransform.TransformPosition(FVector(InAreaExtents.Max.X, InAreaExtents.Max.Y, 0.0)));
-	FBox LocalExtents(Extremas, 4);
-
-	// Indices of the landscape components needed for rendering this area : 
-	FIntRect BoundingIndices;
-	BoundingIndices.Min = FIntPoint(FMath::FloorToInt32(LocalExtents.Min.X / ComponentSizeQuads), FMath::FloorToInt32(LocalExtents.Min.Y / ComponentSizeQuads));
-	// The max here is meant to be an exclusive bound, hence the +1
-	BoundingIndices.Max = FIntPoint(FMath::FloorToInt32(LocalExtents.Max.X / ComponentSizeQuads), FMath::FloorToInt32(LocalExtents.Max.Y / ComponentSizeQuads)) + FIntPoint(1);
-
-	// There could be missing components, so the effective area is actually a subset of this area :
 	FIntRect EffectiveBoundingIndices;
-	// Go through each loaded component and find out the actual bounds of the area we need to render :
-	for (int32 KeyY = BoundingIndices.Min.Y; KeyY < BoundingIndices.Max.Y; ++KeyY)
+
+	// Consider invalid extents as meaning "infinite", in which case, return all loaded components : 
+	if (!InAreaExtents.bIsValid)
 	{
-		for (int32 KeyX = BoundingIndices.Min.X; KeyX < BoundingIndices.Max.X; ++KeyX)
+		OutOverlappedComponents.Reserve(XYtoComponentMap.Num());
+		for (const TPair<FIntPoint, ULandscapeComponent*>& XYComponentPair : XYtoComponentMap)
 		{
-			FIntPoint Key(KeyX, KeyY);
-			if (ULandscapeComponent* Component = XYtoComponentMap.FindRef(Key))
+			EffectiveBoundingIndices.Union(FIntRect(XYComponentPair.Key, XYComponentPair.Key + FIntPoint(1)));
+			OutOverlappedComponents.Add(XYComponentPair);
+		}
+	}
+	else
+	{
+		// Compute the AABB for this area in landscape space to find which of the landscape components are overlapping :
+		FVector Extremas[4];
+		const FTransform& LandscapeTransform = LandscapeActor->GetTransform();
+		Extremas[0] = LandscapeTransform.InverseTransformPosition(InAreaWorldTransform.TransformPosition(FVector(InAreaExtents.Min.X, InAreaExtents.Min.Y, 0.0)));
+		Extremas[1] = LandscapeTransform.InverseTransformPosition(InAreaWorldTransform.TransformPosition(FVector(InAreaExtents.Min.X, InAreaExtents.Max.Y, 0.0)));
+		Extremas[2] = LandscapeTransform.InverseTransformPosition(InAreaWorldTransform.TransformPosition(FVector(InAreaExtents.Max.X, InAreaExtents.Min.Y, 0.0)));
+		Extremas[3] = LandscapeTransform.InverseTransformPosition(InAreaWorldTransform.TransformPosition(FVector(InAreaExtents.Max.X, InAreaExtents.Max.Y, 0.0)));
+		FBox LocalExtents(Extremas, 4);
+
+		// Indices of the landscape components needed for rendering this area : 
+		FIntRect BoundingIndices;
+		BoundingIndices.Min = FIntPoint(FMath::FloorToInt32(LocalExtents.Min.X / ComponentSizeQuads), FMath::FloorToInt32(LocalExtents.Min.Y / ComponentSizeQuads));
+		// The max here is meant to be an exclusive bound, hence the +1
+		BoundingIndices.Max = FIntPoint(FMath::FloorToInt32(LocalExtents.Max.X / ComponentSizeQuads), FMath::FloorToInt32(LocalExtents.Max.Y / ComponentSizeQuads)) + FIntPoint(1);
+
+		// Go through each loaded component and find out the actual bounds of the area we need to render :
+		for (int32 KeyY = BoundingIndices.Min.Y; KeyY < BoundingIndices.Max.Y; ++KeyY)
+		{
+			for (int32 KeyX = BoundingIndices.Min.X; KeyX < BoundingIndices.Max.X; ++KeyX)
 			{
-				EffectiveBoundingIndices.Union(FIntRect(Key, Key + FIntPoint(1)));
-				OutOverlappedComponents.Add(Key, Component);
+				FIntPoint Key(KeyX, KeyY);
+				if (ULandscapeComponent* Component = XYtoComponentMap.FindRef(Key))
+				{
+					EffectiveBoundingIndices.Union(FIntRect(Key, Key + FIntPoint(1)));
+					OutOverlappedComponents.Add(Key, Component);
+				}
 			}
 		}
 	}
@@ -5967,7 +5980,7 @@ void ULandscapeInfo::RegisterActorComponent(ULandscapeComponent* Component, bool
 	check(Component);
 
 	FIntPoint ComponentKey = Component->GetSectionBase() / Component->ComponentSizeQuads;
-	auto RegisteredComponent = XYtoComponentMap.FindRef(ComponentKey);
+	ULandscapeComponent* RegisteredComponent = XYtoComponentMap.FindRef(ComponentKey);
 
 	if (RegisteredComponent != Component)
 	{
@@ -6023,7 +6036,7 @@ void ULandscapeInfo::UnregisterActorComponent(ULandscapeComponent* Component)
 	if (ensure(Component))
 	{
 		FIntPoint ComponentKey = Component->GetSectionBase() / Component->ComponentSizeQuads;
-		auto RegisteredComponent = XYtoComponentMap.FindRef(ComponentKey);
+		ULandscapeComponent* RegisteredComponent = XYtoComponentMap.FindRef(ComponentKey);
 
 		if (RegisteredComponent == Component)
 		{
@@ -6036,7 +6049,7 @@ void ULandscapeInfo::UnregisterActorComponent(ULandscapeComponent* Component)
 		// When removing a key, we need to iterate to find the new bounds
 		XYComponentBounds = FIntRect(MAX_int32, MAX_int32, MIN_int32, MIN_int32);
 
-		for (const auto& XYComponentPair : XYtoComponentMap)
+		for (const TPair<FIntPoint, ULandscapeComponent*>& XYComponentPair : XYtoComponentMap)
 	{
 			XYComponentBounds.Include(XYComponentPair.Key);
 	}
diff --git a/Engine/Source/Runtime/Landscape/Private/LandscapeBlueprintBrushBase.cpp b/Engine/Source/Runtime/Landscape/Private/LandscapeBlueprintBrushBase.cpp
--- a/Engine/Source/Runtime/Landscape/Private/LandscapeBlueprintBrushBase.cpp
+++ b/Engine/Source/Runtime/Landscape/Private/LandscapeBlueprintBrushBase.cpp
@@ -4,14 +4,20 @@
 #include "CoreMinimal.h"
 #include "LandscapeProxy.h"
 #include "Landscape.h"
+#include "LandscapeEditResourcesSubsystem.h"
 #include "LandscapeEditTypes.h"
 #include "LandscapeInfo.h"
 #include "LandscapeLayerInfoObject.h"
 #include "LandscapePrivate.h"
+#include "ProfilingDebugging/RealtimeGPUProfiler.h"
 #include "Misc/MapErrors.h"
 #include "Misc/UObjectToken.h"
 #include "Logging/MessageLog.h"
 #include "Logging/TokenizedMessage.h"
+#include "Engine/Engine.h"
+#include "TextureResource.h"
+#include "RHITransition.h"
+#include "Materials/MaterialExpressionLandscapeVisibilityMask.h"
 
 #include UE_INLINE_GENERATED_CPP_BY_NAME(LandscapeBlueprintBrushBase)
 
@@ -26,6 +32,20 @@ static TAutoConsoleVariable<int32> CVarLandscapeBrushPadding(
 	TEXT("The number of frames to wait before pushing a full Landscape update when a brush is calling RequestLandscapeUpdate"));
 #endif
 
+
+// ----------------------------------------------------------------------------------
+
+FLandscapeBrushParameters::FLandscapeBrushParameters(bool bInIsHeightmapMerge, const FTransform& InRenderAreaWorldTransform, const FIntPoint& InRenderAreaSize, UTextureRenderTarget2D* InCombinedResult, const FName& InWeightmapLayerName /*= FName()*/)
+	: RenderAreaWorldTransform(InRenderAreaWorldTransform)
+	, RenderAreaSize(InRenderAreaSize)
+	, CombinedResult(InCombinedResult)
+	, LayerType(bInIsHeightmapMerge ? ELandscapeToolTargetType::Heightmap : (InWeightmapLayerName == UMaterialExpressionLandscapeVisibilityMask::ParameterName) ? ELandscapeToolTargetType::Visibility : ELandscapeToolTargetType::Weightmap)
+	, WeightmapLayerName(InWeightmapLayerName)
+{}
+
+
+// ----------------------------------------------------------------------------------
+
 ALandscapeBlueprintBrushBase::ALandscapeBlueprintBrushBase(const FObjectInitializer& ObjectInitializer)
 #if WITH_EDITORONLY_DATA
 	: OwningLandscape(nullptr)
@@ -35,7 +55,7 @@ ALandscapeBlueprintBrushBase::ALandscapeBlueprintBrushBase(const FObjectInitiali
 	, AffectVisibilityLayer(false)
 	, bIsVisible(true)
 	, LastRequestLayersContentUpdateFrameNumber(InvalidLastRequestLayersContentUpdateFrameNumber)
-#endif
+#endif // WITH_EDITORONLY_DATA
 {
 #if WITH_EDITOR
 	USceneComponent* SceneComp = CreateDefaultSubobject<USceneComponent>(TEXT("RootComponent"));
@@ -50,11 +70,59 @@ ALandscapeBlueprintBrushBase::ALandscapeBlueprintBrushBase(const FObjectInitiali
 
 #if WITH_EDITORONLY_DATA
 	bIsSpatiallyLoaded = false;
-#endif
+#endif // WITH_EDITORONLY_DATA
 }
 
 PRAGMA_DISABLE_DEPRECATION_WARNINGS
 
+#if WITH_EDITOR
+UTextureRenderTarget2D* ALandscapeBlueprintBrushBase::Execute(const FLandscapeBrushParameters& InParameters)
+{
+	TRACE_CPUPROFILER_EVENT_SCOPE(ALandscapeBlueprintBrushBase::Execute);
+	if ((InParameters.CombinedResult == nullptr) || (OwningLandscape == nullptr))
+	{
+		return nullptr;
+	}
+
+	// Do the render params require a new call to initialize?
+	const FIntPoint NewLandscapeRenderTargetSize = FIntPoint(InParameters.CombinedResult->SizeX, InParameters.CombinedResult->SizeY);
+	if (!CurrentRenderAreaWorldTransform.Equals(InParameters.RenderAreaWorldTransform) || (CurrentRenderAreaSize != InParameters.RenderAreaSize) || CurrentRenderTargetSize != NewLandscapeRenderTargetSize)
+	{
+		CurrentRenderAreaWorldTransform = InParameters.RenderAreaWorldTransform;
+		CurrentRenderAreaSize = InParameters.RenderAreaSize;
+		CurrentRenderTargetSize = NewLandscapeRenderTargetSize;
+
+		TGuardValue<bool> AutoRestore(GAllowActorScriptExecutionInEditor, true);
+		Initialize(CurrentRenderAreaWorldTransform, CurrentRenderAreaSize, CurrentRenderTargetSize);
+	}
+
+	// Time to render :
+	FString LayerDetailString;
+	if (InParameters.LayerType != ELandscapeToolTargetType::Heightmap)
+	{
+		LayerDetailString = FString::Format(TEXT(" ({0})"), { *InParameters.WeightmapLayerName.ToString() });
+	}
+	UTextureRenderTarget2D* Result = nullptr;
+	{
+		SCOPED_DRAW_EVENTF_GAMETHREAD(LandscapeLayers, TEXT("BP Render (%s): %s"), UEnum::GetValueAsString(InParameters.LayerType), LayerDetailString);
+
+		TGuardValue<bool> AutoRestore(GAllowActorScriptExecutionInEditor, true);
+		Result = RenderLayer(InParameters);
+	}
+
+	// If the BP brush failed to render, we still need to passthrough from the read RT to the write RT in order not to lose what has been merged so far : 
+	if ((Result != nullptr) &&
+		((Result->SizeX != InParameters.CombinedResult->SizeX) || (Result->SizeY != InParameters.CombinedResult->SizeY)))
+	{
+		UE_LOG(LogLandscape, Warning, TEXT("In landscape %s, the BP brush %s failed to render for (%s%s). Make sure the brush properly implements RenderLayer and returns a render target of the appropriate size: expected (%i, %i), actual (%i, %i). This brush will be skipped until then."), 
+			*OwningLandscape->GetActorLabel(), *GetActorLabel(), *UEnum::GetValueAsString(InParameters.LayerType), *LayerDetailString, InParameters.CombinedResult->SizeX, InParameters.CombinedResult->SizeY, Result->SizeX, Result->SizeY );
+		Result = nullptr;
+	}
+
+	return Result;
+}
+#endif // WITH_EDITOR
+
 // Deprecated
 UTextureRenderTarget2D* ALandscapeBlueprintBrushBase::Render_Implementation(bool InIsHeightmap, UTextureRenderTarget2D* InCombinedResult, const FName& InWeightmapLayerName)
 {
@@ -83,16 +151,16 @@ void ALandscapeBlueprintBrushBase::Initialize_Implementation(const FTransform& I
 
 void ALandscapeBlueprintBrushBase::RequestLandscapeUpdate(bool bInUserTriggered)
 {
-#if WITH_EDITORONLY_DATA
+#if WITH_EDITOR
 	UE_LOG(LogLandscape, Verbose, TEXT("ALandscapeBlueprintBrushBase::RequestLandscapeUpdate"));
 	if (OwningLandscape)
 	{
 		uint32 ModeMask = 0;
-		if (AffectsHeightmap())
+		if (CanAffectHeightmap())
 		{
 			ModeMask |= ELandscapeLayerUpdateMode::Update_Heightmap_Editing_NoCollision;
 		}
-		if (AffectsWeightmap() || AffectsVisibilityLayer())
+		if (CanAffectWeightmap() || CanAffectVisibilityLayer())
 		{
 			ModeMask |= ELandscapeLayerUpdateMode::Update_Weightmap_Editing_NoCollision;
 		}
@@ -103,13 +171,215 @@ void ALandscapeBlueprintBrushBase::RequestLandscapeUpdate(bool bInUserTriggered)
 			LastRequestLayersContentUpdateFrameNumber = GFrameNumber == InvalidLastRequestLayersContentUpdateFrameNumber ? GFrameNumber + 1 : GFrameNumber;
 		}
 	}
-#endif
+#endif // WITH_EDITOR
 }
 
 #if WITH_EDITOR
+
+void ALandscapeBlueprintBrushBase::GetRendererStateInfo(const ULandscapeInfo* InLandscapeInfo,
+	UE::Landscape::EditLayers::FEditLayerTargetTypeState& OutSupportedTargetTypeState, UE::Landscape::EditLayers::FEditLayerTargetTypeState& OutEnabledTargetTypeState, TArray<TSet<FName>>& OutRenderGroups) const
+{
+	// What can the brush do?
+	if (CanAffectHeightmap())
+	{
+		OutSupportedTargetTypeState.AddTargetTypeMask(ELandscapeToolTargetTypeFlags::Heightmap);
+	}
+	if (CanAffectWeightmap())
+	{
+		OutSupportedTargetTypeState.AddTargetTypeMask(ELandscapeToolTargetTypeFlags::Weightmap);
+	}
+	if (CanAffectVisibilityLayer())
+	{
+		OutSupportedTargetTypeState.AddTargetTypeMask(ELandscapeToolTargetTypeFlags::Visibility);
+	}
+
+	// What does it currently do?
+	if (AffectsHeightmap())
+	{
+		OutEnabledTargetTypeState.AddTargetTypeMask(ELandscapeToolTargetTypeFlags::Heightmap);
+	}
+	if (AffectsWeightmap())
+	{
+		OutEnabledTargetTypeState.AddTargetTypeMask(ELandscapeToolTargetTypeFlags::Weightmap);
+	}
+	if (AffectsVisibilityLayer())
+	{
+		OutEnabledTargetTypeState.AddTargetTypeMask(ELandscapeToolTargetTypeFlags::Visibility);
+	}
+
+	// Mark which weightmap is supported/enabled : 
+	if (CanAffectWeightmap())
+	{
+		for (const FLandscapeInfoLayerSettings& InfoLayerSettings : InLandscapeInfo->Layers)
+		{
+			if (InfoLayerSettings.LayerInfoObj != nullptr)
+			{
+				FName WeightmapLayerName = InfoLayerSettings.GetLayerName();
+				if (CanAffectWeightmapLayer(WeightmapLayerName))
+				{
+					OutSupportedTargetTypeState.AddWeightmap(WeightmapLayerName);
+					if (AffectsWeightmapLayer(WeightmapLayerName))
+					{
+						OutEnabledTargetTypeState.AddWeightmap(WeightmapLayerName);
+					}
+				}
+			}
+		}
+	}
+}
+
+TArray<UE::Landscape::EditLayers::FEditLayerRenderItem> ALandscapeBlueprintBrushBase::GetRenderItems(const ULandscapeInfo* InLandscapeInfo) const
+{
+	using namespace UE::Landscape::EditLayers;
+
+	FEditLayerTargetTypeState SupportedTargetTypeState, EnabledTargetTypeState;
+	TArray<TSet<FName>> DummyRenderGroups;
+	GetRendererStateInfo(InLandscapeInfo, SupportedTargetTypeState, EnabledTargetTypeState, DummyRenderGroups);
+
+	// By default, for landscape BP brushes, we use FInputWorldArea::EType::Infinite, to indicate they can only reliably work when applied globally on the entire landscape
+	//  This allows full backwards-compatibility but will prevent landscapes from benefiting from batched merge. Users will be able to indicate their brush works in a local fashion
+	//  by overriding this and using another type of input world area
+	FInputWorldArea InputWorldArea(FInputWorldArea::CreateInfinite());
+	// By default, the brush only writes into the component itself (i.e. it renders to the area that it's currently being asked to render to):
+	FOutputWorldArea OutputWorldArea(FOutputWorldArea::CreateLocalComponent());
+	
+	// Use EnabledTargetTypeState because we only want to tell what we'll actually be able to render to (instead of what we'd potentially be able to render to, i.e. what is "supported" by the brush) : 
+	return { FEditLayerRenderItem(EnabledTargetTypeState, InputWorldArea, OutputWorldArea, /*bModifyExistingWeightmapsOnly = */false)};
+}
+
+void ALandscapeBlueprintBrushBase::RenderLayer(ILandscapeEditLayerRenderer::FRenderParams& InRenderParams)
+{
+	using namespace UE::Landscape;
+
+	// By default, use the old way of rendering BP brushes : 
+
+	// Swap the render targets so that the layer's input RT is the latest combined result :
+	// The write render target will be accessed as ERHIAccess::CopyDest all along : 
+	InRenderParams.MergeRenderContext->CycleBlendRenderTargets(/*InDesiredWriteAccess = */ERHIAccess::CopyDest);
+	ULandscapeScratchRenderTarget* WriteRT = InRenderParams.MergeRenderContext->GetBlendRenderTargetWrite();
+	ULandscapeScratchRenderTarget* CurrentLayerReadRT = InRenderParams.MergeRenderContext->GetBlendRenderTargetRead();
+
+	const bool bIsHeightmapMerge = InRenderParams.MergeRenderContext->IsHeightmapMerge();
+
+	// Because we only expose UTextureRenderTarget2D to BP, in the case of weightmaps, we need an additional scratch render target 2D that we will copy the 
+	//  current result of each paint layer into, so that the BP can use it as its source : 
+	ULandscapeEditResourcesSubsystem* LandscapeEditResourcesSubsystem = GEngine->GetEngineSubsystem<ULandscapeEditResourcesSubsystem>();
+	check(LandscapeEditResourcesSubsystem != nullptr);
+	ULandscapeScratchRenderTarget* CurrentLayerReadRT2D = CurrentLayerReadRT;
+
+	// We might require an additional scratch render target (it will get freed/recycled at the end of this function when the AdditionalScratchTexture goes out of scope) : 
+	TUniquePtr<FScratchRenderTargetScope> AdditionalScratchTexture;
+
+	TArray<FName> EnabledWeightmaps;
+	if (!bIsHeightmapMerge)
+	{
+		EnabledWeightmaps = InRenderParams.RendererState.GetEnabledTargetWeightmaps();
+
+		check(CurrentLayerReadRT->IsTexture2DArray() && WriteRT->IsTexture2DArray());
+		FScratchRenderTargetParams ScratchRenderTargetParams(TEXT("BPBrushScratchRT"), /*bInExactDimensions = */false, /*bInUseUAV = */false, /*bInTargetArraySlicesIndependently = */false,
+			CurrentLayerReadRT->GetResolution(), /*InNumSlices = */0, CurrentLayerReadRT->GetFormat(), CurrentLayerReadRT->GetClearColor(), ERHIAccess::CopyDest);
+		// We need a new scratch 2D texture in order to copy the merged result of a single target layer, so that the BP brush can use that as input (as it operates on texture 2D render targets only): 
+		AdditionalScratchTexture = MakeUnique<FScratchRenderTargetScope>(ScratchRenderTargetParams);
+		CurrentLayerReadRT2D = AdditionalScratchTexture->RenderTarget;
+
+		// The original texture array will be accessed as ERHIAccess::CopySrc all along : 
+		CurrentLayerReadRT->TransitionTo(ERHIAccess::CopySrc);
+	}
+
+	TArray<FName> TargetLayerNames = InRenderParams.MergeRenderContext->ConvertTargetLayerBitIndicesToNamesChecked(InRenderParams.RenderGroupBitIndices);
+	const int32 NumTargetLayers = TargetLayerNames.Num();
+	for (int32 TargetLayerIndex = 0; TargetLayerIndex < NumTargetLayers; ++TargetLayerIndex)
+	{
+		const FName TargetLayerName = TargetLayerNames[TargetLayerIndex];
+		SCOPED_DRAW_EVENTF_GAMETHREAD(LandscapeLayers, TEXT("Render %s"), TargetLayerName);
+
+		// If necessary, copy from the texture array's slice to the scratch render target 2D : 
+		if (!bIsHeightmapMerge)
+		{
+			SCOPED_DRAW_EVENTF_GAMETHREAD(LandscapeLayers, TEXT("Copy Source (slice %i) -> %s"), TargetLayerIndex, CurrentLayerReadRT2D->GetDebugName());
+
+			ULandscapeScratchRenderTarget::FCopyFromScratchRenderTargetParams CopyParams(CurrentLayerReadRT);
+			// Copy from the proper slice in the texture array : 
+			CopyParams.SourceSliceIndex = TargetLayerIndex;
+			CurrentLayerReadRT2D->CopyFrom(CopyParams);
+			CurrentLayerReadRT2D->TransitionTo(ERHIAccess::SRVMask);
+		}
+
+		check(CurrentLayerReadRT2D->GetCurrentState() == ERHIAccess::SRVMask);
+		check(WriteRT->GetCurrentState() == ERHIAccess::CopyDest);
+
+		UTextureRenderTarget2D* ReadRT2D = CurrentLayerReadRT2D->GetRenderTarget2D();
+		// If the BP brush failed to render, we still need to passthrough from the read RT to the write RT in order not to lose what has been merged so far : 
+		UTextureRenderTarget2D* OutputRT2D = ReadRT2D;
+
+		// Only render the target layer if it's effectively enabled for this merge : it's possible there are target layers in the render group that we don't support or are not enabled so we have to 
+		//  do the validation here first :
+		if (!bIsHeightmapMerge && EnabledWeightmaps.Contains(TargetLayerName))
+		{
+			// Execute (i.e. (Initialize/)Render the BP brush) : 
+			FLandscapeBrushParameters BrushParameters(bIsHeightmapMerge, InRenderParams.RenderAreaWorldTransform, InRenderParams.RenderAreaSectionRect.Size(), ReadRT2D, TargetLayerName);
+			if (UTextureRenderTarget2D* BrushOutputRT2D = Execute(BrushParameters))
+			{
+				// Only consider the BP brush's result if it's valid :
+				OutputRT2D = BrushOutputRT2D;
+			}
+		}
+
+		// TODO: handle conversion/handling of RT not same size as internal size
+		check((OutputRT2D->SizeX == OutputRT2D->SizeX) && (OutputRT2D->SizeY == ReadRT2D->SizeY));
+
+		// Resolve back to the write RT 
+		{
+			SCOPED_DRAW_EVENTF_GAMETHREAD(LandscapeLayers, TEXT("Copy BP Render Result -> %s (slice %i)"), WriteRT->GetDebugName(), TargetLayerIndex);
+
+			// The RT returned by the brush is in SRV state so we need a transition: 
+			ENQUEUE_RENDER_COMMAND(TransitionToCopySrc)([Resource = OutputRT2D->GetResource()](FRHICommandListImmediate& InRHICmdList) mutable
+			{
+				InRHICmdList.Transition(FRHITransitionInfo(Resource->TextureRHI, ERHIAccess::SRVMask, ERHIAccess::CopySrc));
+			});
+
+			ULandscapeScratchRenderTarget::FCopyFromTextureParams CopyParams(OutputRT2D);
+			CopyParams.DestSliceIndex = TargetLayerIndex;
+			WriteRT->CopyFrom(CopyParams);
+
+			// It's also expected we leave the RT return by the BP as SRV :
+			ENQUEUE_RENDER_COMMAND(TransitionToSRV)([Resource = OutputRT2D->GetResource()](FRHICommandListImmediate& InRHICmdList) mutable
+			{
+				InRHICmdList.Transition(FRHITransitionInfo(Resource->TextureRHI, ERHIAccess::CopySrc, ERHIAccess::SRVMask));
+			});
+		}
+	}
+
+	// Leave the render targets in the state they're expected to be in: 
+	WriteRT->TransitionTo(ERHIAccess::RTV);
+	CurrentLayerReadRT->TransitionTo(ERHIAccess::SRVMask);
+}
+
+FString ALandscapeBlueprintBrushBase::GetEditLayerRendererDebugName() const
+{
+	return GetActorNameOrLabel();
+}
+
+TArray<UE::Landscape::EditLayers::FEditLayerRendererState> ALandscapeBlueprintBrushBase::GetEditLayerRendererStates(const ULandscapeInfo* InLandscapeInfo, bool bInSkipBrush)
+{
+	using namespace UE::Landscape::EditLayers;
+
+	if (OwningLandscape == nullptr)
+	{
+		return {};
+	}
+
+	FEditLayerRendererState RendererState(this, InLandscapeInfo);
+	// Force the renderer to be fully disabled in case we are asked to skip the brush :
+	if (bInSkipBrush)
+	{
+		RendererState.DisableTargetTypeMask(ELandscapeToolTargetTypeFlags::All);
+	}
+	return { RendererState };
+}
+
 void ALandscapeBlueprintBrushBase::PushDeferredLayersContentUpdate()
 {
-#if WITH_EDITORONLY_DATA
 	// Avoid computing collision and client updates every frame
 	// Wait until we didn't trigger any more landscape update requests (padding of a couple of frames)
 	if (OwningLandscape != nullptr &&
@@ -131,7 +401,6 @@ void ALandscapeBlueprintBrushBase::PushDeferredLayersContentUpdate()
 		}
 		LastRequestLayersContentUpdateFrameNumber = InvalidLastRequestLayersContentUpdateFrameNumber;
 	}
-#endif
 }
 
 void ALandscapeBlueprintBrushBase::Tick(float DeltaSeconds)
@@ -158,19 +427,18 @@ bool ALandscapeBlueprintBrushBase::IsLayerUpdatePending() const
 
 void ALandscapeBlueprintBrushBase::SetIsVisible(bool bInIsVisible)
 {
-#if WITH_EDITORONLY_DATA
 	Modify();
 	bIsVisible = bInIsVisible;
 	if (OwningLandscape)
 	{
 		OwningLandscape->OnBlueprintBrushChanged();
 	}
-#endif
 }
+#endif // WITH_EDITOR
 
 void ALandscapeBlueprintBrushBase::SetCanAffectHeightmap(bool bInCanAffectHeightmap)
 {
-#if WITH_EDITORONLY_DATA
+#if WITH_EDITOR
 	if (bInCanAffectHeightmap != AffectHeightmap)
 	{
 		Modify();
@@ -180,12 +448,12 @@ void ALandscapeBlueprintBrushBase::SetCanAffectHeightmap(bool bInCanAffectHeight
 			OwningLandscape->OnBlueprintBrushChanged();
 		}
 	}
-#endif // WITH_EDITORONLY_DATA
+#endif // WITH_EDITOR
 }
 
 void ALandscapeBlueprintBrushBase::SetCanAffectWeightmap(bool bInCanAffectWeightmap)
 {
-#if WITH_EDITORONLY_DATA
+#if WITH_EDITOR
 	if (bInCanAffectWeightmap != AffectWeightmap)
 	{
 		Modify();
@@ -195,12 +463,12 @@ void ALandscapeBlueprintBrushBase::SetCanAffectWeightmap(bool bInCanAffectWeight
 			OwningLandscape->OnBlueprintBrushChanged();
 		}
 	}
-#endif // WITH_EDITORONLY_DATA
+#endif // WITH_EDITOR
 }
 
 void ALandscapeBlueprintBrushBase::SetCanAffectVisibilityLayer(bool bInCanAffectVisibilityLayer)
 {
-#if WITH_EDITORONLY_DATA
+#if WITH_EDITOR
 	if (bInCanAffectVisibilityLayer != AffectVisibilityLayer)
 	{
 		Modify();
@@ -210,52 +478,57 @@ void ALandscapeBlueprintBrushBase::SetCanAffectVisibilityLayer(bool bInCanAffect
 			OwningLandscape->OnBlueprintBrushChanged();
 		}
 	}
-#endif // WITH_EDITORONLY_DATA
+#endif // WITH_EDITOR
+}
+
+#if WITH_EDITOR
+bool ALandscapeBlueprintBrushBase::CanAffectWeightmapLayer(const FName& InLayerName) const
+{
+	if (!CanAffectWeightmap())
+	{
+		return false;
+	}
+
+	// By default, it's the same implementation as AffectsWeightmapLayer : if the weightmap layer name is in our list, consider we can affect it :
+	//  CanAffectWeightmapLayer can be overridden in child classes that don't use AffectedWeightmapLayers to list the weightmaps they can affect
+	return AffectedWeightmapLayers.Contains(InLayerName);
 }
 
 bool ALandscapeBlueprintBrushBase::AffectsWeightmapLayer(const FName& InLayerName) const
 {
-#if WITH_EDITORONLY_DATA
 	if (!CanAffectWeightmap())
 	{
 		return false;
 	}
 
+	// By default, it's the same implementation as CanAffectWeightmapLayer : if the weightmap layer name is in our list, consider we do affect it :
+	//  AffectsWeightmapLayer can be overridden in child classes that don't use AffectedWeightmapLayers to list the weightmaps they're currently affecting :
 	return AffectedWeightmapLayers.Contains(InLayerName);
-#else // WITH_EDITORONLY_DATA
-	return false;
-#endif // !WITH_EDITORONLY_DATA
 }
 
 void ALandscapeBlueprintBrushBase::PostEditMove(bool bFinished)
 {
 	Super::PostEditMove(bFinished);
-#if WITH_EDITORONLY_DATA
 	RequestLandscapeUpdate();
-#endif
 }
 
 void ALandscapeBlueprintBrushBase::PostEditChangeProperty(FPropertyChangedEvent& PropertyChangedEvent)
 {
 	Super::PostEditChangeProperty(PropertyChangedEvent);
-#if WITH_EDITORONLY_DATA
 	if (OwningLandscape && UpdateOnPropertyChange)
 	{
 		OwningLandscape->OnBlueprintBrushChanged();
 	}
-#endif
 }
 
 void ALandscapeBlueprintBrushBase::Destroyed()
 {
 	Super::Destroyed();
-#if WITH_EDITORONLY_DATA
 	if (OwningLandscape && !GIsReinstancing)
 	{
 		OwningLandscape->RemoveBrush(this);
 	}
 	OwningLandscape = nullptr;
-#endif
 }
 
 void ALandscapeBlueprintBrushBase::CheckForErrors()
@@ -284,7 +557,6 @@ void ALandscapeBlueprintBrushBase::GetRenderDependencies(TSet<UObject*>& OutDepe
 
 void ALandscapeBlueprintBrushBase::SetOwningLandscape(ALandscape* InOwningLandscape)
 {
-#if WITH_EDITORONLY_DATA
 	if (OwningLandscape == InOwningLandscape)
 	{
 		return;
@@ -304,16 +576,11 @@ void ALandscapeBlueprintBrushBase::SetOwningLandscape(ALandscape* InOwningLandsc
 	{
 		OwningLandscape->OnBlueprintBrushChanged();
 	}
-#endif
 }
 
 ALandscape* ALandscapeBlueprintBrushBase::GetOwningLandscape() const
 {
-#if WITH_EDITORONLY_DATA
 	return OwningLandscape;
-#else
-	return nullptr;
-#endif
 }
 #endif // WITH_EDITOR
 
