Introducing landscape batched merge (activated by CVar : landscape.EditLayersLocalMerge.Enable 2): complete refactor of the landscape edit layers merge algorithm to solve the taxing GPU memory require

diff --git a/Engine/Source/Runtime/Landscape/Private/LandscapeEdit.cpp b/Engine/Source/Runtime/Landscape/Private/LandscapeEdit.cpp
--- a/Engine/Source/Runtime/Landscape/Private/LandscapeEdit.cpp
+++ b/Engine/Source/Runtime/Landscape/Private/LandscapeEdit.cpp
@@ -83,6 +83,7 @@ LandscapeEdit.cpp: Landscape editing
 #include "WorldPartition/WorldPartitionActorDescInstance.h"
 #include "ActorPartition/ActorPartitionSubsystem.h"
 #include "LandscapeUtils.h"
+#include "LandscapeUtilsPrivate.h"
 #include "LandscapeSplineActor.h"
 #include "Materials/MaterialExpressionLandscapeGrassOutput.h"
 #include "ShaderPlatformCachedIniValue.h"
@@ -5698,11 +5699,11 @@ ALandscapeProxy* ULandscapeInfo::MoveComponentsToProxy(const TArray<ULandscapeCo
 		FLandscapeEditDataInterface LandscapeEdit(this);
 		for (ULandscapeComponent* Component : TargetSelectedComponents)
 		{
-			Component->ReallocateWeightmaps(&LandscapeEdit, false, true, true, LandscapeProxy);
+			Component->ReallocateWeightmaps(/*DataInterface = */&LandscapeEdit, /*InEditLayerGuid = */FGuid(), /*bInSaveToTransactionBuffer = */true, /*bool bInForceReallocate = */true, /*InTargetProxy = */LandscapeProxy, /*InRestrictSharingToComponents = */nullptr);
 			Component->ForEachLayer([&](const FGuid& LayerGuid, FLandscapeLayerComponentData& LayerData)
 			{
 				FScopedSetLandscapeEditingLayer Scope(Landscape, LayerGuid);
-				Component->ReallocateWeightmaps(&LandscapeEdit, true, true, true, LandscapeProxy);
+				Component->ReallocateWeightmaps(/*DataInterface = */&LandscapeEdit, LayerGuid, /*bInSaveToTransactionBuffer = */true, /*bool bInForceReallocate = */true, /*InTargetProxy = */LandscapeProxy, /*InRestrictSharingToComponents = */nullptr);
 			});
 			Landscape->RequestLayersContentUpdateForceAll();
 		}
@@ -7076,14 +7077,21 @@ void ULandscapeInfo::ClearSelectedRegion(bool bIsComponentwise /*= true*/)
 	}
 }
 
-void ULandscapeComponent::ReallocateWeightmaps(FLandscapeEditDataInterface* DataInterface, bool InCanUseEditingWeightmap, bool InSaveToTransactionBuffer, bool InForceReallocate, ALandscapeProxy* InTargetProxy, TArray<UTexture*>* OutNewCreatedTextures)
+void ULandscapeComponent::ReallocateWeightmaps(FLandscapeEditDataInterface* DataInterface, bool InCanUseEditingWeightmap, bool InSaveToTransactionBuffer, bool InForceReallocate, 
+	ALandscapeProxy* InTargetProxy, TArray<UTexture*>* OutNewCreatedTextures)
 {
 	FGuid TargetLayerGuid = InCanUseEditingWeightmap ? GetEditingLayerGUID() : FGuid();
-	ReallocateWeightmapsInternal(DataInterface, TargetLayerGuid, InSaveToTransactionBuffer, InForceReallocate, InTargetProxy, OutNewCreatedTextures);
+	TArray<UTexture*> CreatedTextures = ReallocateWeightmaps(DataInterface, TargetLayerGuid, InSaveToTransactionBuffer, InForceReallocate, InTargetProxy, /*InRestrictSharingToComponents = */nullptr);
+	if (OutNewCreatedTextures != nullptr)
+	{
+		*OutNewCreatedTextures = CreatedTextures;
+	}
 }
 
-void ULandscapeComponent::ReallocateWeightmapsInternal(FLandscapeEditDataInterface* DataInterface, const FGuid& InEditLayerGuid, bool InSaveToTransactionBuffer, bool InForceReallocate, ALandscapeProxy* InTargetProxy, TArray<UTexture*>* OutNewCreatedTextures)
+TArray<UTexture*> ULandscapeComponent::ReallocateWeightmaps(FLandscapeEditDataInterface* DataInterface, const FGuid& InEditLayerGuid, bool bInSaveToTransactionBuffer, bool bInForceReallocate, 
+	ALandscapeProxy* InTargetProxy, TSet<ULandscapeComponent*>* InRestrictSharingToComponents)
 {
+	TArray<UTexture*> CreatedTextures;
 	int32 NeededNewChannels = 0;
 	ALandscapeProxy* TargetProxy = InTargetProxy ? InTargetProxy : GetLandscapeProxy();
 	
@@ -7099,7 +7107,7 @@ void ULandscapeComponent::ReallocateWeightmapsInternal(FLandscapeEditDataInterfa
 	TArray<TObjectPtr<ULandscapeWeightmapUsage>>& ComponentWeightmapTexturesUsage = GetWeightmapTexturesUsage(InEditLayerGuid);
 
 	// When force reallocating, skip tests to see if allocations are necessary based on Component's WeightmapLayeAllocInfo
-	if (!InForceReallocate)
+	if (!bInForceReallocate)
 	{
 		for (int32 LayerIdx = 0; LayerIdx < ComponentWeightmapLayerAllocations.Num(); LayerIdx++)
 		{
@@ -7112,20 +7120,20 @@ void ULandscapeComponent::ReallocateWeightmapsInternal(FLandscapeEditDataInterfa
 		// All channels allocated!
 		if (NeededNewChannels == 0)
 		{
-			return;
+			return {};
 		}
 	}
 
 	ULandscapeInfo* LandscapeInfo = GetLandscapeInfo();
 
 	bool bMarkPackageDirty = DataInterface == nullptr ? true : DataInterface->GetShouldDirtyPackage();
-	if (InSaveToTransactionBuffer)
+	if (bInSaveToTransactionBuffer)
 	{
 		LandscapeInfo->ModifyObject(this, bMarkPackageDirty);
 		LandscapeInfo->ModifyObject(TargetProxy, bMarkPackageDirty);
 	}
 
-	if (!InForceReallocate)
+	if (!bInForceReallocate)
 	{
 		// UE_LOG(LogLandscape, Log, TEXT("----------------------"));
 		// UE_LOG(LogLandscape, Log, TEXT("Component %s needs %d layers (%d new)"), *GetName(), WeightmapLayerAllocations.Num(), NeededNewChannels);
@@ -7188,7 +7196,7 @@ void ULandscapeComponent::ReallocateWeightmapsInternal(FLandscapeEditDataInterfa
 
 						if (NeededNewChannels == 0)
 						{
-							return;
+							return {};
 						}
 					}
 				}
@@ -7241,6 +7249,17 @@ void ULandscapeComponent::ReallocateWeightmapsInternal(FLandscapeEditDataInterfa
 								int32 TryDistanceSquared = (TryWeightmapUsage->ChannelUsage[ChanIdx]->GetSectionBase() - GetSectionBase()).SizeSquared();
 								if (TryDistanceSquared < BestDistanceSquared)
 								{
+									if (InRestrictSharingToComponents != nullptr)
+									{
+										TSet<ULandscapeComponent*, DefaultKeyFuncs<ULandscapeComponent*>, TInlineSetAllocator<4>> WeightmapComponents(MakeArrayView(TryWeightmapUsage->GetUniqueValidComponents()));
+										int32 NumWeightmapComponents = WeightmapComponents.Num();
+										// Don't pick this candidate if it will lead the texture to be shared with a component that's not in the provided list :
+										if ((NumWeightmapComponents > 0) && WeightmapComponents.Intersect(*InRestrictSharingToComponents).Num() != NumWeightmapComponents)
+										{
+											break;
+										}
+									}
+
 									CurrentWeightmapTexture = ItPair.Key;
 									CurrentWeightmapUsage = TryWeightmapUsage;
 									BestDistanceSquared = TryDistanceSquared;
@@ -7283,10 +7302,7 @@ void ULandscapeComponent::ReallocateWeightmapsInternal(FLandscapeEditDataInterfa
 
 			CurrentWeightmapTexture->PostEditChange();
 
-			if (OutNewCreatedTextures != nullptr)
-			{
-				OutNewCreatedTextures->Add(CurrentWeightmapTexture);
-			}
+			CreatedTextures.Add(CurrentWeightmapTexture);
 
 			// Store it in the usage map
 			CurrentWeightmapUsage = TargetProxy->WeightmapUsageMap.Add(CurrentWeightmapTexture, TargetProxy->CreateWeightmapUsage());
@@ -7368,6 +7384,8 @@ void ULandscapeComponent::ReallocateWeightmapsInternal(FLandscapeEditDataInterfa
 	SetWeightmapTexturesUsageInternal(MoveTemp(NewComponentWeightmapTexturesUsage), InEditLayerGuid);
 
 	TargetProxy->ValidateProxyLayersWeightmapUsage();
+
+	return CreatedTextures;
 }
 
 void ALandscapeProxy::RemoveInvalidWeightmaps()
@@ -7576,7 +7594,7 @@ void ULandscapeComponent::InitWeightmapData(TArray<ULandscapeLayerInfoObject*>&
 		new (WeightmapLayerAllocations)FWeightmapLayerAllocationInfo(LayerInfos[Idx]);
 	}
 
-	ReallocateWeightmaps();
+	ReallocateWeightmaps(/*DataInterface = */nullptr, GetEditingLayerGUID(), /*bInSaveToTransactionBuffer = */true, /*bool bInForceReallocate = */false, /*InTargetProxy = */nullptr, /*InRestrictSharingToComponents = */nullptr);
 
 	check(WeightmapLayerAllocations.Num() > 0 && WeightmapTextures.Num() > 0);
 
@@ -7769,7 +7787,7 @@ void ULandscapeComponent::GenerateMobilePlatformPixelData(bool bIsCooking, const
 
 	const int32 NumWeightTextures = FMath::DivideAndRoundUp(static_cast<int32>(Algo::CountIf(MobileWeightmapLayerAllocations, [](const FWeightmapLayerAllocationInfo& AllocationInfo) { return AllocationInfo.LayerInfo; })), 4);
 
-	const bool MobileWeightmapTextureArrayEnabled = UE::Landscape::IsMobileWeightmapTextureArrayEnabled();
+	const bool MobileWeightmapTextureArrayEnabled = UE::Landscape::Private::IsMobileWeightmapTextureArrayEnabled();
 	
 	if (MobileWeightmapTextureArrayEnabled && NumWeightTextures > 0)
 	{
diff --git a/Engine/Source/Runtime/Landscape/Private/LandscapeEditInterface.cpp b/Engine/Source/Runtime/Landscape/Private/LandscapeEditInterface.cpp
--- a/Engine/Source/Runtime/Landscape/Private/LandscapeEditInterface.cpp
+++ b/Engine/Source/Runtime/Landscape/Private/LandscapeEditInterface.cpp
@@ -1897,7 +1897,7 @@ void ULandscapeComponent::FillLayer(ULandscapeLayerInfoObject* LayerInfo, FLands
 	{
 		FillLayerIdx = ComponentWeightmapLayerAllocations.Num();
 		ComponentWeightmapLayerAllocations.Add(FWeightmapLayerAllocationInfo(LayerInfo));
-		Component->ReallocateWeightmapsInternal(&LandscapeEdit, EditLayerGuid);
+		Component->ReallocateWeightmaps(/*DataInterface =*/&LandscapeEdit, EditLayerGuid, /*bInSaveToTransactionBuffer = */true, /*bool bInForceReallocate = */false, /*InTargetProxy = */nullptr, /*InRestrictSharingToComponents = */nullptr);
 	}
 
 	check(FillLayerIdx != INDEX_NONE);
@@ -2863,7 +2863,7 @@ void FLandscapeEditDataInterface::SetAlphaData(ULandscapeLayerInfoObject* const
 
 				UpdateLayerIdx = ComponentWeightmapLayerAllocations.Num();
 				new (ComponentWeightmapLayerAllocations) FWeightmapLayerAllocationInfo(LayerInfo);
-				Component->ReallocateWeightmapsInternal(this, GetEditLayer());
+				Component->ReallocateWeightmaps(/*DataInterface =*/this, GetEditLayer(), /*bInSaveToTransactionBuffer = */true, /*bool bInForceReallocate = */false, /*InTargetProxy = */nullptr, /*InRestrictSharingToComponents = */nullptr);
 
 				if (!Component->GetLandscapeProxy()->HasLayersContent())
 				{
@@ -3317,7 +3317,7 @@ void FLandscapeEditDataInterface::SetAlphaData(const TSet<ULandscapeLayerInfoObj
 						}
 						ComponentWeightmapLayerAllocations.Emplace(LayerInfoNeedingAllocation);
 					}
-					Component->ReallocateWeightmapsInternal(this, GetEditLayer());
+					Component->ReallocateWeightmaps(/*DataInterface =*/this, GetEditLayer(), /*bInSaveToTransactionBuffer = */true, /*bool bInForceReallocate = */false, /*InTargetProxy = */nullptr, /*InRestrictSharingToComponents = */nullptr);
 					
 					if (!Component->GetLandscapeProxy()->HasLayersContent())
 					{
diff --git a/Engine/Source/Runtime/Landscape/Private/LandscapeEditLayer.cpp b/Engine/Source/Runtime/Landscape/Private/LandscapeEditLayer.cpp
--- a/Engine/Source/Runtime/Landscape/Private/LandscapeEditLayer.cpp
+++ b/Engine/Source/Runtime/Landscape/Private/LandscapeEditLayer.cpp
@@ -20,6 +20,22 @@ void ULandscapeEditLayerBase::PostLoad()
 	SetFlags(RF_Transactional);
 }
 
+#if WITH_EDITOR
+
+const FLandscapeLayer* ULandscapeEditLayerBase::GetOwningLayer() const
+{
+	if (OwningLandscape != nullptr)
+	{
+		TArrayView<const FLandscapeLayer> Layers = OwningLandscape->GetLayers();
+		return Layers.FindByPredicate([this](const FLandscapeLayer& Struct) { return Struct.EditLayer == this; });
+	}
+
+	return nullptr;
+}
+
+#endif // WITH_EDITOR
+
+
 // ----------------------------------------------------------------------------------
 
 bool ULandscapeEditLayer::SupportsTargetType(ELandscapeToolTargetType InType) const
diff --git a/Engine/Source/Runtime/Landscape/Private/LandscapeEditLayerRenderer.cpp b/Engine/Source/Runtime/Landscape/Private/LandscapeEditLayerRenderer.cpp
new file mode 100644
--- /dev/null
+++ b/Engine/Source/Runtime/Landscape/Private/LandscapeEditLayerRenderer.cpp
@@ -0,0 +1,1235 @@
+// Copyright Epic Games, Inc. All Rights Reserved.
+
+#include "LandscapeEditLayerRenderer.h"
+#include "Algo/AllOf.h"
+#include "Algo/Transform.h"
+#include "Algo/Count.h"
+#include "Engine/Engine.h"
+#include "EngineModule.h"
+#include "Landscape.h"
+#include "LandscapeComponent.h"
+#include "LandscapeEditResourcesSubsystem.h"
+#include "LandscapeLayerInfoObject.h"
+#include "LandscapePrivate.h"
+#include "LandscapeUtils.h"
+#include "Materials/MaterialExpressionLandscapeVisibilityMask.h"
+#include "PixelShaderUtils.h"
+#include "RenderGraph.h"
+#include "RenderingThread.h"
+#include "RHIAccess.h"
+#include "SceneView.h"
+#include "TextureResource.h"
+#include "VisualLogger/VisualLogger.h"
+
+#define LOCTEXT_NAMESPACE "LandscapeEditLayerRenderer"
+
+extern TAutoConsoleVariable<int32> CVarLandscapeEditLayersMaxResolutionPerRenderBatch;
+extern TAutoConsoleVariable<float> CVarLandscapeBatchedMergeVisualLogOffsetIncrement;
+extern TAutoConsoleVariable<int32> CVarLandscapeEditLayersClearBeforeEachWriteToScratch;
+extern TAutoConsoleVariable<int32> CVarLandscapeBatchedMergeVisualLogShowMergeType;
+extern TAutoConsoleVariable<bool> CVarLandscapeBatchedMergeVisualLogShowMergeProcess;
+extern TAutoConsoleVariable<float> CVarLandscapeBatchedMergeVisualLogAlpha;
+
+namespace UE::Landscape::EditLayers
+{
+
+// ----------------------------------------------------------------------------------
+// /Engine/Private/Landscape/LandscapeEditLayersUtils.usf shaders : 
+
+class FMarkValidityPS : public FGlobalShader
+{
+	DECLARE_GLOBAL_SHADER(FMarkValidityPS);
+	SHADER_USE_PARAMETER_STRUCT(FMarkValidityPS, FGlobalShader);
+
+	BEGIN_SHADER_PARAMETER_STRUCT(FParameters, )
+	END_SHADER_PARAMETER_STRUCT()
+
+	using FPermutationDomain = TShaderPermutationDomain<>;
+
+	static bool ShouldCompilePermutation(const FGlobalShaderPermutationParameters& Parameters)
+	{
+		return UE::Landscape::DoesPlatformSupportEditLayers(Parameters.Platform);
+	}
+
+	static void ModifyCompilationEnvironment(const FGlobalShaderPermutationParameters& Parameters, FShaderCompilerEnvironment& OutEnvironment)
+	{
+		OutEnvironment.SetDefine(TEXT("MARK_VALIDITY"), 1);
+	}
+};
+
+IMPLEMENT_GLOBAL_SHADER(FMarkValidityPS, "/Engine/Private/Landscape/LandscapeEditLayersUtils.usf", "MarkValidityPS", SF_Pixel);
+
+BEGIN_SHADER_PARAMETER_STRUCT(FMarkValidityPSParameters, )
+	SHADER_PARAMETER_STRUCT_INCLUDE(FPixelShaderUtils::FRasterizeToRectsVS::FParameters, VS)
+	SHADER_PARAMETER_STRUCT_INCLUDE(FMarkValidityPS::FParameters, PS)
+	RENDER_TARGET_BINDING_SLOTS()
+END_SHADER_PARAMETER_STRUCT()
+
+
+// ----------------------------------------------------------------------------------
+
+class FCopyQuadsPS : public FGlobalShader
+{
+	DECLARE_GLOBAL_SHADER(FCopyQuadsPS);
+	SHADER_USE_PARAMETER_STRUCT(FCopyQuadsPS, FGlobalShader);
+
+	BEGIN_SHADER_PARAMETER_STRUCT(FParameters, )
+		SHADER_PARAMETER_STRUCT_REF(FViewUniformShaderParameters, View)
+		SHADER_PARAMETER_RDG_TEXTURE(Texture2D, InSourceTexture)
+	END_SHADER_PARAMETER_STRUCT()
+
+	using FPermutationDomain = TShaderPermutationDomain<>;
+
+	static bool ShouldCompilePermutation(const FGlobalShaderPermutationParameters& Parameters)
+	{
+		return UE::Landscape::DoesPlatformSupportEditLayers(Parameters.Platform);
+	}
+
+	static void ModifyCompilationEnvironment(const FGlobalShaderPermutationParameters& Parameters, FShaderCompilerEnvironment& OutEnvironment)
+	{
+		OutEnvironment.SetDefine(TEXT("COPY_QUADS"), 1);
+	}
+};
+
+IMPLEMENT_GLOBAL_SHADER(FCopyQuadsPS, "/Engine/Private/Landscape/LandscapeEditLayersUtils.usf", "CopyQuadsPS", SF_Pixel);
+
+BEGIN_SHADER_PARAMETER_STRUCT(FCopyQuadsPSParameters, )
+	SHADER_PARAMETER_STRUCT_INCLUDE(FPixelShaderUtils::FRasterizeToRectsVS::FParameters, VS)
+	SHADER_PARAMETER_STRUCT_INCLUDE(FCopyQuadsPS::FParameters, PS)
+	RENDER_TARGET_BINDING_SLOTS()
+END_SHADER_PARAMETER_STRUCT()
+
+
+// ----------------------------------------------------------------------------------
+
+FString ConvertTargetLayerNamesToString(const TArrayView<const FName>& InTargetLayerNames)
+{
+	TArray<FString> TargetLayerStrings;
+	Algo::Transform(InTargetLayerNames, TargetLayerStrings, [](FName InTargetLayerName) { return InTargetLayerName.ToString(); });
+	return *FString::Join(TargetLayerStrings, TEXT(","));
+}
+
+
+// ----------------------------------------------------------------------------------
+
+#if WITH_EDITOR
+
+FEditLayerTargetTypeState::FEditLayerTargetTypeState(ELandscapeToolTargetTypeFlags InTargetTypeMask, const TArrayView<const FName>& InSupportedWeightmaps)
+	: Weightmaps(InSupportedWeightmaps)
+{
+	SetTargetTypeMask(InTargetTypeMask);
+}
+
+bool FEditLayerTargetTypeState::IsActive(ELandscapeToolTargetType InTargetType, FName InWeightmapLayerName /*= NAME_None*/) const
+{
+	if (EnumHasAnyFlags(TargetTypeMask, GetLandscapeToolTargetTypeAsFlags(InTargetType)))
+	{
+		if (InTargetType != ELandscapeToolTargetType::Heightmap)
+		{
+			return Weightmaps.Contains(InWeightmapLayerName);
+		}
+		return true;
+	}
+
+	return false;
+}
+
+const TArray<FName>& FEditLayerTargetTypeState::GetActiveWeightmaps() const
+{
+	if (EnumHasAnyFlags(TargetTypeMask, ELandscapeToolTargetTypeFlags::Weightmap | ELandscapeToolTargetTypeFlags::Visibility))
+	{
+		return Weightmaps;
+	}
+	static TArray<FName> EmptyWeightmaps;
+	return EmptyWeightmaps;
+}
+
+void FEditLayerTargetTypeState::SetTargetTypeMask(ELandscapeToolTargetTypeFlags InTargetTypeMask)
+{
+	if (InTargetTypeMask != TargetTypeMask)
+	{
+		TargetTypeMask = InTargetTypeMask;
+
+		// Special case for the visibility weightmap, where we want to make sure the weightmap layer name is specified if visibility is supported (and vice versa) : 
+		if (EnumHasAnyFlags(TargetTypeMask, ELandscapeToolTargetTypeFlags::Visibility))
+		{
+			AddWeightmap(UMaterialExpressionLandscapeVisibilityMask::ParameterName);
+		}
+		else if (Weightmaps.Contains(UMaterialExpressionLandscapeVisibilityMask::ParameterName))
+		{
+			RemoveWeightmap(UMaterialExpressionLandscapeVisibilityMask::ParameterName);
+		}
+	}
+}
+
+void FEditLayerTargetTypeState::AddTargetType(ELandscapeToolTargetType InTargetType)
+{
+	SetTargetTypeMask(TargetTypeMask | GetLandscapeToolTargetTypeAsFlags(InTargetType));
+}
+
+void FEditLayerTargetTypeState::AddTargetTypeMask(ELandscapeToolTargetTypeFlags InTargetTypeMask)
+{
+	SetTargetTypeMask(TargetTypeMask | InTargetTypeMask);
+}
+
+void FEditLayerTargetTypeState::RemoveTargetType(ELandscapeToolTargetType InTargetType)
+{
+	SetTargetTypeMask(TargetTypeMask & ~GetLandscapeToolTargetTypeAsFlags(InTargetType));
+}
+
+void FEditLayerTargetTypeState::RemoveTargetTypeMask(ELandscapeToolTargetTypeFlags InTargetTypeMask)
+{
+	SetTargetTypeMask(TargetTypeMask & ~InTargetTypeMask);
+}
+
+FEditLayerTargetTypeState FEditLayerTargetTypeState::Intersect(const FEditLayerTargetTypeState& InOther) const
+{
+	TSet<FName> OtherWeightmaps(InOther.GetActiveWeightmaps());
+	return FEditLayerTargetTypeState(InOther.GetTargetTypeMask() & TargetTypeMask, OtherWeightmaps.Intersect(TSet<FName>(Weightmaps)).Array());
+}
+
+void FEditLayerTargetTypeState::AddWeightmap(FName InWeightmapLayerName)
+{
+	checkf(EnumHasAnyFlags(TargetTypeMask, ELandscapeToolTargetTypeFlags::Weightmap | ELandscapeToolTargetTypeFlags::Visibility), 
+		TEXT("Cannot add weightmap %s to a target type state that doesn't support weightmaps"), *InWeightmapLayerName.ToString());
+
+	checkf(!EnumHasAnyFlags(TargetTypeMask, ELandscapeToolTargetTypeFlags::Visibility) || (InWeightmapLayerName == UMaterialExpressionLandscapeVisibilityMask::ParameterName) || Weightmaps.Contains(UMaterialExpressionLandscapeVisibilityMask::ParameterName),
+		TEXT("Visibility should always come with a weightmap named %s"), *UMaterialExpressionLandscapeVisibilityMask::ParameterName.ToString());
+
+	Weightmaps.AddUnique(InWeightmapLayerName);
+}
+
+void FEditLayerTargetTypeState::RemoveWeightmap(FName InWeightmapLayerName)
+{
+	checkf(EnumHasAnyFlags(TargetTypeMask, ELandscapeToolTargetTypeFlags::Weightmap | ELandscapeToolTargetTypeFlags::Visibility), 
+		TEXT("Cannot remove weightmap %s from a target type state that doesn't support weightmaps"), *InWeightmapLayerName.ToString());
+
+	checkf(!EnumHasAnyFlags(TargetTypeMask, ELandscapeToolTargetTypeFlags::Visibility) || (InWeightmapLayerName != UMaterialExpressionLandscapeVisibilityMask::ParameterName), 
+		TEXT("Cannot remove weightmap %s from a target type state that supports visibility"), *UMaterialExpressionLandscapeVisibilityMask::ParameterName.ToString());
+
+	Weightmaps.Remove(InWeightmapLayerName);
+}
+
+bool FEditLayerTargetTypeState::operator==(const FEditLayerTargetTypeState& InOther) const
+{
+	return (TargetTypeMask == InOther.TargetTypeMask) 
+		// TODO [jonathan.bard] : This is really bad for performance if called several times. This test can be replaced by a hash test of all ordered weightmaps: 
+		&& (TSet<FName>(Weightmaps).Intersect(TSet<FName>(InOther.Weightmaps)).Num() == Weightmaps.Num());
+}
+
+FString FEditLayerTargetTypeState::ToString() const
+{
+	FString Result = FString::Printf(TEXT("Target types: %s"), *UE::Landscape::GetLandscapeToolTargetTypeFlagsAsString(TargetTypeMask));
+	if (EnumHasAnyFlags(TargetTypeMask, ELandscapeToolTargetTypeFlags::Weightmap | ELandscapeToolTargetTypeFlags::Visibility))
+	{
+		Result += FString::Printf(TEXT("\nWeightmaps: %s"), *ConvertTargetLayerNamesToString(Weightmaps));
+	}
+	return Result;
+}
+
+// ----------------------------------------------------------------------------------
+
+FEditLayerRendererState::FEditLayerRendererState(ILandscapeEditLayerRenderer* InRenderer, const ULandscapeInfo* InLandscapeInfo)
+	: Renderer(InRenderer)
+{
+	Renderer->GetRendererStateInfo(InLandscapeInfo, SupportedTargetTypeState, EnabledTargetTypeState, RenderGroups);
+
+	// Make sure that each supported weightmap belongs to one render group and one only. For those that are in no render group, put them in their own group, that simply means this renderer 
+	//  can render them without requesting the presence of other target layers (e.g. no weight-blending)
+	for (FName TargetLayerName : SupportedTargetTypeState.GetActiveWeightmaps())
+	{
+		const int32 RenderGroupCount = Algo::CountIf(RenderGroups, [TargetLayerName](const TSet<FName>& InRenderGroup) { return InRenderGroup.Contains(TargetLayerName); });
+		checkf(RenderGroupCount < 2, TEXT("Target layer %s belongs to more than 1 render group in edit layer renderer %s. This is forbidden: in the end, it must belong to 1 and 1 only."),
+			*TargetLayerName.ToString(), *InRenderer->GetEditLayerRendererDebugName());
+		if (RenderGroupCount == 0)
+		{
+			RenderGroups.Add({ TargetLayerName });
+		}
+	}
+}
+
+bool FEditLayerRendererState::IsTargetSupported(ELandscapeToolTargetType InTargetType, FName InWeightmapLayerName) const
+{
+	return SupportedTargetTypeState.IsActive(InTargetType, InWeightmapLayerName);
+}
+
+const TArray<FName>& FEditLayerRendererState::GetSupportedTargetWeightmaps() const
+{
+	return SupportedTargetTypeState.GetActiveWeightmaps();
+}
+
+bool FEditLayerRendererState::IsTargetEnabled(ELandscapeToolTargetType InTargetType, FName InWeightmapLayerName) const
+{
+	return SupportedTargetTypeState.IsActive(InTargetType, InWeightmapLayerName) && EnabledTargetTypeState.IsActive(InTargetType, InWeightmapLayerName);
+}
+
+void FEditLayerRendererState::EnableTargetType(ELandscapeToolTargetType InTargetType)
+{
+	checkf(!EnumHasAnyFlags(SupportedTargetTypeState.GetTargetTypeMask(), GetLandscapeToolTargetTypeAsFlags(InTargetType)), 
+		TEXT("Target type %s cannot be enabled on this renderer state because it is not supported. Make sure that target types are supported before enabling them"), *UEnum::GetValueAsString(InTargetType));
+	EnabledTargetTypeState.AddTargetType(InTargetType);
+}
+
+void FEditLayerRendererState::EnableTargetTypeMask(ELandscapeToolTargetTypeFlags InTargetTypeMask)
+{
+	for (ELandscapeToolTargetTypeFlags TargetTypeFlag : MakeFlagsRange(InTargetTypeMask))
+	{
+		EnableTargetType(UE::Landscape::GetLandscapeToolTargetTypeSingleFlagAsType(TargetTypeFlag));
+	}
+}
+
+void FEditLayerRendererState::DisableTargetType(ELandscapeToolTargetType InTargetType)
+{
+	EnabledTargetTypeState.RemoveTargetType(InTargetType);
+}
+
+void FEditLayerRendererState::DisableTargetTypeMask(ELandscapeToolTargetTypeFlags InTargetTypeMask)
+{
+	EnabledTargetTypeState.RemoveTargetTypeMask(InTargetTypeMask);
+}
+
+bool FEditLayerRendererState::EnableTarget(ELandscapeToolTargetType InTargetType, FName InWeightmapLayerName)
+{
+	EnableTargetType(InTargetType);
+	EnabledTargetTypeState.AddWeightmap(InWeightmapLayerName);
+	
+	// The target has to be both supported and enabled to be considered fully enabled : 
+	return IsTargetEnabled(InTargetType, InWeightmapLayerName);
+}
+
+void FEditLayerRendererState::DisableTarget(FName InWeightmapLayerName)
+{
+	EnabledTargetTypeState.RemoveWeightmap(InWeightmapLayerName);
+}
+
+TArray<FName> FEditLayerRendererState::GetEnabledTargetWeightmaps() const
+{
+	// Find the weightmaps that are both supported and enabled : 
+	FEditLayerTargetTypeState SupportedAndEnabledState = SupportedTargetTypeState.Intersect(EnabledTargetTypeState);
+	return SupportedAndEnabledState.GetActiveWeightmaps();
+}
+
+
+// ----------------------------------------------------------------------------------
+
+bool FMergeRenderBatch::operator<(const FMergeRenderBatch& InOther) const
+{
+	// Sort by coordinates for making debugging more "logical" : 
+	if (MinComponentKey.Y < InOther.MinComponentKey.Y)
+	{
+		return true;
+	}
+	else if (MinComponentKey.Y == InOther.MinComponentKey.Y)
+	{
+		return (MinComponentKey.X < InOther.MinComponentKey.X);
+	}
+	return false;
+}
+
+int32 FMergeRenderBatch::ComputeSubsectionRects(ULandscapeComponent* InComponent, TArray<FIntRect, TInlineAllocator<4>>& OutSubsectionRects, TArray<FIntRect, TInlineAllocator<4>>& OutSubsectionRectsWithDuplicateBorders) const
+{
+	check(ComponentsToRender.Contains(InComponent));
+	const int32 NumSubsections = Landscape->NumSubsections;
+	const int32 ComponentSizeQuads = Landscape->ComponentSizeQuads;
+	const int32 SubsectionSizeQuads = Landscape->SubsectionSizeQuads;
+	const int32 SubsectionVerts = SubsectionSizeQuads + 1;
+	const int32 TotalNumSubsections = NumSubsections * NumSubsections;
+	OutSubsectionRects.Reserve(TotalNumSubsections);
+	OutSubsectionRectsWithDuplicateBorders.Reserve(TotalNumSubsections);
+
+	const FIntPoint ComponentSectionBase = InComponent->GetSectionBase();
+	checkf((ComponentSectionBase.X >= SectionRect.Min.X) && (ComponentSectionBase.Y >= SectionRect.Min.Y)
+		&& ((ComponentSectionBase.X + ComponentSizeQuads + 1) <= SectionRect.Max.X) && ((ComponentSectionBase.Y + ComponentSizeQuads + 1) <= SectionRect.Max.Y), 
+		TEXT("The requested component is not included in the render batch"));
+
+	const FIntPoint ComponentLocalKey = (ComponentSectionBase - SectionRect.Min) / ComponentSizeQuads;
+	for (int32 SubY = 0; SubY < NumSubsections; ++SubY)
+	{
+		for (int32 SubX = 0; SubX < NumSubsections; ++SubX)
+		{
+			{
+				FIntPoint SubSectionMin = ComponentSectionBase - SectionRect.Min + FIntPoint(SubX * SubsectionSizeQuads, SubY * SubsectionSizeQuads);
+				FIntPoint SubSectionMax = SubSectionMin + FIntPoint(SubsectionVerts, SubsectionVerts);
+				OutSubsectionRects.Add(FIntRect(SubSectionMin, SubSectionMax));
+			}
+			{
+				FIntPoint SubSectionMin = (ComponentLocalKey * NumSubsections + FIntPoint(SubX, SubY)) * SubsectionVerts;
+				FIntPoint SubSectionMax = SubSectionMin + SubsectionVerts;
+				OutSubsectionRectsWithDuplicateBorders.Add(FIntRect(SubSectionMin, SubSectionMax));
+			}
+		}
+	}
+
+	return TotalNumSubsections;
+}
+
+FIntRect FMergeRenderBatch::ComputeSectionRect(ULandscapeComponent* InComponent, bool bInWithDuplicateBorders) const
+{
+	check(ComponentsToRender.Contains(InComponent));
+
+	const FIntPoint ComponentSectionBase = InComponent->GetSectionBase();
+	checkf((ComponentSectionBase.X >= SectionRect.Min.X) && (ComponentSectionBase.Y >= SectionRect.Min.Y)
+		&& ((ComponentSectionBase.X + InComponent->ComponentSizeQuads + 1) <= SectionRect.Max.X) && ((ComponentSectionBase.Y + InComponent->ComponentSizeQuads + 1) <= SectionRect.Max.Y),
+		TEXT("The requested component is not included in the render batch"));
+
+	const FIntPoint ComponentLocalKey = (ComponentSectionBase - SectionRect.Min) / InComponent->ComponentSizeQuads;
+	const int32 ComponentSubsectionVerts = InComponent->SubsectionSizeQuads + 1;
+	
+	const int32 ComponentSize = InComponent->NumSubsections * (bInWithDuplicateBorders ? ComponentSubsectionVerts : InComponent->SubsectionSizeQuads);
+	FIntPoint SectionMin = ComponentLocalKey * ComponentSize;
+	FIntPoint SectionMax = SectionMin + ComponentSize;
+
+	return FIntRect(SectionMin, SectionMax);
+}
+
+void FMergeRenderBatch::ComputeAllSubsectionRects(TArray<FIntRect>& OutSubsectionRects, TArray<FIntRect>& OutSubsectionRectsWithDuplicateBorders) const
+{
+	const int32 NumSubsections = Landscape->NumSubsections;
+	const int32 ComponentSizeQuads = Landscape->ComponentSizeQuads;
+	const int32 SubsectionSizeQuads = Landscape->SubsectionSizeQuads;
+	const int32 SubsectionVerts = SubsectionSizeQuads + 1;
+	const int32 TotalNumSubsectionRects = ComponentsToRender.Num() * NumSubsections * NumSubsections;
+	OutSubsectionRects.Reserve(TotalNumSubsectionRects);
+	OutSubsectionRectsWithDuplicateBorders.Reserve(TotalNumSubsectionRects);
+
+	for (ULandscapeComponent* Component : ComponentsToRender)
+	{
+		const FIntPoint ComponentSectionBase = Component->GetSectionBase();
+		checkf((ComponentSectionBase.X >= SectionRect.Min.X) && (ComponentSectionBase.Y >= SectionRect.Min.Y)
+			&& ((ComponentSectionBase.X + ComponentSizeQuads + 1) <= SectionRect.Max.X) && ((ComponentSectionBase.Y + ComponentSizeQuads + 1) <= SectionRect.Max.Y),
+			TEXT("The requested component is not included in the render batch"));
+
+		const FIntPoint ComponentLocalKey = (ComponentSectionBase - SectionRect.Min) / ComponentSizeQuads;
+		TArray<FIntRect, TInlineAllocator<4>> SubSectionRects;
+		for (int32 SubY = 0; SubY < NumSubsections; ++SubY)
+		{
+			for (int32 SubX = 0; SubX < NumSubsections; ++SubX)
+			{
+				{
+					FIntPoint SubSectionMin = ComponentSectionBase - SectionRect.Min + FIntPoint(SubX * SubsectionSizeQuads, SubY * SubsectionSizeQuads);
+					FIntPoint SubSectionMax = SubSectionMin + FIntPoint(SubsectionVerts, SubsectionVerts);
+					OutSubsectionRects.Add(FIntRect(SubSectionMin, SubSectionMax));
+				}
+				{
+					FIntPoint SubSectionMin = (ComponentLocalKey * NumSubsections + FIntPoint(SubX, SubY)) * SubsectionVerts;
+					FIntPoint SubSectionMax = SubSectionMin + SubsectionVerts;
+					OutSubsectionRectsWithDuplicateBorders.Add(FIntRect(SubSectionMin, SubSectionMax));
+				}
+			}
+		}
+	}
+}
+
+FIntPoint FMergeRenderBatch::GetRenderTargetResolution(bool bInWithDuplicateBorders) const
+{
+	return bInWithDuplicateBorders ? Resolution : SectionRect.Size();
+}
+
+
+// ----------------------------------------------------------------------------------
+
+FMergeRenderContext::FMergeRenderContext(ALandscape* InLandscape, bool bInIsHeightmapMerge)
+	: Landscape(InLandscape)
+	, bIsHeightmapMerge(bInIsHeightmapMerge)
+{
+	for (ULandscapeScratchRenderTarget*& BlendRenderTarget : BlendRenderTargets)
+	{
+		BlendRenderTarget = nullptr;
+	}
+}
+
+FMergeRenderContext::~FMergeRenderContext()
+{
+	FreeResources();
+
+	checkf(Algo::AllOf(BlendRenderTargets, [](ULandscapeScratchRenderTarget* InRT) { return (InRT == nullptr); }), TEXT("Every scratch render target should have been freed at this point."));
+}
+
+void FMergeRenderContext::AllocateResources()
+{
+	TRACE_CPUPROFILER_EVENT_SCOPE(FMergeRenderContext::AllocateResources);
+
+	using namespace UE::Landscape;
+
+	// Prepare the transient render resources we'll need throughout the merge: 
+	const int32 NumSlices = IsHeightmapMerge() ? 0 : MaxNeededNumSlices;
+	FLinearColor RenderTargetClearColor(ForceInitToZero);
+	ETextureRenderTargetFormat RenderTargetFormat = ETextureRenderTargetFormat::RTF_R8;
+	if (IsHeightmapMerge())
+	{
+		// Convert the height value 0.0f to how it's stored in the texture : 
+		const uint16 HeightValue = LandscapeDataAccess::GetTexHeight(0.0f);
+		RenderTargetClearColor = FLinearColor((float)((HeightValue - (HeightValue & 255)) >> 8) / 255.0f, (float)(HeightValue & 255) / 255.0f, 0.0f, 0.0f);
+
+		RenderTargetFormat = ETextureRenderTargetFormat::RTF_RGBA8;
+	}
+	// When rendering weightmaps, we should have at least 1 slice (if == 1, we can use a UTextureRenderTarget2D, otherwise, we'll need to use a UTextureRenderTarget2DArray) : 
+	else
+	{
+		checkf(MaxNeededNumSlices > 0, TEXT("Weightmaps should have at least 1 slice"));
+	}
+
+	ULandscapeEditResourcesSubsystem* LandscapeEditResourcesSubsystem = GEngine->GetEngineSubsystem<ULandscapeEditResourcesSubsystem>();
+	check(LandscapeEditResourcesSubsystem != nullptr);
+	checkf(Algo::AllOf(BlendRenderTargets, [](ULandscapeScratchRenderTarget* InRT) { return (InRT == nullptr); }), TEXT("We shouldn't allocate without having freed first."));
+	check(CurrentBlendRenderTargetWriteIndex == -1);
+
+	// We need N render targets large enough to fit all batches : 
+	{
+		// Write : 
+		FScratchRenderTargetParams ScratchRenderTargetParams(TEXT("ScratchRT0"), /*bInExactDimensions = */false, /*bInUseUAV = */false, /*bInTargetArraySlicesIndependently = */(NumSlices > 0),
+			MaxNeededResolution, NumSlices, RenderTargetFormat, RenderTargetClearColor, ERHIAccess::RTV);
+		BlendRenderTargets[0] = LandscapeEditResourcesSubsystem->RequestScratchRenderTarget(ScratchRenderTargetParams);
+		// Read and ReadPrevious : 
+		ScratchRenderTargetParams.DebugName = TEXT("ScratchRT1");
+		ScratchRenderTargetParams.InitialState = ERHIAccess::SRVMask;
+		BlendRenderTargets[1] = LandscapeEditResourcesSubsystem->RequestScratchRenderTarget(ScratchRenderTargetParams);
+		ScratchRenderTargetParams.DebugName = TEXT("ScratchRT2");
+		BlendRenderTargets[2] = LandscapeEditResourcesSubsystem->RequestScratchRenderTarget(ScratchRenderTargetParams);
+	}
+}
+
+void FMergeRenderContext::FreeResources()
+{
+	TRACE_CPUPROFILER_EVENT_SCOPE(FMergeRenderContext::FreeResources);
+
+	ULandscapeEditResourcesSubsystem* LandscapeEditResourcesSubsystem = GEngine->GetEngineSubsystem<ULandscapeEditResourcesSubsystem>();
+	check(LandscapeEditResourcesSubsystem != nullptr);
+
+	// We can now return those scratch render targets to the pool : 
+	for (int32 Index = 0; Index < BlendRenderTargets.Num(); ++Index)
+	{
+		if (BlendRenderTargets[Index] != nullptr)
+		{
+			LandscapeEditResourcesSubsystem->ReleaseScratchRenderTarget(BlendRenderTargets[Index]);
+			BlendRenderTargets[Index] = nullptr;
+		}
+	}
+	
+	CurrentBlendRenderTargetWriteIndex = -1;
+}
+
+void FMergeRenderContext::AllocateBatchResources(const FMergeRenderBatch& InRenderBatch)
+{
+	TRACE_CPUPROFILER_EVENT_SCOPE(FMergeRenderContext::AllocateBatchResources);
+
+	using namespace UE::Landscape;
+
+	// Prepare the transient render resources we'll need for this batch:
+	ULandscapeEditResourcesSubsystem* LandscapeEditResourcesSubsystem = GEngine->GetEngineSubsystem<ULandscapeEditResourcesSubsystem>();
+	check(LandscapeEditResourcesSubsystem != nullptr);
+	check(PerTargetLayerValidityRenderTargets.IsEmpty());
+
+	// We need a RT version of the stencil buffer, one per target layer, to let users sample it as a standard texture :
+	int32 VisibilityScratchRTIndex = 0;
+	ForEachTargetLayer(InRenderBatch.TargetLayerNameBitIndices, 
+		[this, LandscapeEditResourcesSubsystem, &VisibilityScratchRTIndex](int32 InTargetLayerIndex, FName InTargetLayerName)
+		{
+			FScratchRenderTargetParams ScratchRenderTargetParams(FString::Printf(TEXT("VisibilityScratchRT(%i)"), VisibilityScratchRTIndex), /*bInExactDimensions = */false, /*bInUseUAV = */false, /*bInTargetArraySlicesIndependently = */false,
+				MaxNeededResolution, 0, ETextureRenderTargetFormat::RTF_R8, FLinearColor::Black, ERHIAccess::RTV);
+			ULandscapeScratchRenderTarget* RenderTarget = LandscapeEditResourcesSubsystem->RequestScratchRenderTarget(ScratchRenderTargetParams);
+			PerTargetLayerValidityRenderTargets.FindOrAdd(InTargetLayerName) = RenderTarget;
+			++VisibilityScratchRTIndex;
+			return true;
+		});
+}
+
+void FMergeRenderContext::FreeBatchResources(const FMergeRenderBatch& InRenderBatch)
+{
+	TRACE_CPUPROFILER_EVENT_SCOPE(FMergeRenderContext::FreeBatchResources);
+
+	ULandscapeEditResourcesSubsystem* LandscapeEditResourcesSubsystem = GEngine->GetEngineSubsystem<ULandscapeEditResourcesSubsystem>();
+	check(LandscapeEditResourcesSubsystem != nullptr);
+
+	// We can now return those scratch render targets to the pool : 
+	for (auto ItLayerNameRenderTargetPair : PerTargetLayerValidityRenderTargets)
+	{
+		LandscapeEditResourcesSubsystem->ReleaseScratchRenderTarget(ItLayerNameRenderTargetPair.Value);
+	}
+	PerTargetLayerValidityRenderTargets.Empty();
+}
+
+void FMergeRenderContext::CycleBlendRenderTargets(ERHIAccess InDesiredWriteAccess)
+{
+	const bool bFirstWrite = (CurrentBlendRenderTargetWriteIndex < 0);
+	CurrentBlendRenderTargetWriteIndex = (CurrentBlendRenderTargetWriteIndex + 1) % BlendRenderTargets.Num();
+
+	if (!bFirstWrite)
+	{
+		// Optionally clear the write render target for debug purposes : 
+		if (CVarLandscapeEditLayersClearBeforeEachWriteToScratch.GetValueOnGameThread())
+		{
+			GetBlendRenderTargetWrite()->Clear();
+		}
+
+		// Change the state of the new Read and Write (ReadPrevious is already SRV since it was Read before)
+		ULandscapeScratchRenderTarget::FTransitionBatcherScope TransitionsScope;
+		TransitionsScope.TransitionTo(GetBlendRenderTargetRead(), ERHIAccess::SRVMask);
+		if (InDesiredWriteAccess != ERHIAccess::None)
+		{
+			TransitionsScope.TransitionTo(GetBlendRenderTargetWrite(), InDesiredWriteAccess);
+		}
+	}
+}
+
+ULandscapeScratchRenderTarget* EditLayers::FMergeRenderContext::GetBlendRenderTargetWrite() const
+{
+	checkf(CurrentBlendRenderTargetWriteIndex >= 0, TEXT("CycleBlendRenderTargets must be called at least once prior to accessing the blend render targets"));
+	return BlendRenderTargets[CurrentBlendRenderTargetWriteIndex % BlendRenderTargets.Num()];
+}
+
+ULandscapeScratchRenderTarget* FMergeRenderContext::GetBlendRenderTargetRead() const
+{
+	checkf(CurrentBlendRenderTargetWriteIndex >= 0, TEXT("CycleBlendRenderTargets must be called at least once prior to accessing the blend render targets")); 
+	return BlendRenderTargets[(CurrentBlendRenderTargetWriteIndex + BlendRenderTargets.Num() - 1) % BlendRenderTargets.Num()];
+}
+
+ULandscapeScratchRenderTarget* FMergeRenderContext::GetBlendRenderTargetReadPrevious() const
+{
+	checkf(CurrentBlendRenderTargetWriteIndex >= 0, TEXT("CycleBlendRenderTargets must be called at least once prior to accessing the blend render targets"));
+	return BlendRenderTargets[(CurrentBlendRenderTargetWriteIndex + BlendRenderTargets.Num() - 2) % BlendRenderTargets.Num()];
+}
+
+ULandscapeScratchRenderTarget* FMergeRenderContext::GetValidityRenderTarget(const FName& InTargetLayerName) const
+{
+	check(PerTargetLayerValidityRenderTargets.Contains(InTargetLayerName));
+	return PerTargetLayerValidityRenderTargets[InTargetLayerName];
+}
+
+FTransform FMergeRenderContext::ComputeVisualLogTransform(const FTransform& InTransform) const
+{
+	FTransform ZTransform(CurrentVisualLogOffset / InTransform.GetScale3D()); // the Offset is given in world space so unapply the scale before applying the transform
+	return ZTransform * InTransform;
+}
+
+void FMergeRenderContext::IncrementVisualLogOffset()
+{
+	double VisualLogOffsetIncrement = CVarLandscapeBatchedMergeVisualLogOffsetIncrement.GetValueOnGameThread();
+	CurrentVisualLogOffset.Z += VisualLogOffsetIncrement;
+}
+
+void FMergeRenderContext::ResetVisualLogOffset()
+{
+	CurrentVisualLogOffset = FVector(ForceInitToZero);
+}
+
+void FMergeRenderContext::RenderValidityRenderTargets(const FMergeRenderBatch& InRenderBatch)
+{
+	struct FTextureAndRects
+	{
+		FTextureAndRects(const FName& InTargetLayerName, const FString& InTextureDebugName, FTextureResource* InTextureResource)
+			: TargetLayerName(InTargetLayerName)
+			, TextureDebugName(InTextureDebugName)
+			, TextureResource(InTextureResource)
+		{}
+
+		FName TargetLayerName;
+		FString TextureDebugName;
+		FTextureResource* TextureResource = nullptr;
+		TArray<FUintVector4> Rects;
+	};
+
+	TArray<FTextureAndRects> TexturesAndRects;
+	TexturesAndRects.Reserve(InRenderBatch.TargetLayerNameBitIndices.CountSetBits());
+	ForEachTargetLayer(InRenderBatch.TargetLayerNameBitIndices, [this, &TexturesAndRects, &InRenderBatch](int32 InTargetLayerIndex, FName InTargetLayerName)
+		{
+			ULandscapeScratchRenderTarget* ScratchRenderTarget = PerTargetLayerValidityRenderTargets.FindChecked(InTargetLayerName);
+			check(ScratchRenderTarget != nullptr);
+
+			// Make sure the validity mask is entirely cleared first:
+			ScratchRenderTarget->Clear();
+
+			FTextureAndRects& TextureAndRects = TexturesAndRects.Emplace_GetRef(InTargetLayerName, ScratchRenderTarget->GetDebugName(), ScratchRenderTarget->GetRenderTarget2D()->GetResource());
+
+			// Then build a list of quads for marking where the components are valid for this target layer on this batch:
+			const TSet<ULandscapeComponent*>& Components = InRenderBatch.TargetLayersToComponents[InTargetLayerIndex];
+			TextureAndRects.Rects.Reserve(Components.Num());
+			for (ULandscapeComponent* Component : Components)
+			{
+				FIntRect ComponentRect = InRenderBatch.ComputeSectionRect(Component, /*bInWithDuplicateBorders = */false);
+				TextureAndRects.Rects.Add(FUintVector4(ComponentRect.Min.X, ComponentRect.Min.Y, ComponentRect.Max.X + 1, ComponentRect.Max.Y + 1));
+			}
+
+			ScratchRenderTarget->TransitionTo(ERHIAccess::RTV);
+			return true;
+		});
+
+	ENQUEUE_RENDER_COMMAND(MarkTargetLayersValidity)([TexturesAndRects = MoveTemp(TexturesAndRects)](FRHICommandListImmediate& InRHICmdList)
+	{
+		FRDGBuilder GraphBuilder(InRHICmdList, RDG_EVENT_NAME("MarkTargetLayersValidity"));
+
+		for (const FTextureAndRects& TextureAndRects : TexturesAndRects)
+		{
+			FRDGBufferRef RectBuffer = CreateUploadBuffer(GraphBuilder, TEXT("MarkValidityRects"), MakeArrayView(TextureAndRects.Rects));
+			FRDGBufferSRVRef RectBufferSRV = GraphBuilder.CreateSRV(FRDGBufferSRVDesc(RectBuffer, PF_R32G32B32A32_UINT));
+			FRDGTextureRef OutputTexture = GraphBuilder.RegisterExternalTexture(CreateRenderTarget(TextureAndRects.TextureResource->GetTexture2DRHI(), TEXT("ValidityMask")), ERDGTextureFlags::SkipTracking);
+
+			FMarkValidityPSParameters* PassParameters = GraphBuilder.AllocParameters<FMarkValidityPSParameters>();
+			PassParameters->RenderTargets[0] = FRenderTargetBinding(OutputTexture, ERenderTargetLoadAction::ELoad);
+			FGlobalShaderMap* ShaderMap = GetGlobalShaderMap(GMaxRHIFeatureLevel);
+			TShaderRef<FMarkValidityPS> PixelShader = ShaderMap->GetShader<FMarkValidityPS>();
+
+			FPixelShaderUtils::AddRasterizeToRectsPass<FMarkValidityPS>(GraphBuilder,
+				ShaderMap,
+				RDG_EVENT_NAME("MarkValidity(%s) -> %s", *TextureAndRects.TargetLayerName.ToString(), *TextureAndRects.TextureDebugName),
+				PixelShader,
+				PassParameters,
+				/*ViewportSize = */OutputTexture->Desc.Extent,
+				RectBufferSRV,
+				TextureAndRects.Rects.Num(),
+				/*BlendState = */ nullptr,
+				/*RasterizerState = */ nullptr,
+				/*DepthStencilState = */ nullptr,
+				/*StencilRef = */ 0,
+				/*TextureSize = */ OutputTexture->Desc.Extent,
+				/*RectUVBufferSRV = */ nullptr,
+				/*DownsampleFactor = */ 1,
+				ERDGPassFlags::NeverCull); // Use NeverCull because it renders a texture for which tracking is disabled
+		}
+
+		GraphBuilder.Execute();
+	});
+}
+
+void FMergeRenderContext::RenderExpandedRenderTarget(const FMergeRenderBatch& InRenderBatch)
+{
+	TArray<FUintVector4> SourceRects, DestinationRects;
+	{
+		TArray<FIntRect> SourceInclusiveRects, DestinationInclusiveRects;
+		InRenderBatch.ComputeAllSubsectionRects(SourceInclusiveRects, DestinationInclusiveRects);
+		// ComputeAllSubsectionRects returns inclusive bounds while AddRasterizeToRectsPass requires exclusive bounds : 
+		Algo::Transform(SourceInclusiveRects, SourceRects, [](const FIntRect& InRect) { return FUintVector4(InRect.Min.X, InRect.Min.Y, InRect.Max.X + 1, InRect.Max.Y + 1); });
+		Algo::Transform(DestinationInclusiveRects, DestinationRects, [](const FIntRect& InRect) { return FUintVector4(InRect.Min.X, InRect.Min.Y, InRect.Max.X + 1, InRect.Max.Y + 1); });
+	}
+
+	ULandscapeScratchRenderTarget* WriteRT = GetBlendRenderTargetWrite();
+	ULandscapeScratchRenderTarget* ReadRT = GetBlendRenderTargetRead();
+	WriteRT->TransitionTo(ERHIAccess::RTV);
+	ReadRT->TransitionTo(ERHIAccess::SRVMask);
+
+	FSceneInterface* SceneInterface = GetLandscape()->GetWorld()->Scene;
+
+	ENQUEUE_RENDER_COMMAND(Expand)(
+		[ SourceRects = MoveTemp(SourceRects)
+		, DestinationRects = MoveTemp(DestinationRects)
+		, OutputResource = WriteRT->GetRenderTarget2D()->GetResource()
+		, SourceResource = ReadRT->GetRenderTarget2D()->GetResource()
+		, SceneInterface] (FRHICommandListImmediate& InRHICmdList)
+	{
+		FRDGBuilder GraphBuilder(InRHICmdList, RDG_EVENT_NAME("Expand"));
+
+		FRDGBufferRef RectBuffer = CreateUploadBuffer(GraphBuilder, TEXT("ExpandRects"), MakeArrayView(DestinationRects));
+		FRDGBufferSRVRef RectBufferSRV = GraphBuilder.CreateSRV(FRDGBufferSRVDesc(RectBuffer, PF_R32G32B32A32_UINT));
+
+		FRDGBufferRef RectUVBuffer = CreateUploadBuffer(GraphBuilder, TEXT("ExpandRectsUVs"), MakeArrayView(SourceRects));
+		FRDGBufferSRVRef RectUVBufferSRV = GraphBuilder.CreateSRV(FRDGBufferSRVDesc(RectUVBuffer, PF_R32G32B32A32_UINT));
+
+		FRDGTextureRef OutputTexture = GraphBuilder.RegisterExternalTexture(CreateRenderTarget(OutputResource->GetTexture2DRHI(), TEXT("OutputTexture")), ERDGTextureFlags::SkipTracking);
+		FRDGTextureRef SourceTexture = GraphBuilder.RegisterExternalTexture(CreateRenderTarget(SourceResource->GetTexture2DRHI(), TEXT("SourceTexture")), ERDGTextureFlags::SkipTracking);
+
+		// TODO [jonathan.bard] this is just a rhi validation error for unoptimized shaders... once validation is made to not issue those errors, we can remove this
+		// Create a SceneView to please the shader bindings, but it's unused in practice 
+		FSceneViewFamilyContext ViewFamily(FSceneViewFamily::ConstructionValues(nullptr, SceneInterface, FEngineShowFlags(ESFIM_Game)).SetTime(FGameTime::GetTimeSinceAppStart()));
+		FSceneViewInitOptions ViewInitOptions;
+		ViewInitOptions.ViewFamily = &ViewFamily;
+		GetRendererModule().CreateAndInitSingleView(InRHICmdList, &ViewFamily, &ViewInitOptions);
+		const FSceneView* View = ViewFamily.Views[0];
+
+		FCopyQuadsPSParameters* PassParameters = GraphBuilder.AllocParameters<FCopyQuadsPSParameters>();
+		PassParameters->RenderTargets[0] = FRenderTargetBinding(OutputTexture, ERenderTargetLoadAction::ELoad);
+		PassParameters->PS.View = View->ViewUniformBuffer;
+		PassParameters->PS.InSourceTexture = SourceTexture;
+
+		FGlobalShaderMap* ShaderMap = GetGlobalShaderMap(GMaxRHIFeatureLevel);
+		TShaderRef<FCopyQuadsPS> PixelShader = ShaderMap->GetShader<FCopyQuadsPS>();
+
+		FPixelShaderUtils::AddRasterizeToRectsPass<FCopyQuadsPS>(GraphBuilder,
+			ShaderMap,
+			RDG_EVENT_NAME("CopyQuadsPS"),
+			PixelShader,
+			PassParameters,
+			/*ViewportSize = */OutputTexture->Desc.Extent,
+			RectBufferSRV,
+			DestinationRects.Num(),
+			/*BlendState = */nullptr,
+			/*RasterizerState = */nullptr,
+			/*DepthStencilState = */nullptr,
+			/*StencilRef = */0,
+			/*TextureSize = */SourceTexture->Desc.Extent,
+			RectUVBufferSRV,
+			/*DownsampleFactor = */1,
+			ERDGPassFlags::NeverCull); // Use NeverCull because it renders a texture for which tracking is disabled
+
+		GraphBuilder.Execute();
+	});
+}
+
+void FMergeRenderContext::Render(TFunction<void(const FOnRenderBatchTargetGroupDoneParams&)> OnBatchTargetGroupDone)
+{
+	TRACE_CPUPROFILER_EVENT_SCOPE(FMergeRenderContext::Render);
+
+	check(CurrentRenderBatchIndex == INDEX_NONE);
+
+	const bool bVisualLogMergeProcess = CVarLandscapeBatchedMergeVisualLogShowMergeProcess.GetValueOnGameThread();
+
+	const FTransform& LandscapeTransform = Landscape->GetTransform();
+	// For visual logging, start at the top of the landscape's bounding box :
+	FTransform LandscapeWorldTransformForVisLog = LandscapeTransform;
+	FVector LandscapeTopPosition(0.0, 0.0, MaxLocalHeight);
+	LandscapeWorldTransformForVisLog.SetTranslation(LandscapeWorldTransformForVisLog.GetLocation() + LandscapeTransform.TransformVector(LandscapeTopPosition));
+	// TODO [jonathan.bard] : this should simply be: FTransform LandscapeWorldTransformForVisLog = FTransform(LandscapeTopPosition).ToMatrixNoScale() * LandscapeTransform.ToMatrixWithScale()
+
+	AllocateResources();
+
+	// Kick start the blend render targets : 
+	CycleBlendRenderTargets();
+
+	const int32 NumBatches = RenderBatches.Num();
+	for (CurrentRenderBatchIndex = 0; CurrentRenderBatchIndex < NumBatches; ++CurrentRenderBatchIndex)
+	{
+		const FMergeRenderBatch& RenderBatch = RenderBatches[CurrentRenderBatchIndex];
+		FString RenderBatchDebugName = FString::Format(TEXT("Render Batch [{0}] : ({1},{2})->({3},{4})"), { CurrentRenderBatchIndex, RenderBatch.SectionRect.Min.X, RenderBatch.SectionRect.Min.Y, RenderBatch.SectionRect.Max.X, RenderBatch.SectionRect.Max.Y });
+		SCOPED_DRAW_EVENTF_GAMETHREAD(LandscapeLayers, TEXT("%s"), RenderBatchDebugName);
+
+		checkf((RenderBatch.RenderSteps.Num() >= 1) && (RenderBatch.RenderSteps.Last().Type == FMergeRenderStep::EType::SignalBatchMergeGroupDone), TEXT("Any batch should end with a SignalBatchMergeGroupDone step and there \
+			should be at least another step prior to that, otherwise, the batch is just useless."));
+
+		AllocateBatchResources(RenderBatch);
+
+		IncrementVisualLogOffset();
+
+		// Drop a visual log showing the area covered by this batch : 
+		UE_IFVLOG(
+			if (IsVisualLogEnabled() && bVisualLogMergeProcess)
+			{
+				// Pick a new color for each batch : 
+				uint32 Hash = PointerHash(&RenderBatch);
+				uint8 * HashElement = reinterpret_cast<uint8*>(&Hash);
+				FColor Color(HashElement[0], HashElement[1], HashElement[2]);
+
+				UE_VLOG_OBOX(Landscape, LogLandscape, Log, FBox(FVector(RenderBatch.SectionRect.Min) - FVector(0.5, 0.5, 0.0), FVector(RenderBatch.SectionRect.Max) - FVector(0.5, 0.5, 0.0)),
+					ComputeVisualLogTransform(LandscapeWorldTransformForVisLog).ToMatrixWithScale(), Color.WithAlpha(GetVisualLogAlpha()),
+					TEXT("%s"), *FString::Format(TEXT("{0}\nBatch.SectionRect=([{1},{2}],[{3},{4}])"), { *RenderBatchDebugName, RenderBatch.SectionRect.Min.X, RenderBatch.SectionRect.Min.Y, RenderBatch.SectionRect.Max.X, RenderBatch.SectionRect.Max.Y }));
+			});
+
+		const int32 NumRenderSteps = RenderBatch.RenderSteps.Num();
+		for (int32 RenderStepIndex = 0; RenderStepIndex < NumRenderSteps; ++RenderStepIndex)
+		{
+			const FMergeRenderStep& RenderStep = RenderBatch.RenderSteps[RenderStepIndex];
+			const TArray<FName> TargetLayerNames = ConvertTargetLayerBitIndicesToNamesChecked(RenderStep.RenderGroupBitIndices);
+
+			// Compute all necessary info about the components affected by this renderer at this step
+			TArray<FComponentMergeRenderInfo> SortedComponentMergeRenderInfos;
+			SortedComponentMergeRenderInfos.Reserve(RenderStep.ComponentsToRender.Num());
+			Algo::Transform(RenderStep.ComponentsToRender, SortedComponentMergeRenderInfos, [MinComponentKey = RenderBatch.MinComponentKey](ULandscapeComponent* InComponent)
+			{
+				FComponentMergeRenderInfo ComponentMergeRenderInfo;
+				ComponentMergeRenderInfo.Component = InComponent;
+
+				const FIntPoint ComponentKey = InComponent->GetComponentKey();
+				const FIntPoint LocalComponentKey = ComponentKey - MinComponentKey;
+				check((LocalComponentKey.X >= 0) && (LocalComponentKey.Y >= 0));
+				ComponentMergeRenderInfo.ComponentKeyInRenderArea = LocalComponentKey;
+				// Area in the render target for this component : 
+				ComponentMergeRenderInfo.ComponentRegionInRenderArea = FIntRect(LocalComponentKey * InComponent->ComponentSizeQuads, (LocalComponentKey + 1) * InComponent->ComponentSizeQuads);
+
+				return ComponentMergeRenderInfo;
+			});
+			SortedComponentMergeRenderInfos.Sort();
+
+			// Is it a step involving a renderer ? 
+			if (ILandscapeEditLayerRenderer* Renderer = RenderStep.RendererState.GetRenderer())
+			{
+				if (RenderStep.Type == FMergeRenderStep::EType::RenderLayer)
+				{
+					TRACE_CPUPROFILER_EVENT_SCOPE(RenderLayer);
+
+					check(Renderer->CanRender());
+
+					// TODO[jonathan.bard] offset the world transform to account for the half-pixel offset?
+					//RenderParams.RenderAreaWorldTransform = FTransform(LandscapeTransform.GetRotation(), LandscapeTransform.GetTranslation() + FVector(RenderBatch.SectionRect.Min) /** (double)ComponentSizeQuads - FVector(0.5, 0.5, 0)*/, LandscapeTransform.GetScale3D());
+
+					FString RenderStepProfilingEventName = FString::Format(TEXT("Step [{0}] ({1}): Render {2}"), { RenderStepIndex, *ConvertTargetLayerNamesToString(TargetLayerNames), *Renderer->GetEditLayerRendererDebugName() });
+					SCOPED_DRAW_EVENTF_GAMETHREAD(LandscapeLayers, TEXT("%s"), RenderStepProfilingEventName);
+
+					// TODO [jonathan.bard] : this is more of a Batch world transform/section rect at the moment. Shall we have a RenderAreaWorldTransform/RenderAreaSectionRect in FRenderParams and a BatchRenderAreaWorldTransform in FMergeRenderBatch?
+					//  because currently the old BP brushes work with FMergeRenderBatch data (i.e. 1 transform for the batch and a section rect for the entire batch) but eventually, renderers might be interested in just their Render step context, 
+					//  that is : 1 matrix corresponding to the bottom left corner of their list of components to render?
+					FTransform RenderAreaWorldTransform = FTransform(FVector(RenderBatch.SectionRect.Min)) * LandscapeTransform;
+					FIntRect RenderAreaSectionRect = RenderBatch.SectionRect;
+				
+					// Drop some visual cues to help understand how each renderer is applied :
+					UE_IFVLOG(
+						if (IsVisualLogEnabled() && bVisualLogMergeProcess)
+						{
+							FTransform RenderAreaWorldTransformForVisLog = FTransform(FVector(RenderBatch.SectionRect.Min)) * LandscapeWorldTransformForVisLog;
+							IncrementVisualLogOffset();
+							UE_VLOG_LOCATION(Landscape, LogLandscape, Log, ComputeVisualLogTransform(RenderAreaWorldTransformForVisLog).GetTranslation(), 10.0f, FColor::Red, TEXT("%s"), *RenderStepProfilingEventName);
+							UE_VLOG_WIREOBOX(Landscape, LogLandscape, Log, FBox(FVector(RenderBatch.SectionRect.Min) - FVector(0.5, 0.5, 0.0), FVector(RenderBatch.SectionRect.Max) - FVector(0.5, 0.5, 0.0)),
+								ComputeVisualLogTransform(LandscapeWorldTransformForVisLog).ToMatrixWithScale(), FColor::White, TEXT(""));
+
+							// Draw each component's bounds rendered by this renderer : 
+							for (const FComponentMergeRenderInfo& ComponentMergeRenderInfo : SortedComponentMergeRenderInfos)
+							{
+								UE_VLOG_WIREOBOX(Landscape, LogLandscape, Log, FBox(FVector(ComponentMergeRenderInfo.ComponentRegionInRenderArea.Min), FVector(ComponentMergeRenderInfo.ComponentRegionInRenderArea.Max)),
+									ComputeVisualLogTransform(RenderAreaWorldTransformForVisLog).ToMatrixWithScale(), FColor::White, TEXT(""));
+							}
+						});
+			
+					ILandscapeEditLayerRenderer::FRenderParams RenderParams(this, RenderStep.RenderGroupBitIndices, RenderStep.RendererState, SortedComponentMergeRenderInfos, RenderAreaWorldTransform, RenderAreaSectionRect);
+					Renderer->RenderLayer(RenderParams);
+				}
+			}
+			else if ((RenderStep.Type == FMergeRenderStep::EType::SignalBatchMergeGroupDone))
+			{
+				TRACE_CPUPROFILER_EVENT_SCOPE(MergeGroupDone);
+				SCOPED_DRAW_EVENTF_GAMETHREAD(LandscapeLayers, TEXT("Step [%i] (%s) : Render Group Done"), RenderStepIndex, ConvertTargetLayerNamesToString(TargetLayerNames));
+
+				// The last render target we wrote to is the one containing the batch group's merge result : 
+				FOnRenderBatchTargetGroupDoneParams Params(this, RenderBatch, RenderStep.RenderGroupBitIndices, SortedComponentMergeRenderInfos);
+				OnBatchTargetGroupDone(Params);
+			}
+			else
+			{
+				check(false); // unknown render step type
+			}
+		}
+		
+		FreeBatchResources(RenderBatch);
+	}
+
+	FreeResources();
+}
+
+const FMergeRenderBatch* FMergeRenderContext::GetCurrentRenderBatch() const
+{
+	return RenderBatches.IsValidIndex(CurrentRenderBatchIndex) ? &RenderBatches[CurrentRenderBatchIndex] : nullptr;
+}
+
+TBitArray<> FMergeRenderContext::ConvertTargetLayerNamesToBitIndices(TConstArrayView<FName> InTargetLayerNames) const
+{
+	TBitArray<> Result(false, AllTargetLayerNames.Num());
+	for (FName Name : InTargetLayerNames)
+	{
+		if (int32 Index = GetTargetLayerIndexForName(Name); Index != INDEX_NONE)
+		{
+			Result[Index] = true;
+		}
+	}
+	return Result;
+}
+
+TBitArray<> FMergeRenderContext::ConvertTargetLayerNamesToBitIndicesChecked(TConstArrayView<FName> InTargetLayerNames) const
+{
+	TBitArray<> Result(false, AllTargetLayerNames.Num());
+	for (FName Name : InTargetLayerNames)
+	{
+		int32 Index = GetTargetLayerIndexForNameChecked(Name);
+		Result[Index] = true;
+	}
+	return Result;
+}
+
+TArray<FName> FMergeRenderContext::ConvertTargetLayerBitIndicesToNames(const TBitArray<>& InTargetLayerBitIndices) const
+{
+	const int32 NumNames = AllTargetLayerNames.Num();
+	TArray<FName> Names;
+	Names.Reserve(NumNames);
+	for (TConstSetBitIterator It(InTargetLayerBitIndices); It && (It.GetIndex() < NumNames); ++It)
+	{
+		Names.Add(AllTargetLayerNames[It.GetIndex()]);
+	}
+	return Names;
+}
+
+TArray<FName> FMergeRenderContext::ConvertTargetLayerBitIndicesToNamesChecked(const TBitArray<>& InTargetLayerBitIndices) const
+{
+	const int32 NumNames = AllTargetLayerNames.Num();
+	check(InTargetLayerBitIndices.Num() == NumNames);
+	TArray<FName> Names;
+	Names.Reserve(NumNames);
+	for (TConstSetBitIterator It(InTargetLayerBitIndices); It; ++It)
+	{
+		Names.Add(AllTargetLayerNames[It.GetIndex()]);
+	}
+	return Names;
+}
+
+TArray<ULandscapeLayerInfoObject*> FMergeRenderContext::ConvertTargetLayerBitIndicesToLayerInfos(const TBitArray<>& InTargetLayerBitIndices) const
+{
+	const int32 NumTargetLayerInfos = AllTargetLayerNames.Num();
+	TArray<ULandscapeLayerInfoObject*> LayerInfos;
+	LayerInfos.Reserve(NumTargetLayerInfos);
+	for (TConstSetBitIterator It(InTargetLayerBitIndices); It && (It.GetIndex() < NumTargetLayerInfos); ++It)
+	{
+		LayerInfos.Add(WeightmapLayerInfos[It.GetIndex()]);
+	}
+	return LayerInfos;
+}
+
+TArray<ULandscapeLayerInfoObject*> FMergeRenderContext::ConvertTargetLayerBitIndicesToLayerInfosChecked(const TBitArray<>& InTargetLayerBitIndices) const
+{
+	const int32 NumTargetLayerInfos = AllTargetLayerNames.Num();
+	check(InTargetLayerBitIndices.Num() == NumTargetLayerInfos);
+	TArray<ULandscapeLayerInfoObject*> LayerInfos;
+	LayerInfos.Reserve(NumTargetLayerInfos);
+	for (TConstSetBitIterator It(InTargetLayerBitIndices); It; ++It)
+	{
+		LayerInfos.Add(WeightmapLayerInfos[It.GetIndex()]);
+	}
+	return LayerInfos;
+}
+
+void FMergeRenderContext::ForEachTargetLayer(const TBitArray<>& InTargetLayerBitIndices, TFunctionRef<bool(int32 /*InTargetLayerIndex*/, FName /*InTargetLayerName*/)> Fn) const
+{
+	for (TConstSetBitIterator It(InTargetLayerBitIndices); It; ++It)
+	{
+		const int32 TargetLayerIndex = It.GetIndex();
+		if (!AllTargetLayerNames.IsValidIndex(TargetLayerIndex))
+		{
+			return;
+		}
+
+		if (!Fn(TargetLayerIndex, AllTargetLayerNames[TargetLayerIndex]))
+		{
+			return;
+		}
+	}
+}
+
+void FMergeRenderContext::ForEachTargetLayerChecked(const TBitArray<>& InTargetLayerBitIndices, TFunctionRef<bool(int32 /*InTargetLayerIndex*/, FName /*InTargetLayerName*/)> Fn) const
+{
+	const int32 NumNames = AllTargetLayerNames.Num();
+	check(InTargetLayerBitIndices.Num() == NumNames);
+
+	for (TConstSetBitIterator It(InTargetLayerBitIndices); It; ++It)
+	{
+		const int32 TargetLayerIndex = It.GetIndex();
+		if (!Fn(TargetLayerIndex, AllTargetLayerNames[TargetLayerIndex]))
+		{
+			return;
+		}
+	}
+}
+
+#if ENABLE_VISUAL_LOG
+
+int32 FMergeRenderContext::GetVisualLogAlpha()
+{
+	return FMath::Clamp(CVarLandscapeBatchedMergeVisualLogAlpha.GetValueOnGameThread(), 0.0f, 1.0f) * 255;
+}
+
+bool FMergeRenderContext::IsVisualLogEnabled() const
+{
+	switch (CVarLandscapeBatchedMergeVisualLogShowMergeType.GetValueOnGameThread())
+	{
+	case 0: // Disabled
+		return false;
+	case 1: // Heightmaps only
+		return bIsHeightmapMerge;
+	case 2: // Weightmaps only
+		return !bIsHeightmapMerge;
+	case 3: // Both
+		return true;
+	default:
+		return false;
+	}
+}
+
+#endif // ENABLE_VISUAL_LOG
+
+
+// ----------------------------------------------------------------------------------
+
+FIntRect FInputWorldArea::GetLocalComponentKeys(const FIntPoint& InComponentKey) const
+{
+	check(Type == EType::LocalComponent); 
+	return LocalArea + InComponentKey;
+}
+
+FIntRect FInputWorldArea::GetSpecificComponentKeys() const
+{
+	check(Type == EType::SpecificComponent);
+	return LocalArea + SpecificComponentKey;
+}
+
+FBox FInputWorldArea::ComputeWorldAreaAABB(const FTransform& InLandscapeTransform, const FBox& InLandscapeLocalBounds, const FTransform& InComponentTransform, const FBox& InComponentLocalBounds) const
+{
+	switch (Type)
+	{
+		case EType::Infinite:
+		{
+			return InLandscapeLocalBounds.TransformBy(InLandscapeTransform);
+		}
+		case EType::LocalComponent:
+		{
+			return InComponentLocalBounds.TransformBy(InComponentTransform);
+		}
+		case EType::SpecificComponent:
+		{
+			FVector ComponentLocalSize = InComponentLocalBounds.GetSize();
+			FIntRect LocalAreaCoordinates(SpecificComponentKey + LocalArea.Min, SpecificComponentKey + LocalArea.Max);
+			FBox LocalAreaBounds = FBox(
+				FVector(LocalAreaCoordinates.Min.X * ComponentLocalSize.X, LocalAreaCoordinates.Min.Y * ComponentLocalSize.Y, InComponentLocalBounds.Min.Z),
+				FVector(LocalAreaCoordinates.Max.X * ComponentLocalSize.X, LocalAreaCoordinates.Max.Y * ComponentLocalSize.Y, InComponentLocalBounds.Max.Z));
+			return LocalAreaBounds.TransformBy(InComponentTransform);
+		}
+		case EType::OOBox:
+		{
+			return FBox::BuildAABB(OOBox2D.Transform.GetTranslation(), OOBox2D.Transform.TransformVector(FVector(OOBox2D.Extents, 0.0)));
+		}
+		default:
+			check(false);
+	}
+
+	return FBox();
+}
+
+FOOBox2D FInputWorldArea::ComputeWorldAreaOOBB(const FTransform& InLandscapeTransform, const FBox& InLandscapeLocalBounds, const FTransform& InComponentTransform, const FBox& InComponentLocalBounds) const
+{
+	switch (Type)
+	{
+	case EType::Infinite:
+	{
+		FVector Center, Extents;
+		InLandscapeLocalBounds.GetCenterAndExtents(Center, Extents);
+		FTransform LandscapeTransformAtCenter = InLandscapeTransform;
+		LandscapeTransformAtCenter.SetTranslation(InLandscapeTransform.TransformVector(Center));
+		return FOOBox2D(LandscapeTransformAtCenter, FVector2D(Extents));
+	}
+	case EType::LocalComponent:
+	{
+		FVector Center, Extents;
+		InComponentLocalBounds.GetCenterAndExtents(Center, Extents);
+		FTransform ComponentTransformAtCenter = InComponentTransform;
+		ComponentTransformAtCenter.SetTranslation(InComponentTransform.TransformVector(Center));
+		return FOOBox2D(ComponentTransformAtCenter, FVector2D(Extents));
+	}
+	case EType::SpecificComponent:
+	{
+		FVector ComponentLocalSize = InComponentLocalBounds.GetSize();
+		FIntRect LocalAreaCoordinates(SpecificComponentKey + LocalArea.Min, SpecificComponentKey + LocalArea.Max);
+		FBox LocalAreaBounds = FBox(
+			FVector(LocalAreaCoordinates.Min.X * ComponentLocalSize.X, LocalAreaCoordinates.Min.Y * ComponentLocalSize.Y, InComponentLocalBounds.Min.Z),
+			FVector(LocalAreaCoordinates.Max.X * ComponentLocalSize.X, LocalAreaCoordinates.Max.Y * ComponentLocalSize.Y, InComponentLocalBounds.Max.Z));
+		FVector Center, Extents;
+		LocalAreaBounds.GetCenterAndExtents(Center, Extents);
+		FTransform ComponentTransformAtCenter = InComponentTransform;
+		ComponentTransformAtCenter.SetTranslation(InComponentTransform.TransformVector(Center));
+		return FOOBox2D(ComponentTransformAtCenter, FVector2D(Extents));
+	}
+	case EType::OOBox:
+	{
+		return OOBox2D;
+	}
+	default:
+		check(false);
+	}
+
+	return FOOBox2D();
+}
+
+
+// ----------------------------------------------------------------------------------
+
+FBox FOutputWorldArea::ComputeWorldAreaAABB(const FTransform& InComponentTransform, const FBox& InComponentLocalBounds) const
+{
+	switch (Type)
+	{
+	case EType::LocalComponent:
+	{
+		return InComponentLocalBounds.TransformBy(InComponentTransform);
+	}
+	case EType::SpecificComponent:
+	{
+		FVector ComponentLocalSize = InComponentLocalBounds.GetSize();
+		FBox LocalAreaBounds = FBox(
+			FVector(SpecificComponentKey.X * ComponentLocalSize.X, SpecificComponentKey.Y * ComponentLocalSize.Y, InComponentLocalBounds.Min.Z),
+			FVector((SpecificComponentKey.X + 1) * ComponentLocalSize.X, (SpecificComponentKey.Y + 1) * ComponentLocalSize.Y, InComponentLocalBounds.Max.Z));
+		return LocalAreaBounds.TransformBy(InComponentTransform);
+	}
+	case EType::OOBox:
+	{
+		return FBox::BuildAABB(OOBox2D.Transform.GetTranslation(), OOBox2D.Transform.TransformVector(FVector(OOBox2D.Extents, 0.0)));
+	}
+	default:
+		check(false);
+	}
+
+	return FBox();
+}
+
+FOOBox2D FOutputWorldArea::ComputeWorldAreaOOBB(const FTransform& InComponentTransform, const FBox& InComponentLocalBounds) const
+{
+	switch (Type)
+	{
+	case EType::LocalComponent:
+	{
+		FVector Center, Extents;
+		InComponentLocalBounds.GetCenterAndExtents(Center, Extents);
+		FTransform ComponentTransformAtCenter = InComponentTransform;
+		ComponentTransformAtCenter.SetTranslation(InComponentTransform.TransformVector(Center));
+		return FOOBox2D(ComponentTransformAtCenter, FVector2D(Extents));
+	}
+	case EType::SpecificComponent:
+	{
+		FVector ComponentLocalSize = InComponentLocalBounds.GetSize();
+		FBox LocalAreaBounds = FBox(
+			FVector(SpecificComponentKey.X * ComponentLocalSize.X, SpecificComponentKey.Y * ComponentLocalSize.Y, InComponentLocalBounds.Min.Z),
+			FVector((SpecificComponentKey.X + 1) * ComponentLocalSize.X, (SpecificComponentKey.Y + 1) * ComponentLocalSize.Y, InComponentLocalBounds.Max.Z));
+		FVector Center, Extents;
+		LocalAreaBounds.GetCenterAndExtents(Center, Extents);
+		FTransform ComponentTransformAtCenter = InComponentTransform;
+		ComponentTransformAtCenter.SetTranslation(InComponentTransform.TransformVector(Center));
+		return FOOBox2D(ComponentTransformAtCenter, FVector2D(Extents));
+	}
+	case EType::OOBox:
+	{
+		return OOBox2D;
+	}
+	default:
+		check(false);
+	}
+
+	return FOOBox2D();
+}
+
+
+// ----------------------------------------------------------------------------------
+
+bool FComponentMergeRenderInfo::operator<(const FComponentMergeRenderInfo& InOther) const
+{
+	// Sort by X / Y so that the order in which we render them is always consistent : 
+	if (ComponentRegionInRenderArea.Min.Y < InOther.ComponentRegionInRenderArea.Min.Y)
+	{
+		return true;
+	}
+	else if (ComponentRegionInRenderArea.Min.Y == InOther.ComponentRegionInRenderArea.Min.Y)
+	{
+		return (ComponentRegionInRenderArea.Min.X < InOther.ComponentRegionInRenderArea.Min.X);
+	}
+	return false;
+}
+
+
+// ----------------------------------------------------------------------------------
+
+FMergeRenderParams::FMergeRenderParams(bool bInIsHeightmapMerge, ALandscape* InLandscape, TArray<ULandscapeComponent*> InComponentsToMerge, const TArrayView<FEditLayerRendererState>& InEditLayerRendererStates, const TSet<FName>& InWeightmapLayerNames)
+	: bIsHeightmapMerge(bInIsHeightmapMerge)
+	, Landscape(InLandscape)
+	, ComponentsToMerge(InComponentsToMerge)
+	, EditLayerRendererStates(InEditLayerRendererStates)
+	, WeightmapLayerNames(InWeightmapLayerNames)
+{
+	if (bIsHeightmapMerge)
+	{
+		// Make sure no weightmap layer name is passed in the case of heightmap :
+		WeightmapLayerNames.Empty();
+	}
+}
+
+#endif // WITH_EDITOR
+
+} // namespace UE::Landscape::EditLayers
+
+#undef LOCTEXT_NAMESPACE
diff --git a/Engine/Source/Runtime/Landscape/Private/LandscapeEditLayerRendererPrivate.cpp b/Engine/Source/Runtime/Landscape/Private/LandscapeEditLayerRendererPrivate.cpp
new file mode 100644
--- /dev/null
+++ b/Engine/Source/Runtime/Landscape/Private/LandscapeEditLayerRendererPrivate.cpp
@@ -0,0 +1,154 @@
+// Copyright Epic Games, Inc. All Rights Reserved.
+
+#include "LandscapeEditLayerRendererPrivate.h"
+
+#include "Algo/AllOf.h"
+#include "Algo/ForEach.h"
+#include "Algo/Transform.h"
+#include "LandscapeEditResourcesSubsystem.h"
+#include "LandscapeLayerInfoObject.h"
+#include "LandscapeInfo.h"
+#include "RHIAccess.h"
+
+
+#if WITH_EDITOR
+
+// ----------------------------------------------------------------------------------
+
+void ULandscapeDefaultEditLayerRenderer::GetRendererStateInfo(const ULandscapeInfo* InLandscapeInfo,
+	UE::Landscape::EditLayers::FEditLayerTargetTypeState& OutSupportedTargetTypeState, UE::Landscape::EditLayers::FEditLayerTargetTypeState& OutEnabledTargetTypeState, TArray<TSet<FName>>& OutRenderGroups) const
+{
+	// Supports all heightmaps and weightmaps:
+	OutSupportedTargetTypeState.SetTargetTypeMask(ELandscapeToolTargetTypeFlags::All);
+	TSet<FName> WeightmapLayerNames = GatherWeightmapLayerNames(InLandscapeInfo);
+	Algo::ForEach(WeightmapLayerNames, [&OutSupportedTargetTypeState](FName InLayerName) { OutSupportedTargetTypeState.AddWeightmap(InLayerName); });
+	OutEnabledTargetTypeState = OutSupportedTargetTypeState;
+}
+
+TArray<UE::Landscape::EditLayers::FEditLayerRenderItem> ULandscapeDefaultEditLayerRenderer::GetRenderItems(const ULandscapeInfo* InLandscapeInfo) const
+{
+	using namespace UE::Landscape::EditLayers;
+
+	FEditLayerTargetTypeState OutputTargetTypeState(ELandscapeToolTargetTypeFlags::All);
+	// Standard renderer : we don't need more than the component itself to render properly:
+	FInputWorldArea InputWorldArea(FInputWorldArea::CreateLocalComponent());
+	// The renderer only writes into the component itself (i.e. it renders to the area that it's currently being asked to render to):
+	FOutputWorldArea OutputWorldArea(FOutputWorldArea::CreateLocalComponent());
+	// The renderer is only providing default data for existing weightmaps so it doesn't generate new ones, hence we pass bModifyExistingWeightmapsOnly to true : 
+	return { FEditLayerRenderItem(OutputTargetTypeState, InputWorldArea, OutputWorldArea, /*bModifyExistingWeightmapsOnly = */true) };
+}
+
+FString ULandscapeDefaultEditLayerRenderer::GetEditLayerRendererDebugName() const 
+{
+	return TEXT("Default");
+}
+
+void ULandscapeDefaultEditLayerRenderer::RenderLayer(FRenderParams& InRenderParams) 
+{
+	using namespace UE::Landscape::EditLayers;
+
+	FMergeRenderContext* RenderContext = InRenderParams.MergeRenderContext;
+	const FMergeRenderBatch* RenderBatch = RenderContext->GetCurrentRenderBatch();
+
+	RenderContext->CycleBlendRenderTargets(/*InDesiredWriteAccess = */ERHIAccess::RTV);
+	ULandscapeScratchRenderTarget* WriteRT = RenderContext->GetBlendRenderTargetWrite();
+
+	// Start from a blank canvas so that the first layer is blended with nothing underneath :
+	WriteRT->Clear();
+
+	// Render the components of the batch for each target layer into the "pseudo-stencil" buffer, so that it can be sampled by users as a UTextures in UMaterials and such : 
+	RenderContext->RenderValidityRenderTargets(*RenderBatch);
+}
+
+TSet<FName> ULandscapeDefaultEditLayerRenderer::GatherWeightmapLayerNames(const ULandscapeInfo* InLandscapeInfo) const
+{
+	TSet<FName> WeightmapLayerNames;
+	// This renderer supports all weightmaps:
+	Algo::TransformIf(InLandscapeInfo->Layers, WeightmapLayerNames,
+		[](const FLandscapeInfoLayerSettings& InInfoLayerSettings) { return (InInfoLayerSettings.LayerInfoObj != nullptr); },
+		[](const FLandscapeInfoLayerSettings& InInfoLayerSettings) { return InInfoLayerSettings.LayerInfoObj->LayerName; });
+	return WeightmapLayerNames;
+}
+
+
+// ----------------------------------------------------------------------------------
+
+void ULandscapeHeightmapNormalsEditLayerRenderer::GetRendererStateInfo(const ULandscapeInfo* InLandscapeInfo,
+	UE::Landscape::EditLayers::FEditLayerTargetTypeState& OutSupportedTargetTypeState, UE::Landscape::EditLayers::FEditLayerTargetTypeState& OutEnabledTargetTypeState, TArray<TSet<FName>>& OutRenderGroups) const
+{
+	// Only relevant for heightmaps :
+	OutSupportedTargetTypeState.SetTargetTypeMask(ELandscapeToolTargetTypeFlags::Heightmap);
+	OutEnabledTargetTypeState.SetTargetTypeMask(ELandscapeToolTargetTypeFlags::Heightmap);
+}
+
+TArray<UE::Landscape::EditLayers::FEditLayerRenderItem> ULandscapeHeightmapNormalsEditLayerRenderer::GetRenderItems(const ULandscapeInfo* InLandscapeInfo) const
+{
+	using namespace UE::Landscape::EditLayers;
+
+	// Only relevant for heightmaps :
+	FEditLayerTargetTypeState OutputTargetTypeState(ELandscapeToolTargetTypeFlags::Heightmap);
+	// The input is relative and its size is equal to the size of 3x3 landscape components so that we gather all neighbor landscape components around each component : 
+	FInputWorldArea InputWorldArea(FInputWorldArea::CreateLocalComponent(FIntRect(-1, -1, 1, 1)));
+	// The renderer only writes into the component itself (i.e. it renders to the area that it's currently being asked to render to):
+	FOutputWorldArea OutputWorldArea(FOutputWorldArea::CreateLocalComponent());
+	return { FEditLayerRenderItem(OutputTargetTypeState, InputWorldArea, OutputWorldArea, /*bModifyExistingWeightmapsOnly = */false) };
+}
+
+FString ULandscapeHeightmapNormalsEditLayerRenderer::GetEditLayerRendererDebugName() const
+{
+	return TEXT("Normals");
+}
+
+
+// ----------------------------------------------------------------------------------
+
+void ULandscapeWeightmapWeightBlendedLayersRenderer::GetRendererStateInfo(const ULandscapeInfo* InLandscapeInfo,
+	UE::Landscape::EditLayers::FEditLayerTargetTypeState& OutSupportedTargetTypeState, UE::Landscape::EditLayers::FEditLayerTargetTypeState& OutEnabledTargetTypeState, TArray<TSet<FName>>& OutRenderGroups) const
+{
+	// Only relevant for weightmaps :
+	OutSupportedTargetTypeState.SetTargetTypeMask(ELandscapeToolTargetTypeFlags::Weightmap);
+	TSet<FName> WeightBlendedWeightmapLayerNames = GatherWeightBlendedWeightmapLayerNames(InLandscapeInfo);
+	Algo::ForEach(WeightBlendedWeightmapLayerNames, [&OutSupportedTargetTypeState](FName InLayerName) { OutSupportedTargetTypeState.AddWeightmap(InLayerName); });
+	OutEnabledTargetTypeState = OutSupportedTargetTypeState;
+
+	check(Algo::AllOf(WeightBlendedWeightmapLayerNames, [InLandscapeInfo](FName InTargetLayerName)
+	{
+		return InLandscapeInfo->Layers.ContainsByPredicate([InTargetLayerName](const FLandscapeInfoLayerSettings& InInfoLayerSettings) { return (InInfoLayerSettings.GetLayerName() == InTargetLayerName); });
+	}));
+
+	// Now fill in the render groups :
+	if (!WeightBlendedWeightmapLayerNames.IsEmpty())
+	{
+		OutRenderGroups.Add(WeightBlendedWeightmapLayerNames);
+	}
+}
+
+TArray<UE::Landscape::EditLayers::FEditLayerRenderItem> ULandscapeWeightmapWeightBlendedLayersRenderer::GetRenderItems(const ULandscapeInfo* InLandscapeInfo) const
+{
+	using namespace UE::Landscape::EditLayers;
+
+	// Only relevant for weightmaps :
+	FEditLayerTargetTypeState OutputTargetTypeState(ELandscapeToolTargetTypeFlags::Weightmap, GatherWeightBlendedWeightmapLayerNames(InLandscapeInfo).Array());
+	// Standard renderer : we don't need more than the component itself to render properly:
+	FInputWorldArea InputWorldArea(FInputWorldArea::CreateLocalComponent());
+	// The renderer only writes into the component itself (i.e. it renders to the area that it's currently being asked to render to):
+	FOutputWorldArea OutputWorldArea(FOutputWorldArea::CreateLocalComponent());
+	// The renderer is only blending existing weightmaps so it doesn't generate new ones, hence we pass bModifyExistingWeightmapsOnly to true : 
+	return { FEditLayerRenderItem(OutputTargetTypeState, InputWorldArea, OutputWorldArea, /*bModifyExistingWeightmapsOnly = */true) };
+}
+
+FString ULandscapeWeightmapWeightBlendedLayersRenderer::GetEditLayerRendererDebugName() const 
+{
+	return TEXT("Final Weight Blend");
+}
+
+TSet<FName> ULandscapeWeightmapWeightBlendedLayersRenderer::GatherWeightBlendedWeightmapLayerNames(const ULandscapeInfo* InLandscapeInfo) const
+{
+	TSet<FName> WeightBlendedWeightmapLayerNames;
+	Algo::TransformIf(InLandscapeInfo->Layers, WeightBlendedWeightmapLayerNames,
+		[](const FLandscapeInfoLayerSettings& InInfoLayerSettings) { return ((InInfoLayerSettings.LayerInfoObj != nullptr) && !InInfoLayerSettings.LayerInfoObj->bNoWeightBlend); },
+		[](const FLandscapeInfoLayerSettings& InInfoLayerSettings) { return InInfoLayerSettings.LayerInfoObj->LayerName; });
+	return WeightBlendedWeightmapLayerNames;
+}
+
+#endif // WITH_EDITOR
\ No newline at end of file
diff --git a/Engine/Source/Runtime/Landscape/Private/LandscapeEditLayerRendererPrivate.h b/Engine/Source/Runtime/Landscape/Private/LandscapeEditLayerRendererPrivate.h
new file mode 100644
--- /dev/null
+++ b/Engine/Source/Runtime/Landscape/Private/LandscapeEditLayerRendererPrivate.h
@@ -0,0 +1,93 @@
+// Copyright Epic Games, Inc. All Rights Reserved.
+
+#pragma once
+
+#include "LandscapeEditTypes.h"
+#include "LandscapeEditLayerRenderer.h"
+
+#include "LandscapeEditLayerRendererPrivate.generated.h"
+
+class ILandscapeEditLayerRenderer;
+
+
+/**
+* Edit layer renderer added at the bottom of the stack to provide the default value for every requested target layer (heightmaps and weightmaps)
+*  It could have been handled with a simple clear as the first operation when performing the merge, but doing it through a renderer turns out to be an elegant way to resolve the
+*  situation where the first actual edit layer's render item declares dependencies between each component and others (e.g. say you have only a ULandscapeHeightmapNormalsEditLayerRenderer
+*  in the renderer stack, which requires each component's immediate neighbors). In such a situation, the component dependencies would be skipped because dependencies are
+*  between a renderer and its previous one in the stack and since in the case described above (a single renderer in the stack), there's no previous renderer, then the dependencies
+*  would simply not be registered, and the render batches would end up being incorrect as a result.
+*/
+UCLASS()
+class ULandscapeDefaultEditLayerRenderer : public UObject
+#if CPP && WITH_EDITOR
+	, public ILandscapeEditLayerRenderer
+#endif // CPP && WITH_EDITOR
+{
+	GENERATED_BODY()
+
+#if WITH_EDITOR
+public:
+	//~ Begin ILandscapeEditLayerRenderer implementation
+	virtual void GetRendererStateInfo(const ULandscapeInfo* InLandscapeInfo,
+		UE::Landscape::EditLayers::FEditLayerTargetTypeState& OutSupportedTargetTypeState, UE::Landscape::EditLayers::FEditLayerTargetTypeState& OutEnabledTargetTypeState, TArray<TSet<FName>>& OutRenderGroups) const override;
+	virtual TArray<UE::Landscape::EditLayers::FEditLayerRenderItem> GetRenderItems(const ULandscapeInfo* InLandscapeInfo) const override;
+	virtual FString GetEditLayerRendererDebugName() const override;
+	virtual void RenderLayer(FRenderParams& InRenderParams) override;
+	//~ End ILandscapeEditLayerRenderer implementation
+
+private:
+	TSet<FName> GatherWeightmapLayerNames(const ULandscapeInfo* InLandscapeInfo) const;
+#endif // WITH_EDITOR
+};
+
+/**
+* Edit layer renderer added at the top of the stack to generate the normals, right before resolving the textures.
+*  For now, the rendered components require (up to) their 8 immediate neighbors to generate adequate normals on the border so the renderer inserts that strong dependency, so that the components
+*  that are needed in the batch are guaranteed to have their neighbors present in the same batch :
+*/
+UCLASS()
+class ULandscapeHeightmapNormalsEditLayerRenderer : public UObject
+#if CPP && WITH_EDITOR
+	, public ILandscapeEditLayerRenderer
+#endif //CPP &&  WITH_EDITOR
+{
+	GENERATED_BODY()
+
+#if WITH_EDITOR
+public:
+	//~ Begin ILandscapeEditLayerRenderer implementation
+	virtual void GetRendererStateInfo(const ULandscapeInfo* InLandscapeInfo,
+		UE::Landscape::EditLayers::FEditLayerTargetTypeState& OutSupportedTargetTypeState, UE::Landscape::EditLayers::FEditLayerTargetTypeState& OutEnabledTargetTypeState, TArray<TSet<FName>>& OutRenderGroups) const override;
+	virtual TArray<UE::Landscape::EditLayers::FEditLayerRenderItem> GetRenderItems(const ULandscapeInfo* InLandscapeInfo) const override;
+	virtual FString GetEditLayerRendererDebugName() const override;
+	virtual void RenderLayer(FRenderParams& InRenderParams) override;
+	//~ End ILandscapeEditLayerRenderer implementation
+#endif // WITH_EDITOR
+};
+
+/** 
+* Edit layer renderer inserted at the end of the edit layers stack merge to normalize the weights of the rendered weightmaps 
+*/
+UCLASS()
+class ULandscapeWeightmapWeightBlendedLayersRenderer : public UObject
+#if CPP && WITH_EDITOR
+	, public ILandscapeEditLayerRenderer
+#endif // CPP && WITH_EDITOR
+{
+	GENERATED_BODY()
+
+#if WITH_EDITOR
+public:
+	//~ Begin ILandscapeEditLayerRenderer implementation
+	virtual void GetRendererStateInfo(const ULandscapeInfo* InLandscapeInfo, 
+		UE::Landscape::EditLayers::FEditLayerTargetTypeState& OutSupportedTargetTypeState, UE::Landscape::EditLayers::FEditLayerTargetTypeState& OutEnabledTargetTypeState, TArray<TSet<FName>>& OutRenderGroups) const override;
+	virtual TArray<UE::Landscape::EditLayers::FEditLayerRenderItem> GetRenderItems(const ULandscapeInfo* InLandscapeInfo) const override;
+	virtual FString GetEditLayerRendererDebugName() const override;
+	virtual void RenderLayer(FRenderParams& InRenderParams) override;
+	//~ End ILandscapeEditLayerRenderer implementation
+
+	TSet<FName> GatherWeightBlendedWeightmapLayerNames(const ULandscapeInfo* InLandscapeInfo) const;
+#endif // WITH_EDITOR
+};
+
