Introducing landscape batched merge (activated by CVar : landscape.EditLayersLocalMerge.Enable 2): complete refactor of the landscape edit layers merge algorithm to solve the taxing GPU memory require

diff --git a/Engine/Source/Runtime/Landscape/Private/LandscapeEditLayers.cpp b/Engine/Source/Runtime/Landscape/Private/LandscapeEditLayers.cpp
--- a/Engine/Source/Runtime/Landscape/Private/LandscapeEditLayers.cpp
+++ b/Engine/Source/Runtime/Landscape/Private/LandscapeEditLayers.cpp
@@ -16,8 +16,11 @@
 #include "LandscapePrivate.h"
 #include "LandscapeEditReadback.h"
 #include "LandscapeEditResources.h"
+#include "LandscapeEditResourcesSubsystem.h"
+#include "LandscapeEditLayerRendererPrivate.h"
 #include "LandscapeEditTypes.h"
 #include "LandscapeUtils.h"
+#include "LandscapeUtilsPrivate.h"
 #include "LandscapeSubsystem.h"
 #include "LandscapeTextureStreamingManager.h"
 
@@ -25,7 +28,9 @@
 #include "Shader.h"
 #include "GlobalShader.h"
 #include "ShaderParameterUtils.h"
+#include "EngineModule.h"
 #include "Engine/TextureRenderTarget2D.h"
+#include "Engine/TextureRenderTarget2DArray.h"
 #include "Logging/MessageLog.h"
 #include "RenderCaptureInterface.h"
 #include "RenderGraph.h"
@@ -36,6 +41,7 @@
 #include "SceneView.h"
 #include "MaterialCachedData.h"
 #include "ContentStreaming.h"
+#include "Templates/TypeHash.h"
 
 #if WITH_EDITOR
 #include "AssetCompilingManager.h"
@@ -47,10 +53,14 @@
 #include "LandscapeMaterialInstanceConstant.h"
 #include "Materials/MaterialExpressionLandscapeVisibilityMask.h"
 #include "ShaderCompiler.h"
+#include "Algo/Accumulate.h"
 #include "Algo/Count.h"
 #include "Algo/AnyOf.h"
+#include "Algo/AllOf.h"
 #include "Algo/Unique.h"
 #include "Algo/Transform.h"
+#include "Algo/ForEach.h"
+#include "Algo/RemoveIf.h"
 #include "LandscapeSettings.h"
 #include "LandscapeRender.h"
 #include "LandscapeInfoMap.h"
@@ -72,21 +82,14 @@
 #include "Components/RuntimeVirtualTextureComponent.h"
 #include "UnrealEdGlobals.h" // GUnrealEd
 #include "Editor/UnrealEdEngine.h"
+#include "VisualLogger/VisualLogger.h"
 #endif
 
 #define LOCTEXT_NAMESPACE "Landscape"
 
-namespace LandscapeEditLayersLocals
-{
-	FTransform GetLandscapeExtentTransform(const ALandscape& Landscape, const FIntRect& InLandscapeExtent)
-	{
-		FTransform NewLandscapeTransform = Landscape.GetTransform();
-		FVector OffsetVector(InLandscapeExtent.Min.X, InLandscapeExtent.Min.Y, 0.f);
-		FVector Translation = NewLandscapeTransform.TransformFVector4(OffsetVector);
-		NewLandscapeTransform.SetTranslation(Translation);
-		return NewLandscapeTransform;
-	}
-}
+// TODO [jonathan.bard] this define is for when(/if) we implement the uber landscape material in the editor, where weightmaps are *not* RGBA packed but stored in a plain texture array. 
+//  This will allow several simplifications and optimizations to edit layers 
+#define SUPPORTS_LANDSCAPE_EDITORONLY_UBER_MATERIAL 0
 
 // Channel remapping
 extern const size_t ChannelOffsets[4];
@@ -201,6 +204,52 @@
 	0, 
 	TEXT("This will analyze weightmaps on readback and remove unneeded allocations (for unpainted layers)."));
 
+TAutoConsoleVariable<float> CVarLandscapeBatchedMergeVisualLogOffsetIncrement(
+	TEXT("landscape.BatchedMerge.VisualLog.OffsetIncrement"),
+	5000.0f,
+	TEXT("Offset (in unreal units) for visualizing each operation of the batched merge in the viewport via the visual logger."));
+
+TAutoConsoleVariable<float> CVarLandscapeBatchedMergeVisualLogAlpha(
+	TEXT("landscape.BatchedMerge.VisualLog.Alpha"),
+	0.5f,
+	TEXT("Alpha value to use when visualizing batched merge info in the viewport via the visual logger ([0.0, 1.0] range)"));
+
+TAutoConsoleVariable<int32> CVarLandscapeBatchedMergeVisualLogShowMergeType(
+	TEXT("landscape.BatchedMerge.VisualLog.ShowMergeType"),
+	3,
+	TEXT("Filter what to visualize in the visual logger when using batched merge (0 = no visual log, 1 = show heightmaps only, 2 = show weightmaps only, 3 = show all"));
+
+TAutoConsoleVariable<bool> CVarLandscapeBatchedMergeVisualLogShowMergeProcess(
+	TEXT("landscape.BatchedMerge.VisualLog.ShowMergeProcess"),
+	false,
+	TEXT("Allows to visualize the merge process in the visual logger (batches, affected components, etc.)"));
+
+TAutoConsoleVariable<int32> CVarLandscapeBatchedMergeVisualLogShowRenderItemsType(
+	TEXT("landscape.BatchedMerge.VisualLog.ShowRenderItemsType"),
+	0,
+	TEXT("Allows to visualize the edit layer renderers' render items in the visual logger when using batched merge (0 = no visual log, 1 = show input areas, 2 = show output areas, 3 = show all)"));
+
+TAutoConsoleVariable<FString> CVarLandscapeBatchedMergeVisualLogShowRenderItemsEditLayerRendererFilter(
+	TEXT("landscape.BatchedMerge.VisualLog.ShowRenderItemsEditLayerRendererFilter"),
+	"",
+	TEXT("Allows to filter the elements added to the visual log to only those pertaining to a given edit layer renderer : use in conjunction with landscape.BatchedMerge.VisualLog.ShowRenderItemsType (empty : display all elements, otherwise, only display the items related to the edit layer renderer if its name matches (partial match)"));
+
+TAutoConsoleVariable<bool> CVarLandscapeBatchedMergeVisualLogShowAllRenderItems(
+	TEXT("landscape.BatchedMerge.VisualLog.ShowAllRenderItems"),
+	false,
+	TEXT("Allows to visualize all render items : use in conjunction with landscape.BatchedMerge.VisualLog.ShowRenderItemsType (if true, all render items will be displayed. If false, only those that participate to the render will be"));
+
+TAutoConsoleVariable<int32> CVarLandscapeBatchedMergeVisualLogShowComponentDependencies(
+	TEXT("landscape.BatchedMerge.VisualLog.ShowComponentDependencies"),
+	0,
+	TEXT("Allows to visualize the dependencies between landscape components when using batched merge (0 = no visual log, 1 = show component coordinates + area affecting component, 2 = show component coordinates + renderer name + area affecting component)"));
+
+TAutoConsoleVariable<FString> CVarLandscapeBatchedMergeVisualLogShowComponentDependenciesFilter(
+	TEXT("landscape.BatchedMerge.VisualLog.ShowComponentDependenciesFilter"),
+	"",
+	TEXT("Allows to visualize all the dependency graph for the component specified : use the \"X= Y=\" format to specify the component for which to show the dependencies"));
+
+
 void OnLandscapeEditLayersLocalMergeChanged(IConsoleVariable* CVar)
 {
 	for (TObjectIterator<UWorld> It(/*AdditionalExclusionFlags = */RF_ClassDefaultObject, /*bIncludeDerivedClasses = */true, /*InInternalExclusionFlags = */EInternalObjectFlags::Garbage); It; ++It)
@@ -227,7 +276,8 @@
 static FAutoConsoleVariableRef CVarLandscapeEditLayersLocalMerge(
 	TEXT("landscape.EditLayersLocalMerge.Enable"),
 	LandscapeEditLayersLocalMerge,
-	TEXT("This will allow the new merge algorithm (that merges layers at the landscape component level) to be used on landscapes that support it. This is a temporary measure while waiting for non-compatible landscapes to be deprecated. "),
+	TEXT("Setting this to 1 will allow the local merge algorithm (that merges layers at the landscape component level) to be used on landscapes that support it. This is a temporary measure while waiting for non-compatible landscapes to be deprecated.\n\
+		Setting this to 2 will allow the batched merge algorithm (that merges layers in batches). Supports all landscape types. This is also a temporary measure but will be the default once it's stable."),
 	FConsoleVariableDelegate::CreateStatic(OnLandscapeEditLayersLocalMergeChanged));
 
 TAutoConsoleVariable<int32> CVarLandscapeEditLayersMaxComponentsPerHeightmapResolveBatch(
@@ -242,6 +292,18 @@
 	TEXT("Number of components being rendered in a single batch when resolving weightmaps. The higher the number, the more weightmaps can be resolved in a single batch (and the higher the GPU memory consumption since more transient textures will be needed in memory at a time)"),
 	ECVF_RenderThreadSafe);
 
+TAutoConsoleVariable<int32> CVarLandscapeEditLayersMaxResolutionPerRenderBatch(
+	TEXT("landscape.EditLayersLocalMerge.MaxResolutionPerRenderBatch"),
+	1024,
+	TEXT("Maximum supported resolution for merging edit layers in a single batch. The higher the number, the more components can be resolved in a single batch (and the higher the GPU memory consumption since merging requires several temporary textures)"),
+	ECVF_RenderThreadSafe);
+
+TAutoConsoleVariable<int32> CVarLandscapeEditLayersClearBeforeEachWriteToScratch(
+	TEXT("landscape.EditLayersLocalMerge.ClearBeforeEachWriteToScratch"),
+	0,
+	TEXT("Debug to help with RenderDoc debugging : clear each time we're about to write on a scratch render target (since those are reused and can be used to write RTs of different resolutions"),
+	ECVF_RenderThreadSafe);
+
 struct FLandscapeDirty
 {
 	FLandscapeDirty()
@@ -333,6 +395,7 @@
 			{
 				Weightmaps.Reset();
 				WeightmapToComponents.Reset();
+				WeightmapToChannelMask.Reset();
 			}
 
 			// Iterate on all tracked landscape components and keep track of components/heightmaps/weightmaps relationship :
@@ -360,7 +423,10 @@
 							check(Weightmap != nullptr);
 
 							Weightmaps.Add(Weightmap);
+
 							WeightmapToComponents.FindOrAdd(Weightmap).AddUnique(Component);
+							uint8& WeightmapChannelMask = WeightmapToChannelMask.FindOrAdd(Weightmap, 0);
+							WeightmapChannelMask |= (1 << AllocInfo.WeightmapTextureChannel);
 						}
 					}
 				}
@@ -373,24 +439,94 @@
 	TSet< UTexture2D* > Heightmaps;
 	TMap< UTexture2D*, TArray<ULandscapeComponent*> > HeightmapToComponents;
 	TSet< UTexture2D* > Weightmaps;
+	// Key = texture, Value = list of components that use this texture in their weightmap layer allocations 
 	TMap< UTexture2D*, TArray<ULandscapeComponent*> > WeightmapToComponents;
+	// Key = texture, Value = final channel mask for this texture (i.e. which channel is actually allocated in any component)
+	TMap<UTexture2D*, uint8> WeightmapToChannelMask;
 };
 ENUM_CLASS_FLAGS(FTextureToComponentHelper::ERefreshFlags);
 
-static FFileHelper::EColorChannel GetWeightmapColorChannel(const FWeightmapLayerAllocationInfo& AllocInfo)
+// Must match EEditLayerHeightmapBlendMode in LandscapeLayersHeightmapsPS.usf
+enum class ELandscapeEditLayerHeightmapBlendMode : uint32
 {
-	FFileHelper::EColorChannel ColorChannelMapping[] = { FFileHelper::EColorChannel::R, FFileHelper::EColorChannel::G, FFileHelper::EColorChannel::B, FFileHelper::EColorChannel::A };
-	FFileHelper::EColorChannel ColorChannel = FFileHelper::EColorChannel::All;
+	Additive = 0,
+	AlphaBlend,
 
-	if (ensure(AllocInfo.WeightmapTextureChannel < 4))
+	Num,
+};
+
+#endif // WITH_EDITOR
+
+namespace UE::Landscape::Private
+{
+
+#if WITH_EDITOR
+	static FFileHelper::EColorChannel GetWeightmapColorChannel(const FWeightmapLayerAllocationInfo& AllocInfo)
 	{
-		ColorChannel = ColorChannelMapping[AllocInfo.WeightmapTextureChannel];
+		FFileHelper::EColorChannel ColorChannelMapping[] = { FFileHelper::EColorChannel::R, FFileHelper::EColorChannel::G, FFileHelper::EColorChannel::B, FFileHelper::EColorChannel::A };
+		FFileHelper::EColorChannel ColorChannel = FFileHelper::EColorChannel::All;
+	
+		if (ensure(AllocInfo.WeightmapTextureChannel < 4))
+		{
+			ColorChannel = ColorChannelMapping[AllocInfo.WeightmapTextureChannel];
+		}
+
+		return ColorChannel;
 	}
 
-	return ColorChannel;
+	static ELandscapeEditLayerHeightmapBlendMode LandscapeBlendModeToEditLayerBlendMode(ELandscapeBlendMode InLandscapeBlendMode)
+	{
+		switch (InLandscapeBlendMode)
+		{
+		case LSBM_AdditiveBlend:
+			return ELandscapeEditLayerHeightmapBlendMode::Additive;
+		case LSBM_AlphaBlend:
+			return ELandscapeEditLayerHeightmapBlendMode::AlphaBlend;
+		default:
+			check(false);
+		}
+
+		return ELandscapeEditLayerHeightmapBlendMode::Num;
+	}
+#endif // WITH_EDITOR
+	
+
+	// ----------------------------------------------------------------------------------
+	// Texture channel swizzling : 
+	enum class ERGBAChannel : uint8 { R, G, B, A };
+
+	static constexpr uint8 BuildChannelSwizzleMask(ERGBAChannel InChannelR = ERGBAChannel::R, ERGBAChannel InChannelG = ERGBAChannel::G, ERGBAChannel InChannelB = ERGBAChannel::B, ERGBAChannel InChannelA = ERGBAChannel::A)
+	{
+		return (static_cast<uint8>(InChannelR) << 0)
+			| (static_cast<uint8>(InChannelG) << 2)
+			| (static_cast<uint8>(InChannelB) << 4)
+			| (static_cast<uint8>(InChannelA) << 6);
+	}
+
+	static ERGBAChannel ExtractDestinationChannelFromSwizzleMask(ERGBAChannel InSourceChannel, uint8 InSwizzleMask)
+	{
+		uint8 SourceChannelIndex = static_cast<uint8>(InSourceChannel);
+		return static_cast<ERGBAChannel>((((uint8)3 << (SourceChannelIndex * 2)) & InSwizzleMask) >> (SourceChannelIndex * 2));
+	}
+
+	FString GetChannelSwizzleMaskDescription(uint8 InSwizzleMask, int32 InNumChannels = 4)
+	{
+		static constexpr auto ChannelToChar = [](ERGBAChannel InChannel) { return (InChannel == ERGBAChannel::R) ? TCHAR('R') : (InChannel == ERGBAChannel::G) ? TCHAR('G') : (InChannel == ERGBAChannel::B) ? TCHAR('B') : TCHAR('A'); };
+		check(InNumChannels <= 4);
+		const TCHAR ChannelsChar[4] =
+		{
+			ChannelToChar(ExtractDestinationChannelFromSwizzleMask(ERGBAChannel::R, InSwizzleMask)),
+			ChannelToChar(ExtractDestinationChannelFromSwizzleMask(ERGBAChannel::G, InSwizzleMask)),
+			ChannelToChar(ExtractDestinationChannelFromSwizzleMask(ERGBAChannel::B, InSwizzleMask)),
+			ChannelToChar(ExtractDestinationChannelFromSwizzleMask(ERGBAChannel::A, InSwizzleMask))
+		};
+		FString Result;
+		Result.AppendChars(ChannelsChar, InNumChannels);
+		return Result;
+	}
+
+	static constexpr uint8 RGBAToRGBASwizzleMask = BuildChannelSwizzleMask(ERGBAChannel::R, ERGBAChannel::G, ERGBAChannel::B, ERGBAChannel::A);
 }
 
-#endif
-
 // Vertex format and vertex buffer
 
@@ -461,7 +597,8 @@
 	TArray<FLandscapeLayersTriangle> TriangleList;
 };
 
-// Custom Pixel and Vertex shaders
+// ----------------------------------------------------------------------------------
+// /Engine/Private/LandscapeLayersVS.usf shaders
 
 class FLandscapeLayersVS : public FGlobalShader
 {
@@ -498,6 +635,9 @@
 
 IMPLEMENT_GLOBAL_SHADER(FLandscapeLayersVS, "/Engine/Private/LandscapeLayersVS.usf", "VSMain", SF_Vertex);
 
+// ----------------------------------------------------------------------------------
+// /Engine/Private/LandscapeLayersPS.usf shaders
+
 struct FLandscapeLayersHeightmapShaderParameters
 {
 	FLandscapeLayersHeightmapShaderParameters()
@@ -657,7 +797,6 @@
 		, LayerVisible(true)
 		, LayerBlendMode(LSBM_AdditiveBlend)
 		, OutputAsSubstractive(false)
-		, OutputAsNormalized(false)
 		, CurrentMipSize(0, 0)
 		, ParentMipSize(0, 0)
 		, CurrentMipComponentVertexCount(0)
@@ -670,7 +809,6 @@
 	bool LayerVisible;
 	ELandscapeBlendMode LayerBlendMode;
 	bool OutputAsSubstractive;
-	bool OutputAsNormalized;
 	FIntPoint CurrentMipSize;
 	FIntPoint ParentMipSize;
 	int32 CurrentMipComponentVertexCount;
@@ -711,7 +849,7 @@
 		SetTextureParameter(BatchedParameters, ReadTexture2Param, ReadTexture2SamplerParam, TStaticSamplerState<SF_Point, AM_Clamp, AM_Clamp, AM_Clamp>::GetRHI(), InParams.ReadWeightmap2 != nullptr ? InParams.ReadWeightmap2->GetResource()->TextureRHI : GWhiteTexture->TextureRHI);
 
 		FVector4f LayerInfo(InParams.LayerAlpha, InParams.LayerVisible ? 1.0f : 0.0f, InParams.LayerBlendMode == LSBM_AlphaBlend ? 1.0f : 0.f, 0.f);
-		FVector4f OutputConfig(InParams.ApplyLayerModifiers ? 1.0f : 0.0f, InParams.OutputAsSubstractive ? 1.0f : 0.0f, InParams.ReadWeightmap2 != nullptr ? 1.0f : 0.0f, InParams.OutputAsNormalized ? 1.0f : 0.0f);
+		FVector4f OutputConfig(InParams.ApplyLayerModifiers ? 1.0f : 0.0f, InParams.OutputAsSubstractive ? 1.0f : 0.0f, InParams.ReadWeightmap2 != nullptr ? 1.0f : 0.0f, 0.0f);
 
 		SetShaderValue(BatchedParameters, LayerInfoParam, LayerInfo);
 		SetShaderValue(BatchedParameters, OutputConfigParam, OutputConfig);
@@ -814,12 +952,14 @@
 		ReadTexture1Param.Bind(Initializer.ParameterMap, TEXT("ReadTexture1"));
 		ReadTexture1SamplerParam.Bind(Initializer.ParameterMap, TEXT("ReadTexture1Sampler"));
 		SourceOffsetAndSizeUVParam.Bind(Initializer.ParameterMap, TEXT("SourceOffsetAndSizeUV"));
+		ChannelSwizzleMaskParam.Bind(Initializer.ParameterMap, TEXT("ChannelSwizzleMask"));
 	}
 
 	FLandscapeCopyTexturePS()
 	{}
 
-	void SetParameters(FRHIBatchedShaderParameters& BatchedParameters, FRHITexture* InSourceTextureRHI, const FIntPoint& InSourcePosition = FIntPoint(0, 0), const FIntPoint& InCopySizePixels = FIntPoint(0, 0))
+	void SetParameters(FRHIBatchedShaderParameters& BatchedParameters, FRHITexture* InSourceTextureRHI,
+		const FIntPoint& InSourcePosition = FIntPoint::ZeroValue, const FIntPoint& InCopySizePixels = FIntPoint::ZeroValue, uint8 InChannelSwizzleMask = UE::Landscape::Private::RGBAToRGBASwizzleMask)
 	{
 		FVector2f SourceSize(InSourceTextureRHI->GetSizeXY());
 		FVector2f SourceOffsetUV = FVector2f(InSourcePosition) / SourceSize;
@@ -829,18 +969,21 @@
 		FVector2f CopySizeUV = FinalCopySizePixels / SourceSize;
 		SetTextureParameter(BatchedParameters, ReadTexture1Param, ReadTexture1SamplerParam, TStaticSamplerState<SF_Point, AM_Clamp, AM_Clamp, AM_Clamp>::GetRHI(), InSourceTextureRHI);
 		SetShaderValue(BatchedParameters, SourceOffsetAndSizeUVParam, FVector4f(SourceOffsetUV.X, SourceOffsetUV.Y, CopySizeUV.X, CopySizeUV.Y));
+		SetShaderValue(BatchedParameters, ChannelSwizzleMaskParam, static_cast<uint32>(InChannelSwizzleMask));
 	}
 
 private:
 	LAYOUT_FIELD(FShaderResourceParameter, ReadTexture1Param);
 	LAYOUT_FIELD(FShaderResourceParameter, ReadTexture1SamplerParam);
 	LAYOUT_FIELD(FShaderParameter, SourceOffsetAndSizeUVParam);
+	LAYOUT_FIELD(FShaderParameter, ChannelSwizzleMaskParam);
 };
 
 IMPLEMENT_GLOBAL_SHADER(FLandscapeCopyTextureVS, "/Engine/Private/LandscapeLayersPS.usf", "CopyTextureVS", SF_Vertex);
 IMPLEMENT_GLOBAL_SHADER(FLandscapeCopyTexturePS, "/Engine/Private/LandscapeLayersPS.usf", "CopyTexturePS", SF_Pixel);
 
-// Compute shaders data
+// ----------------------------------------------------------------------------------
+// /Engine/Private/LandscapeLayersCS.usf shaders
 
 int32 GLandscapeLayerWeightmapThreadGroupSizeX = 16;
 int32 GLandscapeLayerWeightmapThreadGroupSizeY = 16;
@@ -1191,10 +1334,8 @@
 	FLandscapeLayerWeightmapPackMaterialLayersComputeShaderParameters ShaderParams;
 };
 
-
-// Landscape Heightmap PS
-
 // ----------------------------------------------------------------------------------
+// /Engine/Private/Landscape/LandscapeLayersHeightmapsPS.usf shaders : 
 
 class FLandscapeLayersHeightmapsMergeEditLayersPS : public FGlobalShader
 {
@@ -1236,9 +1377,6 @@
 
 IMPLEMENT_GLOBAL_SHADER(FLandscapeLayersHeightmapsMergeEditLayersPS, "/Engine/Private/Landscape/LandscapeLayersHeightmapsPS.usf", "MergeEditLayers", SF_Pixel);
 
-
-// ----------------------------------------------------------------------------------
-
 class FLandscapeLayersHeightmapsStitchHeightmapPS : public FGlobalShader
 {
 	DECLARE_GLOBAL_SHADER(FLandscapeLayersHeightmapsStitchHeightmapPS);
@@ -1280,8 +1418,6 @@
 
 IMPLEMENT_GLOBAL_SHADER(FLandscapeLayersHeightmapsStitchHeightmapPS, "/Engine/Private/Landscape/LandscapeLayersHeightmapsPS.usf", "StitchHeightmap", SF_Pixel);
 
-// ----------------------------------------------------------------------------------
-
 class FLandscapeLayersHeightmapsFinalizeHeightmapPS : public FGlobalShader
 {
 	DECLARE_GLOBAL_SHADER(FLandscapeLayersHeightmapsFinalizeHeightmapPS);
@@ -1325,9 +1461,6 @@
 
 IMPLEMENT_GLOBAL_SHADER(FLandscapeLayersHeightmapsFinalizeHeightmapPS, "/Engine/Private/Landscape/LandscapeLayersHeightmapsPS.usf", "FinalizeHeightmap", SF_Pixel);
 
-
-// ----------------------------------------------------------------------------------
-
 class FLandscapeLayersHeightmapsGenerateMipsPS : public FGlobalShader
 {
 	DECLARE_GLOBAL_SHADER(FLandscapeLayersHeightmapsGenerateMipsPS);
@@ -1370,7 +1503,9 @@
 
 IMPLEMENT_GLOBAL_SHADER(FLandscapeLayersHeightmapsGenerateMipsPS, "/Engine/Private/Landscape/LandscapeLayersHeightmapsPS.usf", "GenerateMips", SF_Pixel);
 
-// Landscape Weightmap PS
+// ----------------------------------------------------------------------------------
+// /Engine/Private/Landscape/LandscapeLayersWeightmapsPS.usf shaders :
+
 class FLandscapeLayersWeightmapsMergeEditLayersPS : public FGlobalShader
 {
 	DECLARE_GLOBAL_SHADER(FLandscapeLayersWeightmapsMergeEditLayersPS);
@@ -1427,7 +1562,7 @@
 		SHADER_PARAMETER(FUintVector4, InPerChannelNumPaintLayers)
 		SHADER_PARAMETER_RDG_TEXTURE_SRV_ARRAY(Texture2D<float4>, InPerChannelPaintLayerWeightmaps, [4])
 		SHADER_PARAMETER_RDG_BUFFER_SRV(StructuredBuffer<uint>, InPaintLayerInfoIndices)
-		SHADER_PARAMETER_RDG_BUFFER_SRV(StructuredBuffer<FLandscapeWeightmapPaintLayerInfo>, InPaintLayerInfos)
+		SHADER_PARAMETER_RDG_BUFFER_SRV(StructuredBuffer<UE::Landscape::EditLayers::FWeightmapPaintLayerInfo>, InPaintLayerInfos)
 		RENDER_TARGET_BINDING_SLOTS()
 		END_SHADER_PARAMETER_STRUCT()
 
@@ -1502,6 +1637,381 @@
 
 IMPLEMENT_GLOBAL_SHADER(FLandscapeLayersWeightmapsGenerateMipsPS, "/Engine/Private/Landscape/LandscapeLayersWeightmapsPS.usf", "GenerateMips", SF_Pixel);
 
+
+// ----------------------------------------------------------------------------------
+// /Engine/Private/Landscape/LandscapeEditLayersHeightmaps.usf shaders :
+
+class FLandscapeEditLayersHeightmapsMergeEditLayerPS : public FGlobalShader
+{
+	DECLARE_GLOBAL_SHADER(FLandscapeEditLayersHeightmapsMergeEditLayerPS);
+	SHADER_USE_PARAMETER_STRUCT(FLandscapeEditLayersHeightmapsMergeEditLayerPS, FGlobalShader);
+
+public:
+	BEGIN_SHADER_PARAMETER_STRUCT(FParameters, )
+		SHADER_PARAMETER(uint32, InEditLayerBlendMode)
+		SHADER_PARAMETER(float, InEditLayerAlpha)
+		SHADER_PARAMETER_RDG_TEXTURE_SRV(Texture2D<float4>, InCurrentEditLayerHeightmap)
+		SHADER_PARAMETER_RDG_TEXTURE_SRV(Texture2D<float4>, InPreviousEditLayersHeightmap)
+		RENDER_TARGET_BINDING_SLOTS()
+	END_SHADER_PARAMETER_STRUCT()
+
+	static bool ShouldCompilePermutation(const FGlobalShaderPermutationParameters& InParameters)
+	{
+		return UE::Landscape::DoesPlatformSupportEditLayers(InParameters.Platform);
+	}
+
+	static void ModifyCompilationEnvironment(const FGlobalShaderPermutationParameters& InParameters, FShaderCompilerEnvironment& OutEnvironment)
+	{
+		OutEnvironment.SetDefine(TEXT("MERGE_EDIT_LAYER"), 1);
+	}
+
+	static void MergeEditLayerPS(FRDGBuilder& GraphBuilder, FParameters* InParameters, const FName& InTargetLayerName, const FIntPoint& InTextureSize)
+	{
+		FGlobalShaderMap* ShaderMap = GetGlobalShaderMap(GMaxRHIFeatureLevel);
+		TShaderMapRef<FLandscapeEditLayersHeightmapsMergeEditLayerPS> PixelShader(ShaderMap);
+
+		FPixelShaderUtils::AddFullscreenPass(
+			GraphBuilder,
+			ShaderMap,
+			RDG_EVENT_NAME("LandscapeLayers_MergeEditLayerPS (%s)", *InTargetLayerName.ToString()),
+			PixelShader,
+			InParameters,
+			FIntRect(0, 0, InTextureSize.X, InTextureSize.Y), 
+			/*BlendState = */nullptr, 
+			/*RasterizerState = */nullptr, 
+			/*DepthStencilState = */nullptr, 
+			/*StencilRef = */0, 
+			ERDGPassFlags::NeverCull); // Use NeverCull because it renders a texture for which tracking is disabled
+	}
+};
+
+IMPLEMENT_GLOBAL_SHADER(FLandscapeEditLayersHeightmapsMergeEditLayerPS, "/Engine/Private/Landscape/LandscapeEditLayersHeightmaps.usf", "MergeEditLayerPS", SF_Pixel);
+
+class FLandscapeEditLayersHeightmapsGenerateNormalsPS : public FGlobalShader
+{
+	DECLARE_GLOBAL_SHADER(FLandscapeEditLayersHeightmapsGenerateNormalsPS);
+	SHADER_USE_PARAMETER_STRUCT(FLandscapeEditLayersHeightmapsGenerateNormalsPS, FGlobalShader);
+
+public:
+	BEGIN_SHADER_PARAMETER_STRUCT(FParameters, )
+		SHADER_PARAMETER(FUintVector4, InTextureSize)
+		SHADER_PARAMETER(FVector3f, InLandscapeGridScale)
+		SHADER_PARAMETER(uint32, InComponentSizeQuads)
+		SHADER_PARAMETER(FUintVector2, InNumComponents)
+		SHADER_PARAMETER_SAMPLER(SamplerState, InSourceHeightmapSampler)
+		SHADER_PARAMETER_RDG_TEXTURE_SRV(Texture2D<float4>, InSourceHeightmap)
+		SHADER_PARAMETER_RDG_TEXTURE_SRV(Texture2D<uint8>, InValidityTexture)
+		RENDER_TARGET_BINDING_SLOTS()
+	END_SHADER_PARAMETER_STRUCT()
+
+	static bool ShouldCompilePermutation(const FGlobalShaderPermutationParameters& InParameters)
+	{
+		return UE::Landscape::DoesPlatformSupportEditLayers(InParameters.Platform);
+	}
+
+	static void ModifyCompilationEnvironment(const FGlobalShaderPermutationParameters& InParameters, FShaderCompilerEnvironment& OutEnvironment)
+	{
+		OutEnvironment.SetDefine(TEXT("GENERATE_NORMALS"), 1);
+		OutEnvironment.CompilerFlags.Remove(CFLAG_HLSL2021 | CFLAG_PrecompileWithDXC);
+	}
+
+	static void GenerateNormalsPS(FRDGBuilder& GraphBuilder, FParameters* InParameters, const FIntPoint& InTextureSize)
+	{
+		FGlobalShaderMap* ShaderMap = GetGlobalShaderMap(GMaxRHIFeatureLevel);
+		TShaderMapRef<FLandscapeEditLayersHeightmapsGenerateNormalsPS> PixelShader(ShaderMap);
+
+		FPixelShaderUtils::AddFullscreenPass(
+			GraphBuilder,
+			ShaderMap,
+			RDG_EVENT_NAME("LandscapeLayers_GenerateNormalsPS"),
+			PixelShader,
+			InParameters,
+			FIntRect(0, 0, InTextureSize.X, InTextureSize.Y),
+			/*BlendState = */nullptr, 
+			/*RasterizerState = */nullptr, 
+			/*DepthStencilState = */nullptr, 
+			/*StencilRef = */0, 
+			ERDGPassFlags::NeverCull); // Use NeverCull because it renders a texture for which tracking is disabled
+	}
+};
+
+IMPLEMENT_GLOBAL_SHADER(FLandscapeEditLayersHeightmapsGenerateNormalsPS, "/Engine/Private/Landscape/LandscapeEditLayersHeightmaps.usf", "GenerateNormalsPS", SF_Pixel);
+
+class FLandscapeEditLayersHeightmapsGenerateMipsPS : public FGlobalShader
+{
+	DECLARE_GLOBAL_SHADER(FLandscapeEditLayersHeightmapsGenerateMipsPS);
+	SHADER_USE_PARAMETER_STRUCT(FLandscapeEditLayersHeightmapsGenerateMipsPS, FGlobalShader);
+
+public:
+	BEGIN_SHADER_PARAMETER_STRUCT(FParameters, )
+		SHADER_PARAMETER(FUintVector2, InCurrentMipSubsectionSize)
+		SHADER_PARAMETER_RDG_TEXTURE_SRV(Texture2D<float4>, InSourceHeightmap)
+		RENDER_TARGET_BINDING_SLOTS()
+	END_SHADER_PARAMETER_STRUCT()
+
+	static bool ShouldCompilePermutation(const FGlobalShaderPermutationParameters& InParameters)
+	{
+		return UE::Landscape::DoesPlatformSupportEditLayers(InParameters.Platform);
+	}
+
+	static void ModifyCompilationEnvironment(const FGlobalShaderPermutationParameters& InParameters, FShaderCompilerEnvironment& OutEnvironment)
+	{
+		OutEnvironment.SetDefine(TEXT("GENERATE_MIPS"), 1);
+	}
+
+	static void GenerateMipsPS(FRDGBuilder& GraphBuilder, FParameters* InParameters, const FIntPoint& InTextureSize)
+	{
+		FGlobalShaderMap* ShaderMap = GetGlobalShaderMap(GMaxRHIFeatureLevel);
+		TShaderMapRef<FLandscapeEditLayersHeightmapsGenerateMipsPS> PixelShader(ShaderMap);
+
+		FPixelShaderUtils::AddFullscreenPass(
+			GraphBuilder,
+			ShaderMap,
+			RDG_EVENT_NAME("LandscapeLayers_GenerateMipsPS"),
+			PixelShader,
+			InParameters,
+			FIntRect(0, 0, InTextureSize.X, InTextureSize.Y),
+			/*BlendState = */nullptr,
+			/*RasterizerState = */nullptr,
+			/*DepthStencilState = */nullptr,
+			/*StencilRef = */0,
+			ERDGPassFlags::NeverCull); // Use NeverCull because it renders a texture for which tracking is disabled
+	}
+};
+
+IMPLEMENT_GLOBAL_SHADER(FLandscapeEditLayersHeightmapsGenerateMipsPS, "/Engine/Private/Landscape/LandscapeEditLayersHeightmaps.usf", "GenerateMipsPS", SF_Pixel);
+
+
+// ----------------------------------------------------------------------------------
+// /Engine/Private/Landscape/LandscapeEditLayersWeightmaps.usf shaders :
+
+class FLandscapeEditLayersWeightmapsMergeEditLayerPS : public FGlobalShader
+{
+	DECLARE_GLOBAL_SHADER(FLandscapeEditLayersWeightmapsMergeEditLayerPS);
+	SHADER_USE_PARAMETER_STRUCT(FLandscapeEditLayersWeightmapsMergeEditLayerPS, FGlobalShader);
+
+public:
+	BEGIN_SHADER_PARAMETER_STRUCT(FParameters, )
+		SHADER_PARAMETER(uint32, InPaintLayerIndex)
+		SHADER_PARAMETER(uint32, InEditLayerPaintLayerBlendMode)
+		SHADER_PARAMETER(float, InEditLayerAlpha)
+		SHADER_PARAMETER_RDG_TEXTURE_SRV(Texture2DArray<float4>, InCurrentEditLayerWeightmaps)
+		SHADER_PARAMETER_RDG_TEXTURE_SRV(Texture2DArray<float4>, InPreviousEditLayersWeightmaps)
+		RENDER_TARGET_BINDING_SLOTS()
+	END_SHADER_PARAMETER_STRUCT()
+
+	static bool ShouldCompilePermutation(const FGlobalShaderPermutationParameters& InParameters)
+	{
+		return UE::Landscape::DoesPlatformSupportEditLayers(InParameters.Platform);
+	}
+
+	static void ModifyCompilationEnvironment(const FGlobalShaderPermutationParameters& InParameters, FShaderCompilerEnvironment& OutEnvironment)
+	{
+		OutEnvironment.SetDefine(TEXT("MERGE_EDIT_LAYER"), 1);
+	}
+
+	static void MergeEditLayerPS(FRDGBuilder& GraphBuilder, FParameters* InParameters, const FName& InTargetLayerName, const FIntPoint& InTextureSize)
+	{
+		FGlobalShaderMap* ShaderMap = GetGlobalShaderMap(GMaxRHIFeatureLevel);
+		TShaderMapRef<FLandscapeEditLayersWeightmapsMergeEditLayerPS> PixelShader(ShaderMap);
+
+		FPixelShaderUtils::AddFullscreenPass(
+			GraphBuilder,
+			ShaderMap,
+			RDG_EVENT_NAME("LandscapeLayers_MergeEditLayerPS (%s)", *InTargetLayerName.ToString()),
+			PixelShader,
+			InParameters,
+			FIntRect(0, 0, InTextureSize.X, InTextureSize.Y), 
+			/*BlendState = */nullptr,
+			/*RasterizerState = */nullptr,
+			/*DepthStencilState = */nullptr,
+			/*StencilRef = */0,
+			ERDGPassFlags::NeverCull); // Use NeverCull because it renders a texture for which tracking is disabled
+	}
+};
+
+IMPLEMENT_GLOBAL_SHADER(FLandscapeEditLayersWeightmapsMergeEditLayerPS, "/Engine/Private/Landscape/LandscapeEditLayersWeightmaps.usf", "MergeEditLayerPS", SF_Pixel);
+
+class FLandscapeEditLayersWeightmapsPerformLegacyWeightBlendingPS : public FGlobalShader
+{
+	DECLARE_GLOBAL_SHADER(FLandscapeEditLayersWeightmapsPerformLegacyWeightBlendingPS);
+	SHADER_USE_PARAMETER_STRUCT(FLandscapeEditLayersWeightmapsPerformLegacyWeightBlendingPS, FGlobalShader);
+
+public:
+	BEGIN_SHADER_PARAMETER_STRUCT(FParameters, )
+		SHADER_PARAMETER(uint32, InPaintLayerIndex)
+		SHADER_PARAMETER(uint32, InNumPaintLayers)
+		SHADER_PARAMETER_RDG_BUFFER_SRV(StructuredBuffer<UE::Landscape::FWeightmapPaintLayerInfo>, InPaintLayerInfos)
+		SHADER_PARAMETER_RDG_TEXTURE_SRV(Texture2DArray<float4>, InCurrentEditLayerWeightmaps)
+		RENDER_TARGET_BINDING_SLOTS()
+	END_SHADER_PARAMETER_STRUCT()
+
+	static bool ShouldCompilePermutation(const FGlobalShaderPermutationParameters& InParameters)
+	{
+		return UE::Landscape::DoesPlatformSupportEditLayers(InParameters.Platform);
+	}
+
+	static void ModifyCompilationEnvironment(const FGlobalShaderPermutationParameters& InParameters, FShaderCompilerEnvironment& OutEnvironment)
+	{
+		OutEnvironment.SetDefine(TEXT("PERFORM_LEGACY_WEIGHT_BLENDING"), 1);
+	}
+
+	static void PerformLegacyWeightBlendingPS(FRDGBuilder& GraphBuilder, FParameters* InParameters, const FName& InTargetLayerName, const FIntPoint& InTextureSize)
+	{
+		FGlobalShaderMap* ShaderMap = GetGlobalShaderMap(GMaxRHIFeatureLevel);
+		TShaderMapRef<FLandscapeEditLayersWeightmapsPerformLegacyWeightBlendingPS> PixelShader(ShaderMap);
+
+		FPixelShaderUtils::AddFullscreenPass(
+			GraphBuilder,
+			ShaderMap,
+			RDG_EVENT_NAME("LandscapeLayers_PerformLegacyWeightBlendingPS (%s)", *InTargetLayerName.ToString()),
+			PixelShader,
+			InParameters,
+			FIntRect(0, 0, InTextureSize.X, InTextureSize.Y), 
+			/*BlendState = */nullptr,
+			/*RasterizerState = */nullptr,
+			/*DepthStencilState = */nullptr,
+			/*StencilRef = */0,
+			ERDGPassFlags::NeverCull); // Use NeverCull because it renders a texture for which tracking is disabled
+	}
+};
+
+IMPLEMENT_GLOBAL_SHADER(FLandscapeEditLayersWeightmapsPerformLegacyWeightBlendingPS, "/Engine/Private/Landscape/LandscapeEditLayersWeightmaps.usf", "PerformLegacyWeightBlendingPS", SF_Pixel);
+
+class FLandscapeEditLayersWeightmapsPackWeightmapPS : public FGlobalShader
+{
+	DECLARE_GLOBAL_SHADER(FLandscapeEditLayersWeightmapsPackWeightmapPS);
+	SHADER_USE_PARAMETER_STRUCT(FLandscapeEditLayersWeightmapsPackWeightmapPS, FGlobalShader);
+
+public:
+	BEGIN_SHADER_PARAMETER_STRUCT(FParameters, )
+		SHADER_PARAMETER(FIntVector4, InSourceSliceIndices)
+		SHADER_PARAMETER_ARRAY(FUintVector4, InSourcePixelOffsets, [4])
+		SHADER_PARAMETER(FUintVector2, InSubsectionPixelOffset)
+		SHADER_PARAMETER(uint32, InIsAdditive)
+		SHADER_PARAMETER_RDG_TEXTURE_SRV(Texture2DArray<float4>, InSourceWeightmaps)
+		SHADER_PARAMETER_RDG_TEXTURE_SRV(Texture2D<float4>, InWeightmapBeingPacked)
+		RENDER_TARGET_BINDING_SLOTS()
+	END_SHADER_PARAMETER_STRUCT()
+
+	static bool ShouldCompilePermutation(const FGlobalShaderPermutationParameters& InParameters)
+	{
+		return UE::Landscape::DoesPlatformSupportEditLayers(InParameters.Platform);
+	}
+
+	static void ModifyCompilationEnvironment(const FGlobalShaderPermutationParameters& InParameters, FShaderCompilerEnvironment& OutEnvironment)
+	{
+		OutEnvironment.SetDefine(TEXT("PACK_WEIGHTMAP"), 1);
+	}
+
+	static void PackWeightmapPS(FRDGBuilder& GraphBuilder, FParameters* InParameters, const FIntRect& InTextureRect)
+	{
+		FGlobalShaderMap* ShaderMap = GetGlobalShaderMap(GMaxRHIFeatureLevel);
+		TShaderMapRef<FLandscapeEditLayersWeightmapsPackWeightmapPS> PixelShader(ShaderMap);
+
+		FPixelShaderUtils::AddFullscreenPass(
+			GraphBuilder,
+			ShaderMap,
+			RDG_EVENT_NAME("LandscapeLayers_PackWeightmapPS"),
+			PixelShader,
+			InParameters,
+			InTextureRect,
+			/*BlendState = */nullptr,
+			/*RasterizerState = */nullptr,
+			/*DepthStencilState = */nullptr,
+			/*StencilRef = */0,
+			ERDGPassFlags::NeverCull); // Use NeverCull because it renders a texture for which tracking is disabled
+	}
+};
+
+IMPLEMENT_GLOBAL_SHADER(FLandscapeEditLayersWeightmapsPackWeightmapPS, "/Engine/Private/Landscape/LandscapeEditLayersWeightmaps.usf", "PackWeightmapPS", SF_Pixel);
+
+class FLandscapeEditLayersWeightmapsGenerateMipsPS : public FGlobalShader
+{
+	DECLARE_GLOBAL_SHADER(FLandscapeEditLayersWeightmapsGenerateMipsPS);
+	SHADER_USE_PARAMETER_STRUCT(FLandscapeEditLayersWeightmapsGenerateMipsPS, FGlobalShader);
+
+public:
+	BEGIN_SHADER_PARAMETER_STRUCT(FParameters, )
+		SHADER_PARAMETER(FUintVector2, InCurrentMipSubsectionSize)
+		SHADER_PARAMETER_RDG_TEXTURE_SRV(Texture2D<float4>, InSourceWeightmap)
+		RENDER_TARGET_BINDING_SLOTS()
+	END_SHADER_PARAMETER_STRUCT()
+
+	static bool ShouldCompilePermutation(const FGlobalShaderPermutationParameters& InParameters)
+	{
+		return UE::Landscape::DoesPlatformSupportEditLayers(InParameters.Platform);
+	}
+
+	static void ModifyCompilationEnvironment(const FGlobalShaderPermutationParameters& InParameters, FShaderCompilerEnvironment& OutEnvironment)
+	{
+		OutEnvironment.SetDefine(TEXT("GENERATE_MIPS"), 1);
+	}
+
+	static void GenerateMipsPS(FRDGBuilder& GraphBuilder, FParameters* InParameters, const FIntPoint& InTextureSize)
+	{
+		FGlobalShaderMap* ShaderMap = GetGlobalShaderMap(GMaxRHIFeatureLevel);
+		TShaderMapRef<FLandscapeEditLayersWeightmapsGenerateMipsPS> PixelShader(ShaderMap);
+
+		FPixelShaderUtils::AddFullscreenPass(
+			GraphBuilder,
+			ShaderMap,
+			RDG_EVENT_NAME("LandscapeLayers_GenerateMipsPS"),
+			PixelShader,
+			InParameters,
+			FIntRect(0, 0, InTextureSize.X, InTextureSize.Y), 
+			/*BlendState = */nullptr,
+			/*RasterizerState = */nullptr,
+			/*DepthStencilState = */nullptr,
+			/*StencilRef = */0,
+			ERDGPassFlags::NeverCull); // Use NeverCull because it renders a texture for which tracking is disabled
+	}
+};
+
+IMPLEMENT_GLOBAL_SHADER(FLandscapeEditLayersWeightmapsGenerateMipsPS, "/Engine/Private/Landscape/LandscapeEditLayersWeightmaps.usf", "GenerateMipsPS", SF_Pixel);
+
+
+// ----------------------------------------------------------------------------------
+// /Engine/Private/Landscape/LandscapeEditLayersUtils.usf shaders : 
+
+class FCopyQuadsMultiSourcePS : public FGlobalShader
+{
+	DECLARE_GLOBAL_SHADER(FCopyQuadsMultiSourcePS);
+	SHADER_USE_PARAMETER_STRUCT(FCopyQuadsMultiSourcePS, FGlobalShader);
+
+	static constexpr int32 NumMultiSources = 63; // This has to match the number of input textures supported by CopyQuadsMultiSourcePS
+
+	BEGIN_SHADER_PARAMETER_STRUCT(FParameters, )
+		SHADER_PARAMETER_STRUCT_REF(FViewUniformShaderParameters, View)
+		SHADER_PARAMETER_RDG_BUFFER_SRV(Buffer<uint2>, InQuadInfos)
+		SHADER_PARAMETER_RDG_TEXTURE_ARRAY(Texture2D, InSourceTexture, [NumMultiSources])
+	END_SHADER_PARAMETER_STRUCT()
+
+
+	class FUseSwizleMask : SHADER_PERMUTATION_BOOL("USE_SWIZZLE_MASK");
+	using FPermutationDomain = TShaderPermutationDomain<FUseSwizleMask>;
+
+	static bool ShouldCompilePermutation(const FGlobalShaderPermutationParameters& InParameters)
+	{
+		return UE::Landscape::DoesPlatformSupportEditLayers(InParameters.Platform);
+	}
+
+	static void ModifyCompilationEnvironment(const FGlobalShaderPermutationParameters& Parameters, FShaderCompilerEnvironment& OutEnvironment)
+	{
+		OutEnvironment.SetDefine(TEXT("COPY_QUADS_MULTISOURCE"), 1);
+	}
+};
+
+IMPLEMENT_GLOBAL_SHADER(FCopyQuadsMultiSourcePS, "/Engine/Private/Landscape/LandscapeEditLayersUtils.usf", "CopyQuadsMultiSourcePS", SF_Pixel);
+
+BEGIN_SHADER_PARAMETER_STRUCT(FCopyQuadsMultiSourcePSParameters, )
+SHADER_PARAMETER_STRUCT_INCLUDE(FPixelShaderUtils::FRasterizeToRectsVS::FParameters, VS)
+SHADER_PARAMETER_STRUCT_INCLUDE(FCopyQuadsMultiSourcePS::FParameters, PS)
+RENDER_TARGET_BINDING_SLOTS()
+END_SHADER_PARAMETER_STRUCT()
+
+
+// ----------------------------------------------------------------------------------
 // Copy texture render command
 
 struct FLandscapeLayersCopyTextureParams
@@ -1515,7 +2025,7 @@
 		}
 		if (InDestTexture != nullptr)
 		{
-			SourceResourceDebugName = InDestTexture->GetName();
+			DestResourceDebugName = InDestTexture->GetName();
 			DestResource = InDestTexture->GetResource();
 		}
 	}
@@ -1544,6 +2054,8 @@
 	uint32 DestArrayIndex = 0;
 	ERHIAccess SourceAccess = ERHIAccess::SRVMask;
 	ERHIAccess DestAccess = ERHIAccess::SRVMask;
+	// There's a shader-version of the copy that is able to swizzle RGBA channels : this mask allows to specify how : 
+	uint8 ChannelSwizzleMask = UE::Landscape::Private::RGBAToRGBASwizzleMask;
 };
 
 class FLandscapeLayersCopyTexture_RenderThread
@@ -1556,6 +2068,24 @@
 	const FLandscapeLayersCopyTextureParams& GetParams() const { return Params; }
 	void Copy(FRHICommandListImmediate& InRHICmdList)
 	{
+		// We must use the PS version if swizzling channels or if the format is different (e.g. R8G8B8A8 to R8)
+		if ((Params.SourceResource->TextureRHI->GetFormat() != Params.DestResource->TextureRHI->GetFormat())
+			|| (Params.ChannelSwizzleMask != UE::Landscape::Private::RGBAToRGBASwizzleMask))
+		{
+			checkf(EnumHasAllFlags(Params.DestResource->TextureRHI->GetFlags(), ETextureCreateFlags::RenderTargetable), TEXT("Cannot request swizzling if the texture is not render-targetable"));
+			checkf(!Params.DestResource->TextureRHI->GetDesc().IsTextureArray() || EnumHasAllFlags(Params.DestResource->TextureRHI->GetFlags(), ETextureCreateFlags::TargetArraySlicesIndependently), TEXT("Cannot request swizzling on a texture array if the slices are not individually render-targetable"));
+			CopyInternalPS(InRHICmdList);
+		}
+		else
+		{
+			CopyInternal(InRHICmdList);
+		}
+	}
+
+private:
+	void CopyInternal(FRHICommandListImmediate& InRHICmdList)
+	{
+		// TODO [jonathan.bard] : make those perf tags optional : with the amount of textures we copy, it slows down texture copies quite a bit : 
 		SCOPED_GPU_STAT(InRHICmdList, LandscapeLayers_CopyTexture);
 		SCOPED_DRAW_EVENTF(InRHICmdList, LandscapeLayers, TEXT("LandscapeLayers_Copy %s -> %s, Mip (%d -> %d), Array Index (%d -> %d)"), Params.SourceResourceDebugName, Params.DestResourceDebugName, Params.SourceMip, Params.DestMip, Params.SourceArrayIndex, Params.DestArrayIndex);
 
@@ -1590,10 +2120,64 @@
 		InRHICmdList.Transition(FRHITransitionInfo(Params.DestResource->TextureRHI, ERHIAccess::CopyDest, Params.DestAccess));
 	}
 
+	void CopyInternalPS(FRHICommandListImmediate& InRHICmdList)
+	{
+		using namespace UE::Landscape::Private;
+
+		const int32 NumChannelsDest = GPixelFormats[Params.DestResource->TextureRHI->GetDesc().Format].NumComponents;
+
+		// TODO [jonathan.bard] : make those perf tags optional : with the amount of textures we copy, it slows down texture copies quite a bit : 
+		SCOPED_GPU_STAT(InRHICmdList, LandscapeLayers_CopyTexturePS);
+		SCOPED_DRAW_EVENTF(InRHICmdList, LandscapeLayers, TEXT("LandscapeLayers_CopyPS %s -> %s, Mip (%d -> %d), Array Index (%d -> %d), [%s]"),
+			Params.SourceResourceDebugName, Params.DestResourceDebugName, Params.SourceMip, Params.DestMip, Params.SourceArrayIndex, Params.DestArrayIndex, GetChannelSwizzleMaskDescription(Params.ChannelSwizzleMask, NumChannelsDest));
+
+		FIntPoint SourceSize(Params.SourceResource->GetSizeX() >> Params.SourceMip, Params.SourceResource->GetSizeY() >> Params.SourceMip);
+		FIntPoint DestSize(Params.DestResource->GetSizeX() >> Params.DestMip, Params.DestResource->GetSizeY() >> Params.DestMip);
+
+		// If CopySize is passed, used that as the size (and don't adjust with the mip level : consider that the user has computed it properly) : 
+		FIntPoint Size;
+		Size.X = (Params.CopySize.X > 0) ? Params.CopySize.X : SourceSize.X;
+		Size.Y = (Params.CopySize.Y > 0) ? Params.CopySize.Y : SourceSize.Y;
+		check((Params.SourceArrayIndex == 0) && (Params.SourceMip == 0) && (Params.DestMip == 0)); // The PS version of copy is not supported on texture arrays and mips for now
+
+		InRHICmdList.Transition(FRHITransitionInfo(Params.SourceResource->TextureRHI, Params.SourceAccess, ERHIAccess::SRVGraphics));
+		InRHICmdList.Transition(FRHITransitionInfo(Params.DestResource->TextureRHI, Params.DestAccess, ERHIAccess::RTV));
+
+		int32 PassArraySlice = Params.DestResource->TextureRHI->GetDesc().IsTextureArray() ? Params.DestArrayIndex : -1; // Little hack to make sure we pass -1 to FRHIRenderPassInfo for a non-texture array resource as that's what it expects :
+		FRHIRenderPassInfo RPInfo(Params.DestResource->TextureRHI, ERenderTargetActions::DontLoad_Store, /*ResolveRT = */nullptr, /*InMipIndex = */0, PassArraySlice);
+		InRHICmdList.BeginRenderPass(RPInfo, TEXT("CopyTexture"));
+
+		FGlobalShaderMap* GlobalShaderMap = GetGlobalShaderMap(GMaxRHIFeatureLevel);
+		TShaderMapRef< FLandscapeCopyTextureVS > VertexShader(GlobalShaderMap);
+		TShaderMapRef< FLandscapeCopyTexturePS > PixelShader(GlobalShaderMap);
+
+		FGraphicsPipelineStateInitializer GraphicsPSOInit;
+		InRHICmdList.ApplyCachedRenderTargets(GraphicsPSOInit);
+		GraphicsPSOInit.BlendState = TStaticBlendState<>::GetRHI();
+		GraphicsPSOInit.RasterizerState = TStaticRasterizerState<>::GetRHI();
+		GraphicsPSOInit.DepthStencilState = TStaticDepthStencilState<false, CF_Always>::GetRHI();
+		GraphicsPSOInit.PrimitiveType = PT_TriangleList;
+		GraphicsPSOInit.BoundShaderState.VertexDeclarationRHI = GetVertexDeclarationFVector4();
+		GraphicsPSOInit.BoundShaderState.VertexShaderRHI = VertexShader.GetVertexShader();
+		GraphicsPSOInit.BoundShaderState.PixelShaderRHI = PixelShader.GetPixelShader();
+		SetGraphicsPipelineState(InRHICmdList, GraphicsPSOInit, 0);
+
+		SetShaderParametersLegacyPS(InRHICmdList, PixelShader, Params.SourceResource->TextureRHI, Params.SourcePosition, Size, Params.ChannelSwizzleMask);
+
+		InRHICmdList.SetViewport((float)Params.DestPosition.X, (float)Params.DestPosition.Y, 0.0f, (float)(Params.DestPosition.X + Size.X), (float)(Params.DestPosition.Y + Size.Y), 1.0f);
+		InRHICmdList.DrawIndexedPrimitive(GTwoTrianglesIndexBuffer.IndexBufferRHI, 0, 0, 4, 0, 2, 1);
+
+		InRHICmdList.EndRenderPass();
+
+		InRHICmdList.Transition(FRHITransitionInfo(Params.SourceResource->TextureRHI, ERHIAccess::SRVGraphics, Params.SourceAccess));
+		InRHICmdList.Transition(FRHITransitionInfo(Params.DestResource->TextureRHI, ERHIAccess::RTV, Params.DestAccess));
+	}
+
 private:
 	FLandscapeLayersCopyTextureParams Params;
 };
 
+// ----------------------------------------------------------------------------------
 // Clear command
 
 class LandscapeLayersWeightmapClear_RenderThread
@@ -1626,6 +2210,7 @@
 	FTextureRenderTargetResource* RenderTargetResource;
 };
 
+// ----------------------------------------------------------------------------------
 // Render command
 
 template<typename ShaderDataType, typename ShaderPixelClass, typename ShaderPixelMipsClass>
@@ -1813,7 +2398,11 @@
 	return true;
 }
 
-bool ALandscape::SupportsEditLayersLocalMerge()
+ELandscapeEditLayersMergeMode ALandscape::GetEditLayersMergeMode()
+{
+	ELandscapeEditLayersMergeMode DesiredMode = static_cast<ELandscapeEditLayersMergeMode>(FMath::Min<uint8>(static_cast<uint8>(LandscapeEditLayersLocalMerge), static_cast<uint8>(ELandscapeEditLayersMergeMode::Invalid) - 1));
+
+	if (DesiredMode == ELandscapeEditLayersMergeMode::LocalMerge)
 {
 	ALandscape* Landscape = GetLandscapeActor();
 	check(Landscape);
@@ -1823,11 +2412,12 @@
 		// No BP brush is supported for local merge of edit layers : 
 		if (Layer.bVisible && Algo::AnyOf(Layer.Brushes, [](const FLandscapeLayerBrush& Brush) { return (Brush.GetBrush() != nullptr) && Brush.GetBrush()->IsVisible(); }))
 		{
-			return false;
+				return ELandscapeEditLayersMergeMode::GlobalMerge;
+			}
 		}
 	}
 
-	return !!LandscapeEditLayersLocalMerge;
+	return DesiredMode;
 }
 
 bool ALandscape::HasNormalCaptureBPBrushLayer()
@@ -1865,7 +2455,7 @@
 	ALandscape* Landscape = GetLandscapeActor();
 	check(Landscape);
 
-	if (SupportsEditLayersLocalMerge())
+	if (GetEditLayersMergeMode() != ELandscapeEditLayersMergeMode::GlobalMerge)
 	{
 		ReleaseLayersRenderingResource();
 
@@ -1896,7 +2486,6 @@
 			}
 		}
 
-		//static bool bWarned = false;
 		if (CurrentMipSizeX > GRHIGlobals.MaxTextureDimensions || CurrentMipSizeY > GRHIGlobals.MaxTextureDimensions)
 		{
 			if (!bWarnedGlobalMergeDimensionsExceeded)
@@ -1916,7 +2505,7 @@
 			Landscape->HeightmapRTList.Init(nullptr, (int32)EHeightmapRTType::HeightmapRT_Count);
 		}
 
-		auto InitOrResizeRT = [](UTextureRenderTarget2D *RT, int32 ResX, int32 ResY, bool bInit)
+		auto InitOrResizeRT = [](UTextureRenderTarget2D* RT, int32 ResX, int32 ResY, bool bInit)
 		{
 			if (bInit)
 			{
@@ -1929,8 +2518,8 @@
 			}
 		};
 
-			for (int32 i = 0; i < (int32)EHeightmapRTType::HeightmapRT_Count; ++i)
-			{
+		for (int32 i = 0; i < (int32)EHeightmapRTType::HeightmapRT_Count; ++i)
+		{
 			if (bCreateFromScratch)
 			{
 				FText DisplayName = StaticEnum<EHeightmapRTType>()->GetDisplayValueAsText((EHeightmapRTType)i);
@@ -1943,31 +2532,31 @@
 				Landscape->HeightmapRTList[i]->ClearColor = FLinearColor(0.0f, 0.0f, 0.0f, 0.0f);
 			}
 
-				if (i < (int32)EHeightmapRTType::HeightmapRT_Mip1) // Landscape size RT
-				{
+			if (i < (int32)EHeightmapRTType::HeightmapRT_Mip1) // Landscape size RT
+			{
 				InitOrResizeRT(Landscape->HeightmapRTList[i], CurrentMipSizeX, CurrentMipSizeY, bCreateFromScratch);
-				}
-				else // Mips
-				{
-					CurrentMipSizeX >>= 1;
-					CurrentMipSizeY >>= 1;
+			}
+			else // Mips
+			{
+				CurrentMipSizeX >>= 1;
+				CurrentMipSizeY >>= 1;
 				InitOrResizeRT(Landscape->HeightmapRTList[i], CurrentMipSizeX, CurrentMipSizeY, bCreateFromScratch);
-				}
-
-				// Only generate required mips RT
-				if (CurrentMipSizeX == ComponentCounts.X && CurrentMipSizeY == ComponentCounts.Y)
-				{
-					break;
-				}
 			}
 
+			// Only generate required mips RT
+			if (CurrentMipSizeX == ComponentCounts.X && CurrentMipSizeY == ComponentCounts.Y)
+			{
+				break;
+			}
+		}
+
 		// Weightmap mip size
 		CurrentMipSizeX = ((SubsectionSizeQuads + 1) * NumSubsections) * ComponentCounts.X;
 		CurrentMipSizeY = ((SubsectionSizeQuads + 1) * NumSubsections) * ComponentCounts.Y;
 		bCreateFromScratch = (Landscape->WeightmapRTList.Num() == 0);
 
 		if (bCreateFromScratch)
-				{
+		{
 			Landscape->WeightmapRTList.Init(nullptr, (int32)EWeightmapRTType::WeightmapRT_Count);
 		}
 
@@ -1991,27 +2580,27 @@
 				{
 					Landscape->WeightmapRTList[i]->RenderTargetFormat = RTF_R8;
 				}
-				}
+			}
 
 			if (i < (int32)EWeightmapRTType::WeightmapRT_Mip0)
-				{
+			{
 				InitOrResizeRT(Landscape->WeightmapRTList[i], CurrentMipSizeX, CurrentMipSizeY, bCreateFromScratch);
-				}
-				else // Mips
-				{
+			}
+			else // Mips
+			{
 				InitOrResizeRT(Landscape->WeightmapRTList[i], CurrentMipSizeX, CurrentMipSizeY, bCreateFromScratch);
 
-					CurrentMipSizeX >>= 1;
-					CurrentMipSizeY >>= 1;
-				}
-
-				// Only generate required mips RT
-				if (CurrentMipSizeX < ComponentCounts.X && CurrentMipSizeY < ComponentCounts.Y)
-				{
-					break;
-				}
+				CurrentMipSizeX >>= 1;
+				CurrentMipSizeY >>= 1;
 			}
 
+			// Only generate required mips RT
+			if (CurrentMipSizeX < ComponentCounts.X && CurrentMipSizeY < ComponentCounts.Y)
+			{
+				break;
+			}
+		}
+
 		InitializeLayersWeightmapResources();
 	}
 	return true;
@@ -2045,3 +2634,303 @@
 	if (Info == nullptr)
 	{
 		return;
+	}
+
+	Info->ForEachLandscapeProxy([&](ALandscapeProxy* Proxy)
+	{
+		for (auto& ItPair : Proxy->HeightmapsCPUReadback)
+		{
+			FLandscapeEditLayerReadback* HeightmapCPUReadback = ItPair.Value;
+			delete HeightmapCPUReadback;
+		}
+		Proxy->HeightmapsCPUReadback.Empty();
+
+		for (auto& ItPair : Proxy->WeightmapsCPUReadback)
+		{
+			FLandscapeEditLayerReadback* WeightmapCPUReadback = ItPair.Value;
+			delete WeightmapCPUReadback;
+		}
+		Proxy->WeightmapsCPUReadback.Empty();
+		return true;
+	});
+
+	if (CombinedLayersWeightmapAllMaterialLayersResource != nullptr)
+	{
+		BeginReleaseResource(CombinedLayersWeightmapAllMaterialLayersResource);
+	}
+
+	if (CurrentLayersWeightmapAllMaterialLayersResource != nullptr)
+	{
+		BeginReleaseResource(CurrentLayersWeightmapAllMaterialLayersResource);
+	}
+
+	if (WeightmapScratchExtractLayerTextureResource != nullptr)
+	{
+		BeginReleaseResource(WeightmapScratchExtractLayerTextureResource);
+	}
+
+	if (WeightmapScratchPackLayerTextureResource != nullptr)
+	{
+		BeginReleaseResource(WeightmapScratchPackLayerTextureResource);
+	}
+
+	{
+		TRACE_CPUPROFILER_EVENT_SCOPE(LandscapeLayers_Flush_ResourceRelease);
+		FlushRenderingCommands();
+	}
+
+	delete CombinedLayersWeightmapAllMaterialLayersResource;
+	delete CurrentLayersWeightmapAllMaterialLayersResource;
+	delete WeightmapScratchExtractLayerTextureResource;
+	delete WeightmapScratchPackLayerTextureResource;
+
+	CombinedLayersWeightmapAllMaterialLayersResource = nullptr;
+	CurrentLayersWeightmapAllMaterialLayersResource = nullptr;
+	WeightmapScratchExtractLayerTextureResource = nullptr;
+	WeightmapScratchPackLayerTextureResource = nullptr;
+}
+
+FIntPoint ALandscape::ComputeComponentCounts() const
+{
+	ULandscapeInfo* Info = GetLandscapeInfo();
+
+	if (Info == nullptr)
+	{
+		return FIntPoint(INDEX_NONE, INDEX_NONE);
+	}
+
+	FIntPoint NumComponents(0, 0);
+	FIntPoint MaxSectionBase(TNumericLimits<int32>::Min(), TNumericLimits<int32>::Min());
+	FIntPoint MinSectionBase(TNumericLimits<int32>::Max(), TNumericLimits<int32>::Max());
+
+	Info->ForEachLandscapeProxy([&MaxSectionBase, &MinSectionBase](ALandscapeProxy* Proxy)
+	{
+		for (ULandscapeComponent* Component : Proxy->LandscapeComponents)
+		{
+			MaxSectionBase.X = FMath::Max(MaxSectionBase.X, Component->SectionBaseX);
+			MaxSectionBase.Y = FMath::Max(MaxSectionBase.Y, Component->SectionBaseY);
+
+			MinSectionBase.X = FMath::Min(MinSectionBase.X, Component->SectionBaseX);
+			MinSectionBase.Y = FMath::Min(MinSectionBase.Y, Component->SectionBaseY);
+		}
+		return true;
+	});
+
+	if ((MaxSectionBase.X >= MinSectionBase.X) && (MaxSectionBase.Y >= MinSectionBase.Y))
+	{
+		NumComponents.X = ((MaxSectionBase.X - MinSectionBase.X) / ComponentSizeQuads) + 1;
+		NumComponents.Y = ((MaxSectionBase.Y - MinSectionBase.Y) / ComponentSizeQuads) + 1;
+	}
+
+	return NumComponents;
+}
+
+void ALandscape::CopyOldDataToDefaultLayer()
+{
+	ULandscapeInfo* Info = GetLandscapeInfo();
+
+	if (Info == nullptr)
+	{
+		return;
+	}
+
+	Info->ForEachLandscapeProxy([&](ALandscapeProxy* Proxy)
+	{
+		CopyOldDataToDefaultLayer(Proxy);
+		return true;
+	});
+}
+
+void ALandscape::CopyOldDataToDefaultLayer(ALandscapeProxy* InProxy)
+{
+	ULandscapeInfo* Info = GetLandscapeInfo();
+
+	if (Info == nullptr)
+	{
+		return;
+	}
+
+	if (InProxy->LandscapeComponents.IsEmpty())
+	{
+		// No data to migrate, we can early-out to avoid modifying the proxy:
+		return;
+	}
+
+	InProxy->Modify();
+
+	const FLandscapeLayer* DefaultLayer = GetLayerConst(0);
+	check(DefaultLayer != nullptr);
+
+	struct FWeightmapTextureData
+	{
+		UTexture2D* Texture;
+		ULandscapeWeightmapUsage* Usage;
+	};
+
+	TMap<UTexture2D*, FWeightmapTextureData> ProcessedWeightmaps;
+	TArray<UTexture2D*> ProcessedHeightmaps;
+	TArray<ULandscapeComponent*> WeightmapsComponentsToCleanup;
+
+	for (ULandscapeComponent* Component : InProxy->LandscapeComponents)
+	{
+		FLandscapeLayerComponentData* LayerData = Component->GetLayerData(DefaultLayer->Guid);
+
+		if (ensure(LayerData != nullptr && LayerData->IsInitialized()))
+		{
+			// Heightmap
+			UTexture2D* ComponentHeightmap = Component->GetHeightmap();
+
+			if (!ProcessedHeightmaps.Contains(ComponentHeightmap))
+			{
+				ProcessedHeightmaps.Add(ComponentHeightmap);
+
+				UTexture* DefaultLayerHeightmap = LayerData->HeightmapData.Texture;
+				check(DefaultLayerHeightmap != nullptr);
+
+				// Only copy Mip0 as other mips will get regenerated
+				TArray64<uint8> ExistingMip0Data;
+				ComponentHeightmap->Source.GetMipData(ExistingMip0Data, 0);
+
+				// Calling modify here makes sure that async texture compilation finishes so we can Lock the mip
+				DefaultLayerHeightmap->Modify();
+				FColor* Mip0Data = (FColor*)DefaultLayerHeightmap->Source.LockMip(0);
+				FMemory::Memcpy(Mip0Data, ExistingMip0Data.GetData(), ExistingMip0Data.Num());
+				DefaultLayerHeightmap->Source.UnlockMip(0);
+
+				DefaultLayerHeightmap->UpdateResource();
+			}
+
+			// Weightmaps
+			WeightmapsComponentsToCleanup.Add(Component);
+
+			const TArray<UTexture2D*>& ComponentWeightmapTextures = Component->GetWeightmapTextures();
+			const TArray<FWeightmapLayerAllocationInfo>& ComponentLayerAllocations = Component->GetWeightmapLayerAllocations();
+			const TArray<ULandscapeWeightmapUsage*>& ComponentWeightmapTexturesUsage = Component->GetWeightmapTexturesUsage();
+
+			LayerData->WeightmapData.Textures.AddDefaulted(ComponentWeightmapTextures.Num());
+			LayerData->WeightmapData.TextureUsages.AddDefaulted(ComponentWeightmapTexturesUsage.Num());
+
+			for (int32 TextureIndex = 0; TextureIndex < ComponentWeightmapTextures.Num(); ++TextureIndex)
+			{
+				UTexture2D* ComponentWeightmap = ComponentWeightmapTextures[TextureIndex];
+				const FWeightmapTextureData* WeightmapTextureData = ProcessedWeightmaps.Find(ComponentWeightmap);
+
+				if (WeightmapTextureData != nullptr)
+				{
+					LayerData->WeightmapData.Textures[TextureIndex] = WeightmapTextureData->Texture;
+					LayerData->WeightmapData.TextureUsages[TextureIndex] = WeightmapTextureData->Usage;
+					check(WeightmapTextureData->Usage->LayerGuid == DefaultLayer->Guid);
+
+					for (int32 ChannelIndex = 0; ChannelIndex < ULandscapeWeightmapUsage::NumChannels; ++ChannelIndex)
+					{
+						const ULandscapeComponent* ChannelLandscapeComponent = LayerData->WeightmapData.TextureUsages[TextureIndex]->ChannelUsage[ChannelIndex];
+
+						if (ChannelLandscapeComponent != nullptr && ChannelLandscapeComponent == Component)
+						{
+							for (const FWeightmapLayerAllocationInfo& Allocation : ComponentLayerAllocations)
+							{
+								if (Allocation.WeightmapTextureIndex == TextureIndex)
+								{
+									LayerData->WeightmapData.LayerAllocations.Add(Allocation);
+								}
+							}
+
+							break;
+						}
+					}
+				}
+				else
+				{
+					// No need for mip chain on edit layers : 
+					UTexture2D* NewLayerWeightmapTexture = InProxy->CreateLandscapeTexture(ComponentWeightmap->Source.GetSizeX(), ComponentWeightmap->Source.GetSizeY(), TEXTUREGROUP_Terrain_Weightmap, 
+						ComponentWeightmap->Source.GetFormat(), /*OptionalOverrideOuter = */ nullptr, /*bCompress = */false, /*bMipChain = */false);
+
+					// Only copy Mip0 as other mips will get regenerated
+					TArray64<uint8> ExistingMip0Data;
+					ComponentWeightmap->Source.GetMipData(ExistingMip0Data, 0);
+
+					FColor* Mip0Data = (FColor*)NewLayerWeightmapTexture->Source.LockMip(0);
+					FMemory::Memcpy(Mip0Data, ExistingMip0Data.GetData(), ExistingMip0Data.Num());
+					NewLayerWeightmapTexture->Source.UnlockMip(0);
+
+					LayerData->WeightmapData.Textures[TextureIndex] = NewLayerWeightmapTexture;
+					LayerData->WeightmapData.TextureUsages[TextureIndex] = InProxy->WeightmapUsageMap.Add(NewLayerWeightmapTexture, InProxy->CreateWeightmapUsage());
+
+					for (int32 ChannelIndex = 0; ChannelIndex < ULandscapeWeightmapUsage::NumChannels; ++ChannelIndex)
+					{
+						LayerData->WeightmapData.TextureUsages[TextureIndex]->ChannelUsage[ChannelIndex] = ComponentWeightmapTexturesUsage[TextureIndex]->ChannelUsage[ChannelIndex];
+					}
+
+					LayerData->WeightmapData.TextureUsages[TextureIndex]->LayerGuid = DefaultLayer->Guid;
+
+					// Create new Usage for the "final" layer as the other one will now be used by the Default layer
+					for (const FWeightmapLayerAllocationInfo& Allocation : ComponentLayerAllocations)
+					{
+						if (Allocation.WeightmapTextureIndex == TextureIndex)
+						{
+							LayerData->WeightmapData.LayerAllocations.Add(Allocation);
+						}
+					}
+
+					FWeightmapTextureData NewTextureData;
+					NewTextureData.Texture = NewLayerWeightmapTexture;
+					NewTextureData.Usage = LayerData->WeightmapData.TextureUsages[TextureIndex];
+
+					ProcessedWeightmaps.Add(ComponentWeightmap, NewTextureData);
+
+					NewLayerWeightmapTexture->UpdateResource();
+				}
+			}
+		}
+	}
+
+	for (ULandscapeComponent* Component : WeightmapsComponentsToCleanup)
+	{
+		TArray<FWeightmapLayerAllocationInfo>& ComponentLayerAllocations = Component->GetWeightmapLayerAllocations();
+
+		for (FWeightmapLayerAllocationInfo& Allocation : ComponentLayerAllocations)
+		{
+			Allocation.Free();
+		}
+	}
+}
+
+void ALandscape::UpdateProxyLayersWeightmapUsage()
+{
+	ULandscapeInfo* Info = GetLandscapeInfo();
+
+	if (Info == nullptr)
+	{
+		return;
+	}
+
+	Info->ForEachLandscapeProxy([&](ALandscapeProxy* Proxy)
+	{
+		Proxy->UpdateProxyLayersWeightmapUsage();
+		return true;
+	});
+}
+
+void ALandscapeProxy::UpdateProxyLayersWeightmapUsage()
+{
+	if (bNeedsWeightmapUsagesUpdate)
+	{
+		InitializeProxyLayersWeightmapUsage();
+	}
+	check(!bNeedsWeightmapUsagesUpdate);
+}
+
+void ALandscapeProxy::PostEditUndo()
+{
+	check(ULandscapeComponent::UndoRedoModifiedComponentCount == 0);
+
+	Super::PostEditUndo();
+}
+
+void ALandscape::InitializeLandscapeLayersWeightmapUsage()
+{
+	ULandscapeInfo* Info = GetLandscapeInfo();
+
+	if (Info == nullptr)
+	{
+		return;
@@ -2049,306 +2938,6 @@
 
 	Info->ForEachLandscapeProxy([&](ALandscapeProxy* Proxy)
 	{
-		for (auto& ItPair : Proxy->HeightmapsCPUReadback)
-		{
-			FLandscapeEditLayerReadback* HeightmapCPUReadback = ItPair.Value;
-			delete HeightmapCPUReadback;
-		}
-		Proxy->HeightmapsCPUReadback.Empty();
-
-		for (auto& ItPair : Proxy->WeightmapsCPUReadback)
-		{
-			FLandscapeEditLayerReadback* WeightmapCPUReadback = ItPair.Value;
-			delete WeightmapCPUReadback;
-		}
-		Proxy->WeightmapsCPUReadback.Empty();
-		return true;
-	});
-
-	if (CombinedLayersWeightmapAllMaterialLayersResource != nullptr)
-	{
-		BeginReleaseResource(CombinedLayersWeightmapAllMaterialLayersResource);
-	}
-
-	if (CurrentLayersWeightmapAllMaterialLayersResource != nullptr)
-	{
-		BeginReleaseResource(CurrentLayersWeightmapAllMaterialLayersResource);
-	}
-
-	if (WeightmapScratchExtractLayerTextureResource != nullptr)
-	{
-		BeginReleaseResource(WeightmapScratchExtractLayerTextureResource);
-	}
-
-	if (WeightmapScratchPackLayerTextureResource != nullptr)
-	{
-		BeginReleaseResource(WeightmapScratchPackLayerTextureResource);
-	}
-
-	{
-		TRACE_CPUPROFILER_EVENT_SCOPE(LandscapeLayers_Flush_ResourceRelease);
-		FlushRenderingCommands();
-	}
-
-	delete CombinedLayersWeightmapAllMaterialLayersResource;
-	delete CurrentLayersWeightmapAllMaterialLayersResource;
-	delete WeightmapScratchExtractLayerTextureResource;
-	delete WeightmapScratchPackLayerTextureResource;
-
-	CombinedLayersWeightmapAllMaterialLayersResource = nullptr;
-	CurrentLayersWeightmapAllMaterialLayersResource = nullptr;
-	WeightmapScratchExtractLayerTextureResource = nullptr;
-	WeightmapScratchPackLayerTextureResource = nullptr;
-}
-
-FIntPoint ALandscape::ComputeComponentCounts() const
-{
-	ULandscapeInfo* Info = GetLandscapeInfo();
-
-	if (Info == nullptr)
-	{
-		return FIntPoint(INDEX_NONE, INDEX_NONE);
-	}
-
-	FIntPoint NumComponents(0, 0);
-	FIntPoint MaxSectionBase(TNumericLimits<int32>::Min(), TNumericLimits<int32>::Min());
-	FIntPoint MinSectionBase(TNumericLimits<int32>::Max(), TNumericLimits<int32>::Max());
-
-	Info->ForEachLandscapeProxy([&MaxSectionBase, &MinSectionBase](ALandscapeProxy* Proxy)
-	{
-		for (ULandscapeComponent* Component : Proxy->LandscapeComponents)
-		{
-			MaxSectionBase.X = FMath::Max(MaxSectionBase.X, Component->SectionBaseX);
-			MaxSectionBase.Y = FMath::Max(MaxSectionBase.Y, Component->SectionBaseY);
-
-			MinSectionBase.X = FMath::Min(MinSectionBase.X, Component->SectionBaseX);
-			MinSectionBase.Y = FMath::Min(MinSectionBase.Y, Component->SectionBaseY);
-		}
-		return true;
-	});
-
-	if ((MaxSectionBase.X >= MinSectionBase.X) && (MaxSectionBase.Y >= MinSectionBase.Y))
-	{
-		NumComponents.X = ((MaxSectionBase.X - MinSectionBase.X) / ComponentSizeQuads) + 1;
-		NumComponents.Y = ((MaxSectionBase.Y - MinSectionBase.Y) / ComponentSizeQuads) + 1;
-	}
-
-	return NumComponents;
-}
-
-void ALandscape::CopyOldDataToDefaultLayer()
-{
-	ULandscapeInfo* Info = GetLandscapeInfo();
-
-	if (Info == nullptr)
-	{
-		return;
-	}
-
-	Info->ForEachLandscapeProxy([&](ALandscapeProxy* Proxy)
-	{
-		CopyOldDataToDefaultLayer(Proxy);
-		return true;
-	});
-}
-
-void ALandscape::CopyOldDataToDefaultLayer(ALandscapeProxy* InProxy)
-{
-	ULandscapeInfo* Info = GetLandscapeInfo();
-
-	if (Info == nullptr)
-	{
-		return;
-	}
-
-	if (InProxy->LandscapeComponents.IsEmpty())
-	{
-		// No data to migrate, we can early-out to avoid modifying the proxy:
-		return;
-	}
-
-	InProxy->Modify();
-
-	const FLandscapeLayer* DefaultLayer = GetLayerConst(0);
-	check(DefaultLayer != nullptr);
-
-	struct FWeightmapTextureData
-	{
-		UTexture2D* Texture;
-		ULandscapeWeightmapUsage* Usage;
-	};
-
-	TMap<UTexture2D*, FWeightmapTextureData> ProcessedWeightmaps;
-	TArray<UTexture2D*> ProcessedHeightmaps;
-	TArray<ULandscapeComponent*> WeightmapsComponentsToCleanup;
-
-	for (ULandscapeComponent* Component : InProxy->LandscapeComponents)
-	{
-		FLandscapeLayerComponentData* LayerData = Component->GetLayerData(DefaultLayer->Guid);
-
-		if (ensure(LayerData != nullptr && LayerData->IsInitialized()))
-		{
-			// Heightmap
-			UTexture2D* ComponentHeightmap = Component->GetHeightmap();
-
-			if (!ProcessedHeightmaps.Contains(ComponentHeightmap))
-			{
-				ProcessedHeightmaps.Add(ComponentHeightmap);
-
-				UTexture* DefaultLayerHeightmap = LayerData->HeightmapData.Texture;
-				check(DefaultLayerHeightmap != nullptr);
-
-				// Only copy Mip0 as other mips will get regenerated
-				TArray64<uint8> ExistingMip0Data;
-				ComponentHeightmap->Source.GetMipData(ExistingMip0Data, 0);
-
-				// Calling modify here makes sure that async texture compilation finishes so we can Lock the mip
-				DefaultLayerHeightmap->Modify();
-				FColor* Mip0Data = (FColor*)DefaultLayerHeightmap->Source.LockMip(0);
-				FMemory::Memcpy(Mip0Data, ExistingMip0Data.GetData(), ExistingMip0Data.Num());
-				DefaultLayerHeightmap->Source.UnlockMip(0);
-
-				DefaultLayerHeightmap->UpdateResource();
-			}
-
-			// Weightmaps
-			WeightmapsComponentsToCleanup.Add(Component);
-
-			const TArray<UTexture2D*>& ComponentWeightmapTextures = Component->GetWeightmapTextures();
-			const TArray<FWeightmapLayerAllocationInfo>& ComponentLayerAllocations = Component->GetWeightmapLayerAllocations();
-			const TArray<ULandscapeWeightmapUsage*>& ComponentWeightmapTexturesUsage = Component->GetWeightmapTexturesUsage();
-
-			LayerData->WeightmapData.Textures.AddDefaulted(ComponentWeightmapTextures.Num());
-			LayerData->WeightmapData.TextureUsages.AddDefaulted(ComponentWeightmapTexturesUsage.Num());
-
-			for (int32 TextureIndex = 0; TextureIndex < ComponentWeightmapTextures.Num(); ++TextureIndex)
-			{
-				UTexture2D* ComponentWeightmap = ComponentWeightmapTextures[TextureIndex];
-				const FWeightmapTextureData* WeightmapTextureData = ProcessedWeightmaps.Find(ComponentWeightmap);
-
-				if (WeightmapTextureData != nullptr)
-				{
-					LayerData->WeightmapData.Textures[TextureIndex] = WeightmapTextureData->Texture;
-					LayerData->WeightmapData.TextureUsages[TextureIndex] = WeightmapTextureData->Usage;
-					check(WeightmapTextureData->Usage->LayerGuid == DefaultLayer->Guid);
-
-					for (int32 ChannelIndex = 0; ChannelIndex < ULandscapeWeightmapUsage::NumChannels; ++ChannelIndex)
-					{
-						const ULandscapeComponent* ChannelLandscapeComponent = LayerData->WeightmapData.TextureUsages[TextureIndex]->ChannelUsage[ChannelIndex];
-
-						if (ChannelLandscapeComponent != nullptr && ChannelLandscapeComponent == Component)
-						{
-							for (const FWeightmapLayerAllocationInfo& Allocation : ComponentLayerAllocations)
-							{
-								if (Allocation.WeightmapTextureIndex == TextureIndex)
-								{
-									LayerData->WeightmapData.LayerAllocations.Add(Allocation);
-								}
-							}
-
-							break;
-						}
-					}
-				}
-				else
-				{
-					// No need for mip chain on edit layers : 
-					UTexture2D* NewLayerWeightmapTexture = InProxy->CreateLandscapeTexture(ComponentWeightmap->Source.GetSizeX(), ComponentWeightmap->Source.GetSizeY(), TEXTUREGROUP_Terrain_Weightmap, 
-						ComponentWeightmap->Source.GetFormat(), /*OptionalOverrideOuter = */ nullptr, /*bCompress = */false, /*bMipChain = */false);
-
-					// Only copy Mip0 as other mips will get regenerated
-					TArray64<uint8> ExistingMip0Data;
-					ComponentWeightmap->Source.GetMipData(ExistingMip0Data, 0);
-
-					FColor* Mip0Data = (FColor*)NewLayerWeightmapTexture->Source.LockMip(0);
-					FMemory::Memcpy(Mip0Data, ExistingMip0Data.GetData(), ExistingMip0Data.Num());
-					NewLayerWeightmapTexture->Source.UnlockMip(0);
-
-					LayerData->WeightmapData.Textures[TextureIndex] = NewLayerWeightmapTexture;
-					LayerData->WeightmapData.TextureUsages[TextureIndex] = InProxy->WeightmapUsageMap.Add(NewLayerWeightmapTexture, InProxy->CreateWeightmapUsage());
-
-					for (int32 ChannelIndex = 0; ChannelIndex < ULandscapeWeightmapUsage::NumChannels; ++ChannelIndex)
-					{
-						LayerData->WeightmapData.TextureUsages[TextureIndex]->ChannelUsage[ChannelIndex] = ComponentWeightmapTexturesUsage[TextureIndex]->ChannelUsage[ChannelIndex];
-					}
-
-					LayerData->WeightmapData.TextureUsages[TextureIndex]->LayerGuid = DefaultLayer->Guid;
-
-					// Create new Usage for the "final" layer as the other one will now be used by the Default layer
-					for (const FWeightmapLayerAllocationInfo& Allocation : ComponentLayerAllocations)
-					{
-						if (Allocation.WeightmapTextureIndex == TextureIndex)
-						{
-							LayerData->WeightmapData.LayerAllocations.Add(Allocation);
-						}
-					}
-
-					FWeightmapTextureData NewTextureData;
-					NewTextureData.Texture = NewLayerWeightmapTexture;
-					NewTextureData.Usage = LayerData->WeightmapData.TextureUsages[TextureIndex];
-
-					ProcessedWeightmaps.Add(ComponentWeightmap, NewTextureData);
-
-					NewLayerWeightmapTexture->UpdateResource();
-				}
-			}
-		}
-	}
-
-	for (ULandscapeComponent* Component : WeightmapsComponentsToCleanup)
-	{
-		TArray<FWeightmapLayerAllocationInfo>& ComponentLayerAllocations = Component->GetWeightmapLayerAllocations();
-
-		for (FWeightmapLayerAllocationInfo& Allocation : ComponentLayerAllocations)
-		{
-			Allocation.Free();
-		}
-	}
-}
-
-void ALandscape::UpdateProxyLayersWeightmapUsage()
-{
-	ULandscapeInfo* Info = GetLandscapeInfo();
-
-	if (Info == nullptr)
-	{
-		return;
-	}
-
-	Info->ForEachLandscapeProxy([&](ALandscapeProxy* Proxy)
-	{
-		Proxy->UpdateProxyLayersWeightmapUsage();
-		return true;
-	});
-}
-
-void ALandscapeProxy::UpdateProxyLayersWeightmapUsage()
-{
-	if (bNeedsWeightmapUsagesUpdate)
-	{
-		InitializeProxyLayersWeightmapUsage();
-	}
-	check(!bNeedsWeightmapUsagesUpdate);
-}
-
-void ALandscapeProxy::PostEditUndo()
-{
-	check(ULandscapeComponent::UndoRedoModifiedComponentCount == 0);
-
-	Super::PostEditUndo();
-}
-
-void ALandscape::InitializeLandscapeLayersWeightmapUsage()
-{
-	ULandscapeInfo* Info = GetLandscapeInfo();
-
-	if (Info == nullptr)
-	{
-		return;
-	}
-
-	Info->ForEachLandscapeProxy([&](ALandscapeProxy* Proxy)
-	{
 		Proxy->InitializeProxyLayersWeightmapUsage();
 		return true;
 	});
@@ -2613,6 +3202,7 @@
 
 void ExecuteCopyToReadbackTexture(TArray<FLandscapeLayersCopyReadbackTextureParams>& InParams)
 {
+	TRACE_CPUPROFILER_EVENT_SCOPE(ExecuteCopyToReadbackTexture);
 	SCOPED_DRAW_EVENTF_GAMETHREAD(LandscapeLayers, TEXT("Copy to readback textures (%d copies)"), InParams.Num());
 	if (!FApp::CanEverRender())
 	{
@@ -2643,7 +3233,7 @@
 		// Init the CPU read back contexts for all components dependent on this texture. This includes a context containing the current component states : 
 		for (ULandscapeComponent* ComponentToResolve : *Components)
 		{
-			const FIntPoint ComponentToResolveKey = ComponentToResolve->GetSectionBase() / ComponentToResolve->ComponentSizeQuads;
+			const FIntPoint ComponentToResolveKey = ComponentToResolve->GetComponentKey();
 			const int32 ComponentToResolveFlags = ComponentToResolve->GetLayerUpdateFlagPerMode();
 			FLandscapeEditLayerReadback::FPerChannelLayerNames PerChannelLayerNames;
 
@@ -3257,7 +3847,7 @@
 		}
 	}
 
-	UE_LOG(LogLandscapeBP, Display, TEXT("Context: %s"), *InContext);
+	UE_LOG(LogLandscape, Display, TEXT("Context: %s"), *InContext);
 
 	int32 MipSize = ((SubsectionSizeQuads + 1) >> InMipRender);
 
@@ -3292,15 +3882,15 @@
 
 		if (Y > 0 && MipSize > 0 && Y % MipSize == 0)
 		{
-			UE_LOG(LogLandscapeBP, Display, TEXT(""));
+			UE_LOG(LogLandscape, Display, TEXT(""));
 		}
 
-		UE_LOG(LogLandscapeBP, Display, TEXT("%s"), *HeightmapHeightOutput);
+		UE_LOG(LogLandscape, Display, TEXT("%s"), *HeightmapHeightOutput);
 	}
 
 	if (InOutputNormals)
 	{
-		UE_LOG(LogLandscapeBP, Display, TEXT(""));
+		UE_LOG(LogLandscape, Display, TEXT(""));
 
 		for (int32 Y = 0; Y < InDataSize.Y; ++Y)
 		{
@@ -3320,10 +3910,10 @@
 
 			if (Y > 0 && MipSize > 0 && Y % MipSize == 0)
 			{
-				UE_LOG(LogLandscapeBP, Display, TEXT(""));
+				UE_LOG(LogLandscape, Display, TEXT(""));
 			}
 
-			UE_LOG(LogLandscapeBP, Display, TEXT("%s"), *HeightmapNormaltOutput);
+			UE_LOG(LogLandscape, Display, TEXT("%s"), *HeightmapNormaltOutput);
 		}
 	}
 }
@@ -3337,7 +3927,7 @@
 		return;
 	}
 
-	UE_LOG(LogLandscapeBP, Display, TEXT("Context: %s"), *InContext);
+	UE_LOG(LogLandscape, Display, TEXT("Context: %s"), *InContext);
 
 	int32 MipSize = ((SubsectionSizeQuads + 1) >> InMipRender);
 
@@ -3359,10 +3949,10 @@
 
 		if (Y > 0 && MipSize > 0 && Y % MipSize == 0)
 		{
-			UE_LOG(LogLandscapeBP, Display, TEXT(""));
+			UE_LOG(LogLandscape, Display, TEXT(""));
 		}
 
-		UE_LOG(LogLandscapeBP, Display, TEXT("%s"), *WeightmapOutput);
+		UE_LOG(LogLandscape, Display, TEXT("%s"), *WeightmapOutput);
 	}
 }
 
@@ -3728,33 +4318,6 @@
 	return bIsReady;
 }
 
-// Must match EEditLayerHeightmapBlendMode in LandscapeLayersHeightmapsPS.usf
-enum class ELandscapeEditLayerHeightmapBlendMode : uint32
-{
-	Additive = 0,
-	AlphaBlend,
-
-	Num,
-};
-
-// Must match EEditLayerWeightmapBlendMode in LandscapeLayersWeightmapsPS.usf
-enum class ELandscapeEditLayerWeightmapBlendMode : uint32
-{
-	Additive = 0,
-	Substractive,
-
-	Num,
-};
-
-// Must match EWeightmapPaintLayerFlags in LandscapeLayersWeightmapsPS.usf
-enum class EWeightmapPaintLayerFlags : uint32
-{
-	IsVisibilityLayer = (1 << 0), // This paint layer is the visibility layer
-	IsWeightBlended = (1 << 1), // Blend the paint layer's value with all the other paint layers weights
-
-	None = 0
-};
-
 // Must match FEditLayerHeightmapMergeInfo in LandscapeLayersHeightmapsPS.usf
 struct FLandscapeEditLayerHeightmapMergeInfo
 {
@@ -3773,14 +4336,15 @@
 {
 	uint32 SourceWeightmapTextureIndex = (uint32)INDEX_NONE; // The index in InPackedWeightmaps of the texture to read from for this layer
 	uint32 SourceWeightmapTextureChannel = (uint32)INDEX_NONE; // The channel of the texture to read from for this layer
-	ELandscapeEditLayerWeightmapBlendMode BlendMode = ELandscapeEditLayerWeightmapBlendMode::Num; // How this layer blends with the previous ones in the layers stack
+	UE::Landscape::EditLayers::EWeightmapBlendMode BlendMode = UE::Landscape::EditLayers::EWeightmapBlendMode::Num; // How this layer blends with the previous ones in the layers stack
 	float Alpha = 1.0f; // Alpha value to be used in the blend
 };
 
-// Must match FWeightmapPaintLayerInfo in LandscapeLayersWeightmapsPS.usf
-struct FLandscapeWeightmapPaintLayerInfo
+// Must match FPerEditLayerWeightmapPaintLayerInfo in LandscapeEditLayersWeightmaps.usf
+// Additional info about this paint layer on this edit layer
+struct FLandscapePerEditLayerWeightmapPaintLayerInfo
 {
-	EWeightmapPaintLayerFlags Flags = EWeightmapPaintLayerFlags::None; // Additional info about this paint layer
+	UE::Landscape::EditLayers::EWeightmapBlendMode Flags = UE::Landscape::EditLayers::EWeightmapBlendMode::Num; // How this layer blends with the previous ones in the layers stack
 };
 
 // Struct that contains all the information relevant for the edit layers update operation (list of dirty components, heightmaps, weightmaps, etc.
@@ -4032,7 +4596,7 @@
 };
 ENUM_CLASS_FLAGS(FUpdateLayersContentContext::ERefreshFlags);
 
-// Little struct that holds information common to PerformLayersHeightmapsLocalMerge and PerformLayersHeightmapsGlobalMerge
+// Little struct that holds information common to PerformLayersHeightmapsLocalMerge, PerformLayersHeightmapsBatchedMerge and PerformLayersHeightmapsGlobalMerge
 struct FEditLayersHeightmapMergeParams
 {
 	int32 HeightmapUpdateModes;
@@ -4189,21 +4753,6 @@
 		TArray<FLandscapeLayersCopyReadbackTextureParams> DeferredCopyReadbackTextures;
 	};
 
-	static ELandscapeEditLayerHeightmapBlendMode LandscapeBlendModeToEditLayerBlendMode(ELandscapeBlendMode InLandscapeBlendMode)
-	{
-		switch (InLandscapeBlendMode)
-		{
-		case LSBM_AdditiveBlend:
-			return ELandscapeEditLayerHeightmapBlendMode::Additive;
-		case LSBM_AlphaBlend:
-			return ELandscapeEditLayerHeightmapBlendMode::AlphaBlend;
-		default:
-			check(false);
-		}
-
-		return ELandscapeEditLayerHeightmapBlendMode::Num;
-	}
-
 	struct FRDGResources
 	{
 		// Texture array in which all possible edit layers heightmaps can fit : we copy the edit layers heightmaps there in order to dynamically access it in the merge shader : 
@@ -4531,6 +5080,7 @@
 void ALandscape::PrepareLayersHeightmapsLocalMergeRenderThreadData(const FUpdateLayersContentContext& InUpdateLayersContentContext, const FEditLayersHeightmapMergeParams& InMergeParams, EditLayersHeightmapLocalMerge_RenderThread::FMergeInfo& OutRenderThreadData)
 {
 	using namespace EditLayersHeightmapLocalMerge_RenderThread;
+	using namespace UE::Landscape::Private;
 
 	TRACE_CPUPROFILER_EVENT_SCOPE(LandscapeLayers_PrepareLayersHeightmapsLocalMergeRenderThreadData);
 
@@ -4812,4 +5362,1770 @@
 	return InMergeParams.HeightmapUpdateModes;
 }
 
+TArray<UE::Landscape::EditLayers::FEditLayerRendererState> ALandscape::GetEditLayerRendererStates(const ULandscapeInfo* InLandscapeInfo, bool bInSkipBrush)
+{
+	using namespace UE::Landscape::EditLayers;
+
+	TArray<FEditLayerRendererState> RendererStates;
+	for (FLandscapeLayer& Layer : LandscapeEditLayers)
+	{
+		RendererStates.Append(Layer.GetEditLayerRendererStates(InLandscapeInfo, bInSkipBrush));
+	}
+	return RendererStates;
+}
+
+TArray<UE::Landscape::EditLayers::FEditLayerRendererState> FLandscapeLayer::GetEditLayerRendererStates(const ULandscapeInfo* InLandscapeInfo, bool bInSkipBrush)
+{
+	using namespace UE::Landscape::EditLayers;
+
+	// Gather all elements that can render some edit layers data, ordered: 
+	TArray<FEditLayerRendererState> RendererStates;
+	RendererStates.Reserve(1 + Brushes.Num());
+
+	// The edit layer itself is a renderer : 
+	check(EditLayer != nullptr);
+	RendererStates.Emplace(EditLayer, InLandscapeInfo);
+
+	for (FLandscapeLayerBrush& Brush : Brushes)
+	{
+		RendererStates.Append(Brush.GetEditLayerRendererStates(InLandscapeInfo, bInSkipBrush));
+	}
+	return RendererStates;
+}
+
+TArray<UE::Landscape::EditLayers::FEditLayerRendererState> FLandscapeLayerBrush::GetEditLayerRendererStates(const ULandscapeInfo* InLandscapeInfo, bool bInSkipBrush)
+{
+	using namespace UE::Landscape::EditLayers;
+
+	// Gather all elements that can render some edit layers data, ordered: 
+	if (BlueprintBrush == nullptr)
+	{
+		return {};
+	}
+
+	return BlueprintBrush->GetEditLayerRendererStates(InLandscapeInfo, bInSkipBrush);
+}
+
+void ULandscapeEditLayerPersistent::GetRendererStateInfo(const ULandscapeInfo* InLandscapeInfo, 
+	UE::Landscape::EditLayers::FEditLayerTargetTypeState& OutSupportedTargetTypeState, UE::Landscape::EditLayers::FEditLayerTargetTypeState& OutEnabledTargetTypeState, TArray<TSet<FName>>& OutRenderGroups) const
+{
+	using namespace UE::Landscape::EditLayers;
+
+	const FLandscapeLayer* LayerStruct = GetOwningLayer();
+	check(LayerStruct != nullptr);
+
+	// A layer can support all target types :
+	OutSupportedTargetTypeState.SetTargetTypeMask(ELandscapeToolTargetTypeFlags::All);
+	// Add an entry for each weightmap and consider them supported because there's nothing that prevents a given edit layer to write on a given weightmap layer :
+	Algo::ForEach(InLandscapeInfo->Layers, [&OutSupportedTargetTypeState] (const FLandscapeInfoLayerSettings& InInfoLayerSettings) { OutSupportedTargetTypeState.AddWeightmap(InInfoLayerSettings.GetLayerName()); });
+
+	// Compute the default state of each target type : 
+	ELandscapeToolTargetTypeFlags EnabledTargetTypeMask = ELandscapeToolTargetTypeFlags::None;
+	if (LayerStruct->bVisible)
+	{
+		if (LayerStruct->HeightmapAlpha > 0.0f)
+		{
+			EnabledTargetTypeMask |= ELandscapeToolTargetTypeFlags::Heightmap;
+		}
+		if (LayerStruct->WeightmapAlpha > 0.0f)
+		{
+			EnabledTargetTypeMask |= ELandscapeToolTargetTypeFlags::Weightmap | ELandscapeToolTargetTypeFlags::Visibility;
+		}
+	}
+	OutEnabledTargetTypeState = FEditLayerTargetTypeState(EnabledTargetTypeMask, MakeArrayView(OutSupportedTargetTypeState.GetActiveWeightmaps()));
+}
+
+TArray<UE::Landscape::EditLayers::FEditLayerRenderItem> ULandscapeEditLayerPersistent::GetRenderItems(const ULandscapeInfo* InLandscapeInfo) const
+{
+	using namespace UE::Landscape::EditLayers;
+
+	const FLandscapeLayer* LayerStruct = GetOwningLayer();
+	check(LayerStruct != nullptr);
+
+	TArray<FEditLayerRenderItem> RenderItems;
+	RenderItems.Reserve(InLandscapeInfo->XYtoComponentMap.Num() + 1);
+
+	// Heightmaps : We only need a single render item for heightmaps, because heightmaps are always present :
+	{
+		// The layer doesn't need more than the component itself to render properly:
+		FInputWorldArea InputWorldArea(FInputWorldArea::CreateLocalComponent());
+		// The layer only writes into the component itself (i.e. it renders to the area that it's currently being asked to render to):
+		FOutputWorldArea OutputWorldArea(FOutputWorldArea::CreateLocalComponent());
+		RenderItems.Emplace(FEditLayerTargetTypeState(ELandscapeToolTargetTypeFlags::Heightmap), InputWorldArea, OutputWorldArea, /*bModifyExistingWeightmapsOnly = */false);
+	}
+
+	// Weightmaps : add one render item per component in order to be able to indicate exactly which weightmap is needed for each of them. This avoids pre-allocating weightmaps on the merged result where we know we won't write a weightmap : 
+	{
+		TArray<ULandscapeComponent*> AllLandscapeComponents;
+		InLandscapeInfo->XYtoComponentMap.GenerateValueArray(AllLandscapeComponents);
+
+		for (ULandscapeComponent* Component : AllLandscapeComponents)
+		{
+			const FIntPoint ComponentKey = Component->GetComponentKey();
+
+			if (FLandscapeLayerComponentData* LayerData = Component->GetLayerData(LayerStruct->Guid))
+			{
+				TArray<FName, TInlineAllocator<16>> ComponentWeightmaps;
+				FEditLayerTargetTypeState OutputTargetTypeState;
+
+				// Iterate through all allocated weightmaps in order to find which one we will really write to : 
+				for (const FWeightmapLayerAllocationInfo& LayerAllocationInfo : LayerData->WeightmapData.LayerAllocations)
+				{
+					if (LayerAllocationInfo.IsAllocated())
+					{
+						if (LayerAllocationInfo.LayerInfo == ALandscapeProxy::VisibilityLayer)
+						{
+							OutputTargetTypeState.AddTargetTypeMask(ELandscapeToolTargetTypeFlags::Visibility);
+						}
+						else
+						{
+							OutputTargetTypeState.AddTargetTypeMask(ELandscapeToolTargetTypeFlags::Weightmap);
+						}
+						OutputTargetTypeState.AddWeightmap(LayerAllocationInfo.GetLayerName());
+					}
+				}
+
+				if (OutputTargetTypeState.GetTargetTypeMask() != ELandscapeToolTargetTypeFlags::None)
+				{
+					// The layer doesn't need more than the component itself to render properly:
+					FInputWorldArea InputWorldArea(FInputWorldArea::CreateSpecificComponent(ComponentKey));
+					// The layer only writes into the component itself (i.e. it renders to the area that it's currently being asked to render to):
+					FOutputWorldArea OutputWorldArea(FOutputWorldArea::CreateSpecificComponent(ComponentKey));
+					RenderItems.Emplace(OutputTargetTypeState, InputWorldArea, OutputWorldArea, /*bModifyExistingWeightmapsOnly = */false);
+				}
+			}
+		}
+	}
+	return RenderItems;
+}
+
+void ULandscapeEditLayerPersistent::RenderLayer(ILandscapeEditLayerRenderer::FRenderParams& InRenderParams)
+{
+	using namespace UE::Landscape;
+	using namespace UE::Landscape::Private;
+	using namespace UE::Landscape::EditLayers;
+
+	check(!InRenderParams.SortedComponentMergeRenderInfos.IsEmpty());
+
+	ULandscapeInfo* Info = InRenderParams.MergeRenderContext->GetLandscape()->GetLandscapeInfo();
+	check(Info != nullptr);
+	ALandscape* Landscape = InRenderParams.MergeRenderContext->GetLandscape();
+	FSceneInterface* SceneInterface = Landscape->GetWorld()->Scene;
+	ULandscapeScratchRenderTarget* WriteRT = nullptr;
+
+	const FLandscapeLayer* LayerStruct = GetOwningLayer();
+	check(LayerStruct != nullptr);
+
+	// Don't do anything with the blend render targets just yet : if this layer has nothing to render at all (e.g. no weightmap on the rendered area), that will save some processing :
+	auto LazyInitializeBlendRenderTargets = [MergeRenderContext = InRenderParams.MergeRenderContext, &WriteRT]
+	{
+		if (WriteRT == nullptr)
+		{
+			MergeRenderContext->CycleBlendRenderTargets();
+			WriteRT = MergeRenderContext->GetBlendRenderTargetWrite();
+
+			// We will write to the RT using a (bunch of) PS : 
+			WriteRT->Clear();
+			WriteRT->TransitionTo(ERHIAccess::RTV);
+		}
+	};
+
+	TArray<FName> TargetLayerNames = InRenderParams.MergeRenderContext->ConvertTargetLayerBitIndicesToNamesChecked(InRenderParams.RenderGroupBitIndices);
+	const int32 NumTargetLayers = TargetLayerNames.Num();
+
+	// The first step is to copy all of the necessary components' textures to WriteRT and remove the duplicate borders when doing so (subsection by subsection). This is done with a 
+	//  "copy from multiple sources" shader instead of several texture copies, in order to reduce the amount of copy texture commands, which can have a big impact on render-thread performance
+	//  for large landscapes. 
+
+	// For each subsection, we'll add a quad to render, using the following struct :
+	bool bUseSwizzleMask = false;
+	struct FCopyQuadParams
+	{
+		// Texture to read from :
+		FTextureResource* SourceTextureResource = nullptr;
+		// Texture region to read from :
+		FIntRect SourceRect;
+		// Texture region to write to :
+		FIntRect DestinationRect;
+		// In case the destination is a texture array (weightmaps), this is the slice index to write to :
+		int32 DestinationArrayIndex = 0;
+		// Weightmaps are packed in the source texture. This specifies which channel to read from in the source texture for this quad :
+		uint8 ChannelSwizzleMask = RGBAToRGBASwizzleMask;
+	};
+	TArray<FCopyQuadParams> CopyQuadsParams;
+	CopyQuadsParams.Reserve(InRenderParams.SortedComponentMergeRenderInfos.Num() * NumTargetLayers * Landscape->NumSubsections * Landscape->NumSubsections);
+
+	// This is kinda hacky, but we since we reuse AddRasterizeToRectsPass and it passes the source texture's sample coordinates from VS to PS via the UVs, dividing the source rect by a single TextureSize, 
+	//  we correct each quad's source rect by a common scale factor so that it samples at the right location in each source texture: 
+	const FVector2D CommonSourceTextureSize(InRenderParams.MergeRenderContext->GetMaxNeededResolution());
+	const int32 SubsectionSizeVerts = (Landscape->SubsectionSizeQuads + 1);
+
+	// Build the full list of quads to copy :
+	for (const FComponentMergeRenderInfo& ComponentMergeRenderInfo : InRenderParams.SortedComponentMergeRenderInfos)
+	{
+		for (int32 TargetLayerIndex = 0; TargetLayerIndex < NumTargetLayers; ++TargetLayerIndex)
+		{
+			const FName TargetLayerName = TargetLayerNames[TargetLayerIndex];
+			UTexture2D* SourceTexture = nullptr;
+			FVector2D SourceTextureBias(ForceInit);
+			// Copy all channels by default :
+			uint8 CopyTextureSwizzleMask = RGBAToRGBASwizzleMask;
+
+			if (InRenderParams.MergeRenderContext->IsHeightmapMerge())
+			{
+				SourceTexture = ComponentMergeRenderInfo.Component->GetHeightmap(LayerStruct->Guid);
+				SourceTextureBias = FVector2D(ComponentMergeRenderInfo.Component->HeightmapScaleBias.Z, ComponentMergeRenderInfo.Component->HeightmapScaleBias.W);
+			}
+			else 
+			{
+				const TArray<UTexture2D*>& WeightmapTextures = ComponentMergeRenderInfo.Component->GetWeightmapTextures(LayerStruct->Guid);
+				const TArray<FWeightmapLayerAllocationInfo>& AllocInfos = ComponentMergeRenderInfo.Component->GetWeightmapLayerAllocations(LayerStruct->Guid);
+				const FWeightmapLayerAllocationInfo* AllocInfo = AllocInfos.FindByPredicate([TargetLayerName](const FWeightmapLayerAllocationInfo& InAllocInfo) { return InAllocInfo.IsAllocated() && (InAllocInfo.GetLayerName() == TargetLayerName);});
+				if (AllocInfo != nullptr)
+				{
+					SourceTexture = WeightmapTextures[AllocInfo->WeightmapTextureIndex];
+					check(SourceTexture != nullptr);
+					// Note : don't use WeightmapScaleBias here, it has a different meaning than HeightmapScaleBias (very conveniently!) : this is compensated by the FloorToInt32 later on, 
+					//  but still, let's set this to zero here and use the fact that there's no texture sharing on weightmaps : 
+					SourceTextureBias = FVector2D::ZeroVector;
+					// Swizzle the appropriate source channel into the R component (as the target texture array's format is R8) : 
+					CopyTextureSwizzleMask = BuildChannelSwizzleMask(static_cast<ERGBAChannel>(AllocInfo->WeightmapTextureChannel));
+					bUseSwizzleMask = true;
+				}
+			}
+
+			if (SourceTexture != nullptr)
+			{
+				// We know we'll render something now so it's time we initialized our render target : 
+				LazyInitializeBlendRenderTargets();
+				check(WriteRT != nullptr);
+
+				FTextureResource* SourceTextureResource = SourceTexture->GetResource();
+				checkf(!SourceTexture->IsCompiling(), TEXT("All mips must have been loaded prior to using this function (%s)"), *SourceTexture->GetName());
+				// We get the overall texture size via the resource instead of direct GetSizeX/Y calls because the latter is unreliable while the texture is being built.
+				const FVector2D SourceTextureSize(SourceTextureResource->GetSizeX(), SourceTextureResource->GetSizeY());
+				const FIntPoint SourceTextureOffset(FMath::FloorToInt32(SourceTextureBias.X * SourceTextureSize.X), FMath::FloorToInt32(SourceTextureBias.Y * SourceTextureSize.Y));
+				
+				auto SourceTexturePixelCoordinatesToQuadCoords = [&SourceTextureSize, &CommonSourceTextureSize](const FIntPoint& InSourceCoords)
+				{ 
+					FVector2D Result = FVector2D(InSourceCoords) / SourceTextureSize * CommonSourceTextureSize;
+					return FIntPoint(FMath::FloorToInt32(Result.X), FMath::FloorToInt32(Result.Y));
+				};
+
+				// Fill that render target subsection by subsection, in order to bypass the redundant columns/lines on the subsection edges:
+				for (int32 SubsectionY = 0; SubsectionY < Landscape->NumSubsections; ++SubsectionY)
+				{
+					for (int32 SubsectionX = 0; SubsectionX < Landscape->NumSubsections; ++SubsectionX)
+					{
+						const FIntPoint SubsectionKey(SubsectionX, SubsectionY);
+						const FIntPoint SourcePosition = SourceTextureOffset + SubsectionKey * SubsectionSizeVerts;
+						const FIntPoint DestinationPosition = ComponentMergeRenderInfo.ComponentRegionInRenderArea.Min + SubsectionKey * ComponentMergeRenderInfo.Component->SubsectionSizeQuads;
+
+						FCopyQuadParams& QuadParams = CopyQuadsParams.Emplace_GetRef();
+						QuadParams.SourceTextureResource = SourceTextureResource;
+						QuadParams.SourceRect = FIntRect(SourceTexturePixelCoordinatesToQuadCoords(SourcePosition), SourceTexturePixelCoordinatesToQuadCoords(SourcePosition + FIntPoint(SubsectionSizeVerts, SubsectionSizeVerts)));
+						QuadParams.DestinationRect = FIntRect(DestinationPosition, DestinationPosition + FIntPoint(SubsectionSizeVerts, SubsectionSizeVerts));
+						QuadParams.DestinationArrayIndex = TargetLayerIndex;
+						QuadParams.ChannelSwizzleMask = CopyTextureSwizzleMask;
+					}
+				}
+			}
+		}
+	}
+
+	// Sort the list of quads by output texture first, then by source texture, because we'll coalesce several texture copies into the same pass (one output, multiple inputs)
+	CopyQuadsParams.Sort([](const FCopyQuadParams& InLHS, const FCopyQuadParams& InRHS)
+	{
+		if (InLHS.DestinationArrayIndex == InRHS.DestinationArrayIndex)
+		{
+			if (InLHS.SourceTextureResource == InRHS.SourceTextureResource)
+			{
+				return (InLHS.SourceRect.Min.Y == InRHS.SourceRect.Min.Y) 
+					? (InLHS.SourceRect.Min.X < InRHS.SourceRect.Min.X)
+					: (InLHS.SourceRect.Min.Y < InRHS.SourceRect.Min.Y);
+			}
+			return (InLHS.SourceTextureResource < InRHS.SourceTextureResource);
+		}
+
+		return (InLHS.DestinationArrayIndex < InRHS.DestinationArrayIndex);
+	});
+
+	// Now process this list of quads and prepare as many passes as necessary for performing all the copies :
+	if (!CopyQuadsParams.IsEmpty())
+	{
+		ENQUEUE_RENDER_COMMAND(LandscapeLayers_Cmd_CopySourcesToTargetLayer)(
+			[ CopyQuadsParams = MoveTemp(CopyQuadsParams)
+			, OutputResource = WriteRT->GetRenderTarget()->GetResource()
+			, OutputResourceName = WriteRT->GetDebugName()
+			, SceneInterface
+			, CommonSourceTextureSize
+			, bUseSwizzleMask] (FRHICommandListImmediate& InRHICmdList)
+		{
+			SCOPED_DRAW_EVENTF(InRHICmdList, LandscapeLayers, TEXT("CopyEditLayer -> %s"), OutputResourceName);
+			FRDGBuilder GraphBuilder(InRHICmdList, RDG_EVENT_NAME("CopySourcesToTargetLayer"));
+
+			FRDGTextureSRVRef BlackDummySRVRef = GraphBuilder.CreateSRV(FRDGTextureSRVDesc::Create(GSystemTextures.GetBlackDummy(GraphBuilder)));
+
+			FCopyQuadsMultiSourcePS::FPermutationDomain PermutationVector;
+			PermutationVector.Set< FCopyQuadsMultiSourcePS::FUseSwizleMask >(bUseSwizzleMask);
+			FGlobalShaderMap* ShaderMap = GetGlobalShaderMap(GMaxRHIFeatureLevel);
+			TShaderRef<FCopyQuadsMultiSourcePS> PixelShader = ShaderMap->GetShader<FCopyQuadsMultiSourcePS>(PermutationVector);
+
+			// TODO [jonathan.bard] this is just to avoid a RHI validation error for unoptimized shaders... once validation is made to not issue those errors, we can remove this
+			// Create a SceneView to please the shader bindings, but it's unused in practice 
+			FSceneViewFamilyContext ViewFamily(FSceneViewFamily::ConstructionValues(nullptr, SceneInterface, FEngineShowFlags(ESFIM_Game)).SetTime(FGameTime::GetTimeSinceAppStart()));
+			FSceneViewInitOptions ViewInitOptions;
+			ViewInitOptions.ViewFamily = &ViewFamily;
+			GetRendererModule().CreateAndInitSingleView(InRHICmdList, &ViewFamily, &ViewInitOptions);
+			const FSceneView* View = ViewFamily.Views[0];
+
+			// The following variables allow to accumulate quads for rendering in a single pass. Then, we flush them all out when the pass is full and we start again : 
+			int32 CurrentDestinationArrayIndex = INDEX_NONE;
+			TArray<FTextureResource*> CurrentSourceTextureResources;
+			TArray<FUintVector2> CurrentQuadInfos;
+			TArray<FUintVector4> CurrentSourceRects, CurrentDestinationRects;
+
+			// When we've reached the limit of what we can render in one pass, flush all and add a render pass : 
+			auto FlushCopyQuads = [&GraphBuilder, View, OutputResource, &CurrentDestinationArrayIndex, &CurrentSourceTextureResources, &CurrentSourceRects, &CurrentDestinationRects,
+				&CurrentQuadInfos, &CommonSourceTextureSize, BlackDummySRVRef, ShaderMap, PixelShader]()
+			{
+				if (CurrentSourceTextureResources.IsEmpty())
+				{
+					// Nothing to flush :
+					check(CurrentSourceRects.IsEmpty() && CurrentDestinationRects.IsEmpty() && CurrentQuadInfos.IsEmpty());
+					return;
+				}
+
+				check(!CurrentSourceRects.IsEmpty() && (CurrentSourceRects.Num() == CurrentDestinationRects.Num()) && (CurrentSourceRects.Num() == CurrentQuadInfos.Num()));
+
+				FRDGBufferRef RectBuffer = CreateUploadBuffer(GraphBuilder, TEXT("DestinationRects"), TConstArrayView<FUintVector4>(MakeArrayView(CurrentDestinationRects)));
+				FRDGBufferSRVRef RectBufferSRV = GraphBuilder.CreateSRV(FRDGBufferSRVDesc(RectBuffer, PF_R32G32B32A32_UINT));
+
+				FRDGBufferRef RectUVBuffer = CreateUploadBuffer(GraphBuilder, TEXT("RectUVs"), TConstArrayView<FUintVector4>(MakeArrayView(CurrentSourceRects)));
+				FRDGBufferSRVRef RectUVBufferSRV = GraphBuilder.CreateSRV(FRDGBufferSRVDesc(RectUVBuffer, PF_R32G32B32A32_UINT));
+
+				FRDGBufferRef QuadInfosBuffer = CreateUploadBuffer(GraphBuilder, TEXT("QuadInfos"), TConstArrayView<FUintVector2>(MakeArrayView(CurrentQuadInfos)));
+				FRDGBufferSRVRef QuadInfosBufferSRV = GraphBuilder.CreateSRV(FRDGBufferSRVDesc(QuadInfosBuffer, PF_R32G32_UINT));
+
+				FRDGTextureRef OutputTexture = GraphBuilder.RegisterExternalTexture(CreateRenderTarget(OutputResource->GetTextureRHI(), TEXT("OutputTexture")), ERDGTextureFlags::SkipTracking);
+				
+				// We may target an array slice directly : 
+				int32 ArrayIndex = -1;
+				if (OutputTexture->Desc.IsTextureArray())
+				{
+					check(CurrentDestinationArrayIndex < OutputTexture->Desc.ArraySize);
+					ArrayIndex = CurrentDestinationArrayIndex;
+				}
+
+				FCopyQuadsMultiSourcePSParameters* PassParameters = GraphBuilder.AllocParameters<FCopyQuadsMultiSourcePSParameters>();
+				PassParameters->RenderTargets[0] = FRenderTargetBinding(OutputTexture, ERenderTargetLoadAction::ELoad, /*InMipIndex = */0, ArrayIndex);
+				PassParameters->PS.View = View->ViewUniformBuffer;
+				PassParameters->PS.InQuadInfos = QuadInfosBufferSRV;
+				check(CurrentSourceTextureResources.Num() <= FCopyQuadsMultiSourcePS::NumMultiSources);
+				for (int32 TextureIndex = 0; TextureIndex < FCopyQuadsMultiSourcePS::NumMultiSources; ++TextureIndex)
+				{
+					FRDGTextureRef SourceTexture = BlackDummySRVRef->GetParent();
+					if (CurrentSourceTextureResources.IsValidIndex(TextureIndex))
+					{
+						SourceTexture = GraphBuilder.RegisterExternalTexture(CreateRenderTarget(CurrentSourceTextureResources[TextureIndex]->GetTexture2DRHI(), TEXT("SourceTexture")));
+					}
+					PassParameters->PS.InSourceTexture[TextureIndex] = SourceTexture;
+				}
+
+				FPixelShaderUtils::AddRasterizeToRectsPass<FCopyQuadsMultiSourcePS>(GraphBuilder,
+					ShaderMap,
+					RDG_EVENT_NAME("CopyQuadsMultiSourcePS"),
+					PixelShader,
+					PassParameters,
+					/*ViewportSize = */OutputTexture->Desc.Extent,
+					RectBufferSRV,
+					CurrentDestinationRects.Num(),
+					/*BlendState = */nullptr,
+					/*RasterizerState = */nullptr,
+					/*DepthStencilState = */nullptr,
+					/*StencilRef = */0,
+					/*TextureSize = */FIntPoint(FMath::FloorToInt32(CommonSourceTextureSize.X), FMath::FloorToInt32(CommonSourceTextureSize.Y)),
+					RectUVBufferSRV,
+					/*DownsampleFactor = */1,
+					ERDGPassFlags::NeverCull); // Use NeverCull because it renders a texture for which tracking is disabled
+
+				// Reset all for the next pass that comes along : 
+				CurrentDestinationArrayIndex = INDEX_NONE;
+				CurrentSourceTextureResources.Empty();
+				CurrentSourceRects.Empty();
+				CurrentDestinationRects.Empty();
+				CurrentQuadInfos.Empty();
+			};
+
+			for (const FCopyQuadParams& InSingleQuadParams : CopyQuadsParams)
+			{
+				check(InSingleQuadParams.DestinationArrayIndex != INDEX_NONE);
+				// If the output texture/slice has changed since the last iteration, it's time to flush and start a new pass :
+				if (CurrentDestinationArrayIndex != InSingleQuadParams.DestinationArrayIndex)
+				{
+					if (CurrentDestinationArrayIndex != INDEX_NONE)
+					{
+						FlushCopyQuads();
+					}
+					CurrentDestinationArrayIndex = InSingleQuadParams.DestinationArrayIndex;
+				}
+
+				if (CurrentSourceTextureResources.IsEmpty())
+				{
+					CurrentSourceTextureResources.Add(InSingleQuadParams.SourceTextureResource);
+				}
+				else if (InSingleQuadParams.SourceTextureResource != CurrentSourceTextureResources.Last())
+				{
+					// If we've reached the amount of textures we can render in a single pass, we flush the pass and initiate a new one : 
+					if (CurrentSourceTextureResources.Num() == FCopyQuadsMultiSourcePS::NumMultiSources)
+					{
+						FlushCopyQuads();
+						check(CurrentSourceTextureResources.IsEmpty());
+					}
+					CurrentSourceTextureResources.Add(InSingleQuadParams.SourceTextureResource);
+				}
+
+				// If we are using the same texture as the last one, we can render it in the same pass, just append our quad : 
+				if (InSingleQuadParams.SourceTextureResource == CurrentSourceTextureResources.Last())
+				{
+					const int32 SourceTextureIndex = CurrentSourceTextureResources.Num() - 1;
+					CurrentQuadInfos.Add(FUintVector2(static_cast<uint32>(SourceTextureIndex), InSingleQuadParams.ChannelSwizzleMask));
+					CurrentSourceRects.Add(FUintVector4(InSingleQuadParams.SourceRect.Min.X, InSingleQuadParams.SourceRect.Min.Y, InSingleQuadParams.SourceRect.Max.X, InSingleQuadParams.SourceRect.Max.Y));
+					CurrentDestinationRects.Add(FUintVector4(InSingleQuadParams.DestinationRect.Min.X, InSingleQuadParams.DestinationRect.Min.Y, InSingleQuadParams.DestinationRect.Max.X, InSingleQuadParams.DestinationRect.Max.Y));
+				}
+			}
+
+			// Flush the remaining quads if any : 
+			FlushCopyQuads();
+
+			GraphBuilder.Execute();
+		});
+	}
+
+	// No need to do anything if there's nothing to be rendered, this layer will just be ineffective on this batch :
+	if (WriteRT != nullptr)
+	{
+		// The second step is to perform the blend : merge the edit layer with the previous layers in the stack
+		InRenderParams.MergeRenderContext->CycleBlendRenderTargets(/*InDesiredWriteAccess = */ERHIAccess::RTV);
+		WriteRT = InRenderParams.MergeRenderContext->GetBlendRenderTargetWrite();
+		ULandscapeScratchRenderTarget* CurrentLayerReadRT = InRenderParams.MergeRenderContext->GetBlendRenderTargetRead();
+		ULandscapeScratchRenderTarget* PreviousLayersReadRT = InRenderParams.MergeRenderContext->GetBlendRenderTargetReadPrevious();
+	 
+		if (InRenderParams.MergeRenderContext->IsHeightmapMerge())
+		{
+			ENQUEUE_RENDER_COMMAND(LandscapeLayers_Cmd_HeightmapsMergeEditLayer)(
+				[ TargetLayerName = TargetLayerNames[0]
+				, OutputResource = WriteRT->GetRenderTarget2D()->GetResource()
+				, OutputResourceName = WriteRT->GetDebugName()
+				, CurrentEditLayerResource = CurrentLayerReadRT->GetRenderTarget2D()->GetResource()
+				, PreviousEditLayersResource = PreviousLayersReadRT->GetRenderTarget2D()->GetResource()
+				, EffectiveTextureSize = InRenderParams.MergeRenderContext->GetCurrentRenderBatch()->GetRenderTargetResolution(/*bInWithDuplicateBorders = */false)
+				, BlendMode = LayerStruct->BlendMode, HeightmapAlpha = LayerStruct->HeightmapAlpha ]
+				(FRHICommandListImmediate& InRHICmdList)
+			{
+				SCOPED_DRAW_EVENTF(InRHICmdList, LandscapeLayers, TEXT("MergeEditLayer -> %s"), OutputResourceName);
+				FRDGBuilder GraphBuilder(InRHICmdList, RDG_EVENT_NAME("HeightmapsMergeEditLayer"));
+
+				FRDGTextureRef OutputTextureRef = GraphBuilder.RegisterExternalTexture(CreateRenderTarget(OutputResource->TextureRHI, TEXT("OutputTexture")), ERDGTextureFlags::SkipTracking);
+				FRDGTextureRef CurrentEditLayerTextureRef = GraphBuilder.RegisterExternalTexture(CreateRenderTarget(CurrentEditLayerResource->TextureRHI, TEXT("CurrentEditLayerTexture")));
+				FRDGTextureRef PreviousEditLayersTextureRef = GraphBuilder.RegisterExternalTexture(CreateRenderTarget(PreviousEditLayersResource->TextureRHI, TEXT("PreviousEditLayersTexture")));
+
+				FLandscapeEditLayersHeightmapsMergeEditLayerPS::FParameters* PSParams = GraphBuilder.AllocParameters<FLandscapeEditLayersHeightmapsMergeEditLayerPS::FParameters>();
+				PSParams->RenderTargets[0] = FRenderTargetBinding(OutputTextureRef, ERenderTargetLoadAction::ENoAction);
+				PSParams->InEditLayerBlendMode = static_cast<uint32>(LandscapeBlendModeToEditLayerBlendMode(BlendMode));
+				PSParams->InEditLayerAlpha = HeightmapAlpha;
+				PSParams->InCurrentEditLayerHeightmap = GraphBuilder.CreateSRV(FRDGTextureSRVDesc::Create(CurrentEditLayerTextureRef));
+				PSParams->InPreviousEditLayersHeightmap = GraphBuilder.CreateSRV(FRDGTextureSRVDesc::Create(PreviousEditLayersTextureRef));
+
+				FLandscapeEditLayersHeightmapsMergeEditLayerPS::MergeEditLayerPS(GraphBuilder, PSParams, TargetLayerName, EffectiveTextureSize);
+				GraphBuilder.Execute();
+			});
+		}
+		else
+		{
+			struct FMergePaintLayerParams
+			{
+				FMergePaintLayerParams(const FName& InPaintLayerName)
+					: PaintLayerName(InPaintLayerName)
+				{}
+
+				FName PaintLayerName;
+				EWeightmapBlendMode EditLayerPaintLayerBlendMode = EWeightmapBlendMode::Additive;
+				float EditLayerAlpha = 1.0f;
+			};
+
+			// Prepare a list of all paint layers to merge and do it all in 1 command :
+			TArray<FMergePaintLayerParams> MergePaintLayersParams;
+			MergePaintLayersParams.Reserve(NumTargetLayers);
+			for (const FName& TargetLayerName : TargetLayerNames)
+			{
+				FMergePaintLayerParams& MergePaintLayerParams = MergePaintLayersParams.Emplace_GetRef(TargetLayerName);
+
+				const FLandscapeInfoLayerSettings* LayerSettings = Info->Layers.FindByPredicate([TargetLayerName](const FLandscapeInfoLayerSettings& InLayerSettings) { return (InLayerSettings.LayerName == TargetLayerName); });
+				check(LayerSettings != nullptr);
+				ULandscapeLayerInfoObject* LayerInfoObj = LayerSettings->LayerInfoObj;
+				check(LayerInfoObj != nullptr);
+				bool bIsSubtractiveBlend = false;
+				if (const bool* bSubstractiveInLayer = LayerStruct->WeightmapLayerAllocationBlend.Find(LayerInfoObj); (bSubstractiveInLayer != nullptr) && *bSubstractiveInLayer)
+				{
+					MergePaintLayerParams.EditLayerPaintLayerBlendMode = EWeightmapBlendMode::Subtractive;
+				}
+
+				if (TargetLayerName != UMaterialExpressionLandscapeVisibilityMask::ParameterName)
+				{
+					MergePaintLayerParams.EditLayerAlpha = LayerStruct->WeightmapAlpha;
+				}
+			}
+
+			ENQUEUE_RENDER_COMMAND(LandscapeLayers_Cmd_WeightmapsMergeEditLayer)(
+				[ MergePaintLayersParams
+				, OutputResource = WriteRT->GetRenderTarget2DArray()->GetResource()
+				, OutputResourceName = WriteRT->GetDebugName()
+				, CurrentEditLayerResource = CurrentLayerReadRT->GetRenderTarget2DArray()->GetResource()
+				, PreviousEditLayersResource = PreviousLayersReadRT->GetRenderTarget2DArray()->GetResource()
+				, EffectiveTextureSize = InRenderParams.MergeRenderContext->GetCurrentRenderBatch()->GetRenderTargetResolution(/*bInWithDuplicateBorders = */false) ]
+				(FRHICommandListImmediate& InRHICmdList)
+			{
+				SCOPED_DRAW_EVENTF(InRHICmdList, LandscapeLayers, TEXT("MergeEditLayer -> %s"), OutputResourceName);
+				FRDGBuilder GraphBuilder(InRHICmdList, RDG_EVENT_NAME("WeightmapsMergeEditLayer"));
+
+				FRDGTextureRef OutputTextureRef = GraphBuilder.RegisterExternalTexture(CreateRenderTarget(OutputResource->TextureRHI, TEXT("OutputTexture")), ERDGTextureFlags::SkipTracking);
+				FRDGTextureRef CurrentEditLayerTextureRef = GraphBuilder.RegisterExternalTexture(CreateRenderTarget(CurrentEditLayerResource->TextureRHI, TEXT("CurrentEditLayerTexture")));
+				FRDGTextureSRVRef CurrentEditLayerTextureSRVRef = GraphBuilder.CreateSRV(FRDGTextureSRVDesc::Create(CurrentEditLayerTextureRef));
+				FRDGTextureRef PreviousEditLayersTextureRef = GraphBuilder.RegisterExternalTexture(CreateRenderTarget(PreviousEditLayersResource->TextureRHI, TEXT("PreviousEditLayersTexture")));
+				FRDGTextureSRVRef PreviousEditLayersTextureSRVRef = GraphBuilder.CreateSRV(FRDGTextureSRVDesc::Create(PreviousEditLayersTextureRef));
+
+				const int32 NumTargetLayers = MergePaintLayersParams.Num();
+				for (int32 TargetLayerIndex = 0; TargetLayerIndex < NumTargetLayers; ++TargetLayerIndex)
+				{
+					FLandscapeEditLayersWeightmapsMergeEditLayerPS::FParameters* PSParams = GraphBuilder.AllocParameters<FLandscapeEditLayersWeightmapsMergeEditLayerPS::FParameters>();
+					PSParams->RenderTargets[0] = FRenderTargetBinding(OutputTextureRef, ERenderTargetLoadAction::ENoAction, /*InMipIndex = */0, /*InArraySlice = */TargetLayerIndex);
+					PSParams->InPaintLayerIndex = TargetLayerIndex;
+					PSParams->InEditLayerPaintLayerBlendMode = static_cast<uint32>(MergePaintLayersParams[TargetLayerIndex].EditLayerPaintLayerBlendMode);
+					PSParams->InEditLayerAlpha = MergePaintLayersParams[TargetLayerIndex].EditLayerAlpha;
+					PSParams->InCurrentEditLayerWeightmaps = CurrentEditLayerTextureSRVRef;
+					PSParams->InPreviousEditLayersWeightmaps = PreviousEditLayersTextureSRVRef;
+
+					FLandscapeEditLayersWeightmapsMergeEditLayerPS::MergeEditLayerPS(GraphBuilder, PSParams, MergePaintLayersParams[TargetLayerIndex].PaintLayerName, EffectiveTextureSize);
+				}
+				GraphBuilder.Execute();
+			});
+		}
+	}
+}
+
+FString ULandscapeEditLayerPersistent::GetEditLayerRendererDebugName() const
+{
+	const FLandscapeLayer* LayerStruct = GetOwningLayer();
+	check(LayerStruct != nullptr);
+	return LayerStruct->Name.ToString();
+}
+
+namespace UE::Landscape::EditLayers::Private
+{
+	using namespace UE::Landscape::EditLayers;
+
+	/** Struct that holds all the per-component information needed when preparing the batched merge context */
+	struct FComponentToRenderInfo
+	{
+		FComponentToRenderInfo() = default;
+
+		FComponentToRenderInfo(ULandscapeComponent* InComponent, int32 InComponentIndex, const int32 InNumAllComponents, const int32 NumTargetLayersToRender)
+			: Component(InComponent)
+			, ComponentIndex(InComponentIndex)
+			, DependentComponentBitIndices(false, InNumAllComponents)
+			, CombinedSectionRect(InComponent->GetSectionBase(), InComponent->GetSectionBase() + FIntPoint(InComponent->ComponentSizeQuads + 1, InComponent->ComponentSizeQuads + 1))
+			, ComponentKey(InComponent->GetComponentKey())
+			, MinDependentComponentKey(ComponentKey)
+			, MaxDependentComponentKey(ComponentKey)
+			, LocalBounds(InComponent->CachedLocalBox)
+			, WorldBounds(InComponent->CachedLocalBox.TransformBy(InComponent->GetComponentTransform()))
+			, TargetLayerBitIndices(false, NumTargetLayersToRender)
+		{}
+
+		void Finalize(const FIntRect& InDependentComponentInclusiveBounds, const FIntPoint& InComponentSizeQuads)
+		{
+			MinDependentComponentKey = InDependentComponentInclusiveBounds.Min;
+			MaxDependentComponentKey = InDependentComponentInclusiveBounds.Max;
+			CombinedSectionRect.Min = InDependentComponentInclusiveBounds.Min * InComponentSizeQuads;
+			CombinedSectionRect.Max = (InDependentComponentInclusiveBounds.Max + 1) * InComponentSizeQuads + 1;
+		}
+
+		/** Component to render */
+		ULandscapeComponent* Component = nullptr;
+
+		/** Index of the component to render in AllComponentsToRenderInfos */
+		int32 ComponentIndex = INDEX_NONE;
+
+		/** List of components this component depends on to render appropriately. Each dependency is represented by a bit which corresponds to the index of the dependent component in AllComponentsToRenderInfos */
+		TBitArray<> DependentComponentBitIndices;
+
+		/** Section rect of all the components that this component depends on to render */
+		FIntRect CombinedSectionRect;
+
+		/** Coordinate of the component to render */
+		FIntPoint ComponentKey = FIntPoint(ForceInit);
+
+		/** Minimum coordinate of the components that this component depends on to render */
+		FIntPoint MinDependentComponentKey = FIntPoint(MAX_int32, MAX_int32);
+
+		/** Maximum coordinate of the components that this component depends on to render */
+		FIntPoint MaxDependentComponentKey = FIntPoint(MIN_int32, MIN_int32);
+
+		/** Bounding volume of this component in local space */
+		FBox LocalBounds = FBox(ForceInit);
+
+		/** Bounding volume of this component in world space */
+		FBox WorldBounds = FBox(ForceInit);
+
+		/** List of target layers being written by this component. Each target layer is represented by a bit which corresponds to the index of the target layer names in AllTargetLayerNames */
+		TBitArray<> TargetLayerBitIndices;
+	};
+
+	/** Struct that holds all the per-render batch information needed when preparing the batched merge context */
+	struct FComponentToRenderBatchInfo
+	{
+		FComponentToRenderBatchInfo(int32 InNumComponentsToRender, int32 InBatchIndex)
+			: BatchIndex(InBatchIndex)
+			, ComponentToRenderInfoBitIndices(false, InNumComponentsToRender)
+		{
+		}
+
+		FIntRect GetProjectedSectionRect(const FComponentToRenderInfo& InComponentToRenderInfo) const
+		{
+			checkf(MinComponentKey.X != MAX_int32, TEXT("Shouldn't be called when the batch is empty"));
+			FIntRect NewCombinedSectionRect(CombinedSectionRect);
+			NewCombinedSectionRect.Union(InComponentToRenderInfo.CombinedSectionRect);
+			return NewCombinedSectionRect;
+		}
+
+		void AddToBatch(const FComponentToRenderInfo& InComponentToRenderInfo)
+		{
+			ComponentToRenderInfoBitIndices.CombineWithBitwiseOR(InComponentToRenderInfo.DependentComponentBitIndices, EBitwiseOperatorFlags::MinSize);
+			// Special case when it's the first addition to the batch : 
+			if (MinComponentKey.X == MAX_int32)
+			{
+				CombinedSectionRect = InComponentToRenderInfo.CombinedSectionRect;
+			}
+			else
+			{
+				CombinedSectionRect.Union(InComponentToRenderInfo.CombinedSectionRect);
+			}
+			MinComponentKey = MinComponentKey.ComponentMin(InComponentToRenderInfo.MinDependentComponentKey);
+			MaxComponentKey = MaxComponentKey.ComponentMax(InComponentToRenderInfo.MaxDependentComponentKey);
+		}
+
+
+	public:
+		/** Index of the batch within the merge operation */
+		int32 BatchIndex = INDEX_NONE;
+
+		/** Indices (in AllComponentsToRenderInfos) of the components which we need to render within this batch
+		  It's a bit array (1 bit per component to render info) to vastly optimize the batching operation, which is a O(N^2):  */
+		TBitArray<> ComponentToRenderInfoBitIndices;
+
+		/** Section rect of all the components that this batch will render*/
+		FIntRect CombinedSectionRect;
+
+		/** Minimum coordinate of the components that this batch will render */
+		FIntPoint MinComponentKey = FIntPoint(MAX_int32, MAX_int32);
+
+		/** Maximum coordinate of the components that this batch will render */
+		FIntPoint MaxComponentKey = FIntPoint(MIN_int32, MIN_int32);
+	};
+
+	/** Struct that holds all the per-renderer information needed when preparing the batched merge context */
+	struct FEditLayerRendererRenderInfo
+	{
+		FEditLayerRendererRenderInfo() = default;
+		FEditLayerRendererRenderInfo(const FEditLayerRendererState& InRendererState, int32 InRendererIndex, const FTransform& InLandscapeTransform, double InMaxLocalHeight, int32 InNumComponentsToRender)
+			: RendererState(InRendererState)
+			, RendererIndex(InRendererIndex)
+			, ComponentToRenderInfoBitIndices(false, InNumComponentsToRender)
+		{
+
+#if ENABLE_VISUAL_LOG
+			// Pick a random color for each renderer : 
+			uint32 Hash = GetTypeHashHelper(InRendererIndex);
+			uint8* HashElement = reinterpret_cast<uint8*>(&Hash);
+			VisualLogColor = FColor(HashElement[0], HashElement[1], HashElement[2], FMergeRenderContext::GetVisualLogAlpha());
+
+			VisualLogOffsetLocalSpace = InRendererIndex * InLandscapeTransform.InverseTransformVector(FVector(0.0, 0.0, CVarLandscapeBatchedMergeVisualLogOffsetIncrement.GetValueOnGameThread())).Z + InMaxLocalHeight;
+#endif // ENABLE_VISUAL_LOG
+		}
+		
+		FEditLayerRendererState RendererState;
+		int32 RendererIndex = INDEX_NONE;
+
+		// Indices (in AllComponentsToRenderInfos) of the components which we need to render for this renderer
+		TBitArray<> ComponentToRenderInfoBitIndices;
+
+#if ENABLE_VISUAL_LOG
+		FColor VisualLogColor = FColor(ForceInit);
+		double VisualLogOffsetLocalSpace = 0.0;
+#endif // ENABLE_VISUAL_LOG
+	};
+
+	/** Given the list of components that actually need to be rendered, divide the work into batches, such that:
+	 *   - When a component is being rendered by a given batch, all components needed for rendering this given component are present in the batch
+	 *   - All components end up being rendered in at least one of the batches
+	 */
+	TArray<FComponentToRenderBatchInfo> DivideIntoBatches(const TBitArray<>& InFinalComponentsToRenderInfoBitIndices, const TArray<FComponentToRenderInfo>& InAllComponentsToRenderInfos)
+	{
+		TRACE_CPUPROFILER_EVENT_SCOPE(DivideIntoBatches);
+
+		const int32 NumFinalComponentsToRender = InFinalComponentsToRenderInfoBitIndices.CountSetBits();
+
+		// Compute the final batch resolution we will use for this work : 
+		const int32 DesiredMaxBatchResolution = CVarLandscapeEditLayersMaxResolutionPerRenderBatch.GetValueOnGameThread();
+
+		// Compute the minimal batch resolution :
+		FIntPoint MinBatchResolution(MIN_int32, MIN_int32);
+		// Sort the components to render by MinDependentComponentKey in order to minimize the number of batches needed. 
+		TArray<TPair<int32, FIntPoint>> RemainingComponentsToRenderInfoIndices;
+		RemainingComponentsToRenderInfoIndices.Reserve(NumFinalComponentsToRender);
+		for (TConstSetBitIterator It(InFinalComponentsToRenderInfoBitIndices); It; ++It)
+		{
+			const int32 ComponentIndex = It.GetIndex();
+			const FComponentToRenderInfo& ComponentToRenderInfo = InAllComponentsToRenderInfos[ComponentIndex];
+			RemainingComponentsToRenderInfoIndices.Add(MakeTuple(ComponentIndex, ComponentToRenderInfo.MinDependentComponentKey));
+			// The component that has the largest CombinedSectionRect defines the minimum batch size, because it requires all components in that area to render appropriately : 
+			MinBatchResolution = MinBatchResolution.ComponentMax(ComponentToRenderInfo.CombinedSectionRect.Size());
+		}
+
+		FIntPoint MaxBatchResolution = MinBatchResolution.ComponentMax(FIntPoint(DesiredMaxBatchResolution));
+		if ((MinBatchResolution.X > DesiredMaxBatchResolution) || (MinBatchResolution.Y > DesiredMaxBatchResolution))
+		{
+			UE_LOG(LogLandscape, Warning, TEXT("Landscape edit layers merge requires a minimum batch size of resolution %ix%i, which is higher than the current desired maximum batch resolution %ix%i. Consider adjusting the maximum batch resolution (landscape.EditLayersLocalMerge.MaxResolutionPerRenderBatch) or make sure the landscape edit layers renderers in use require smaller work area."),
+				MinBatchResolution.X, MinBatchResolution.Y, DesiredMaxBatchResolution, DesiredMaxBatchResolution);
+			MaxBatchResolution = MinBatchResolution;
+		}
+
+		// Sort the components to render by MinDependentComponentKey in order to minimize the number of batches needed. 
+		//  We actually use the inverse order because we want to process the elements from RemainingComponentsToRenderInfoIndices in inverse order 
+		RemainingComponentsToRenderInfoIndices.Sort([](const TPair<int32, FIntPoint>& InLHS, const TPair<int32, FIntPoint>& InRHS)
+		{
+			const FIntPoint& InLHSMinDependentComponentKey = InLHS.Value;
+			const FIntPoint& InRHSMinDependentComponentKey = InRHS.Value;
+			if (InLHSMinDependentComponentKey.Y > InRHSMinDependentComponentKey.Y)
+			{
+				return true;
+			}
+			else if (InLHSMinDependentComponentKey.Y == InRHSMinDependentComponentKey.Y)
+			{
+				return (InLHSMinDependentComponentKey.X > InRHSMinDependentComponentKey.X);
+			}
+			return false;
+		});
+
+		// Iterate on all the work items and organize them into batches as large as possible (within the maximum allowed resolution)
+		TArray<FComponentToRenderBatchInfo> AllComponentToRenderBatchInfos;
+		if (!RemainingComponentsToRenderInfoIndices.IsEmpty())
+		{
+			TBitArray<TInlineAllocator<1>> TempBitArray(false, InAllComponentsToRenderInfos.Num());
+			while (!RemainingComponentsToRenderInfoIndices.IsEmpty())
+			{
+				const int32 ComponentIndex = RemainingComponentsToRenderInfoIndices.Pop(EAllowShrinking::No).Key;
+				const FComponentToRenderInfo& ComponentToRenderInfo = InAllComponentsToRenderInfos[ComponentIndex];
+				const FIntPoint NeededResolution = ComponentToRenderInfo.CombinedSectionRect.Size();
+				check((NeededResolution.X <= MaxBatchResolution.X) && (NeededResolution.Y <= MaxBatchResolution.Y));
+
+				int32 BestBatchIndex = INDEX_NONE;
+				int32 MinBatchRenderArea = MAX_int32;
+				int32 MaxNumComponentsInCommonWithBatch = -1;
+
+				// Iterate through all batches and try to find which would be able to accept it and amongst those, which would have the minimal overall resolution: 
+				const int32 NumBatches = AllComponentToRenderBatchInfos.Num();
+				for (int32 BatchIndex = 0; BatchIndex < NumBatches; ++BatchIndex)
+				{
+					const FComponentToRenderBatchInfo& Batch = AllComponentToRenderBatchInfos[BatchIndex];
+					FIntRect ProjectedBatchSectionRect = Batch.GetProjectedSectionRect(ComponentToRenderInfo);
+
+					// If after adding this component and its dependent components, the batch still fits within the max allowed resolution, it can accept it : 
+					if ((ProjectedBatchSectionRect.Width() <= MaxBatchResolution.X) && (ProjectedBatchSectionRect.Height() <= MaxBatchResolution.Y))
+					{
+						// Favor the batch that has the most components in common with what we're trying to render : 
+						TempBitArray = TBitArray<>::BitwiseAND(Batch.ComponentToRenderInfoBitIndices, ComponentToRenderInfo.DependentComponentBitIndices, EBitwiseOperatorFlags::MinSize);
+						const int32 NumComponentsInCommonWithBatch = TempBitArray.CountSetBits();
+						// If the batch already has all the components we need, it's a perfect match, we won't ever find a better batch so just stop the search there :
+						if (NumComponentsInCommonWithBatch == ComponentToRenderInfo.DependentComponentBitIndices.CountSetBits())
+						{
+							BestBatchIndex = BatchIndex;
+							break;
+						}
+
+						const int32 ProjectedBatchRenderArea = ProjectedBatchSectionRect.Size().X * ProjectedBatchSectionRect.Size().Y;
+						if (NumComponentsInCommonWithBatch > MaxNumComponentsInCommonWithBatch)
+						{
+							MaxNumComponentsInCommonWithBatch = NumComponentsInCommonWithBatch;
+							MinBatchRenderArea = ProjectedBatchRenderArea;
+							BestBatchIndex = BatchIndex;
+						}
+						else if (NumComponentsInCommonWithBatch == MaxNumComponentsInCommonWithBatch)
+						{
+							// Favor the batch with the lesser area 
+							if (ProjectedBatchRenderArea < MinBatchRenderArea)
+							{
+								MinBatchRenderArea = ProjectedBatchRenderArea;
+								BestBatchIndex = BatchIndex;
+							}
+						}
+					}
+				}
+
+				// If we have found a batch, just add the FComponentToRenderBatchInfo to it, otherwise, add a new batch:
+				FComponentToRenderBatchInfo& SelectedBatch = (BestBatchIndex != INDEX_NONE) ? AllComponentToRenderBatchInfos[BestBatchIndex]
+					: AllComponentToRenderBatchInfos.Add_GetRef(FComponentToRenderBatchInfo(InAllComponentsToRenderInfos.Num(), /*InBatchIndex = */AllComponentToRenderBatchInfos.Num()));
+
+				SelectedBatch.AddToBatch(ComponentToRenderInfo);
+			}
+		}
+		return AllComponentToRenderBatchInfos;
+	}
+
+#if ENABLE_VISUAL_LOG
+	struct FComponentDependenciesVisLogHelper
+	{
+		enum class EShowNodeInfo : uint8
+		{
+			None = 0, 
+			Minimal,
+			Detailed
+		};
+
+		FComponentDependenciesVisLogHelper(const ALandscape* InLandscape, bool bInIsHeightmapMerge, EShowNodeInfo InShowNodeInfo, const FMergeRenderContext& InRenderContext)
+			: Landscape(InLandscape)
+			, LandscapeTransform(InLandscape->GetTransform())
+			, bIsHeightmapMerge(bInIsHeightmapMerge)
+			, ShowNodeInfo(InShowNodeInfo)
+			, LandscapeComponentLocalSize(InLandscape->ComponentSizeQuads)
+			, RenderContext(&InRenderContext)
+		{}
+		
+		const FVector& AddNode(const FComponentToRenderInfo& InComponentRenderInfo, const FEditLayerRendererRenderInfo& InRendererRenderInfo)
+		{
+			const FIntPoint ComponentKey = InComponentRenderInfo.Component->GetComponentKey();
+
+			TPair<int32, int32> Key(InComponentRenderInfo.ComponentIndex, InRendererRenderInfo.RendererIndex);
+			const FVector* Center = KeyToCenter.Find(Key);
+			if (Center == nullptr)
+			{
+				FTransform BaseTransform = FTransform(FVector(0.0, 0.0, InRendererRenderInfo.VisualLogOffsetLocalSpace)) * LandscapeTransform;
+				FBox VisualBounds(
+					FVector(ComponentKey.X * LandscapeComponentLocalSize, ComponentKey.Y * LandscapeComponentLocalSize, 0.0),
+					FVector((ComponentKey.X + 1) * LandscapeComponentLocalSize, (ComponentKey.Y + 1) * LandscapeComponentLocalSize, 0.0));
+
+				FString Message;
+				if (ShowNodeInfo == EShowNodeInfo::Minimal)
+				{
+					FIntRect DependentComponentArea(InComponentRenderInfo.MinDependentComponentKey, InComponentRenderInfo.MaxDependentComponentKey + 1);
+					Message = FString::Printf(TEXT("%s"), *ComponentKey.ToString());
+				}
+				else if (ShowNodeInfo == EShowNodeInfo::Detailed)
+				{
+					Message = FString::Printf(TEXT("%s\n(%s)"), *ComponentKey.ToString(), *InRendererRenderInfo.RendererState.GetRenderer()->GetEditLayerRendererDebugName());
+				}
+				// On the first renderer, show additional info: 
+				if (InRendererRenderInfo.RendererIndex == 0)
+				{
+					FIntRect DependentComponentAreaRelative(InComponentRenderInfo.MinDependentComponentKey - ComponentKey, InComponentRenderInfo.MaxDependentComponentKey - ComponentKey);
+					Message += FString::Printf(TEXT("\n{%s}"), *DependentComponentAreaRelative.ToString());
+					if (!bIsHeightmapMerge)
+					{
+						TArray<FName> TargetLayerNames = RenderContext->ConvertTargetLayerBitIndicesToNames(InComponentRenderInfo.TargetLayerBitIndices);
+						Message += FString::Printf(TEXT("\n%s"), *ConvertTargetLayerNamesToString(TargetLayerNames));
+					}
+				}
+
+				FMatrix Transform = BaseTransform.ToMatrixWithScale();
+				UE_VLOG_OBOX(Landscape, LogLandscape, Log, VisualBounds, Transform, InRendererRenderInfo.VisualLogColor, TEXT("%s"), *Message);
+				Center = &KeyToCenter.Add(Key, Transform.TransformPosition(VisualBounds.GetCenter()));
+			}
+			return *Center;
+		}
+
+		void AddDependency(const FComponentToRenderInfo& InSourceComponentRenderInfo, const FEditLayerRendererRenderInfo& InSourceRendererRenderInfo,
+			const FComponentToRenderInfo& InDestinationComponentRenderInfo, const FEditLayerRendererRenderInfo& InDestinationRendererRenderInfo)
+		{
+			const FVector SourceCenter = AddNode(InSourceComponentRenderInfo, InSourceRendererRenderInfo);
+			const FVector DestinationCenter = AddNode(InDestinationComponentRenderInfo, InDestinationRendererRenderInfo);
+			UE_VLOG_ARROW(Landscape, LogLandscape, Log, SourceCenter, DestinationCenter, InSourceRendererRenderInfo.VisualLogColor, TEXT(""));
+		}
+
+	private:
+		const ALandscape* Landscape = nullptr;
+		FTransform LandscapeTransform;
+		bool bIsHeightmapMerge = false;
+		EShowNodeInfo ShowNodeInfo = EShowNodeInfo::None;
+		double LandscapeComponentLocalSize = 0.0;
+		const FMergeRenderContext* RenderContext = nullptr;
+
+		TMap<TPair<int32, int32>, FVector> KeyToCenter;
+	};
+
+	// Log the shapes of the render item output input render items if requested : 
+	void VisLogRenderItemInput(const ALandscape* InLandscape, const FInputWorldArea& InInputWorldArea, const FEditLayerRendererRenderInfo& InRendererRenderInfo,
+		const FTransform& InLandscapeTransform, const FBox& InLandscapeLoadedBounds, const TArrayView<ULandscapeComponent*>& InComponentsToRender)
+	{
+		// Display the shapes 2 unreal unit (in world space) under the requested offset (so that they're located under the output items) :
+		static double OutputLocalOffsetLocalSpace = -2.0 / InLandscapeTransform.GetScale3D().Z;
+		const double LandscapeComponentLocalSize = InLandscape->ComponentSizeQuads;
+
+		FTransform BaseTransform = FTransform(FVector(0.0, 0.0, InRendererRenderInfo.VisualLogOffsetLocalSpace + OutputLocalOffsetLocalSpace)) * InLandscapeTransform;
+		switch (InInputWorldArea.GetType())
+		{
+		case FInputWorldArea::EType::Infinite:
+		{
+			// Infinite input area means all loaded components: 
+			FBox VisualBounds = InLandscapeLoadedBounds;
+			VisualBounds.Min.Z = 0.0;
+			VisualBounds.Max.Z = 0.0;
+			UE_VLOG_WIREOBOX(InLandscape, LogLandscape, Log, VisualBounds, BaseTransform.ToMatrixWithScale(), InRendererRenderInfo.VisualLogColor, TEXT(""));
+			break;
+		}
+		case FInputWorldArea::EType::LocalComponent:
+		{
+			// Local input area means each of the landscape components : 
+			for (ULandscapeComponent* Component : InComponentsToRender)
+			{
+				FIntRect ComponentKeys = InInputWorldArea.GetLocalComponentKeys(Component->GetComponentKey());
+				// Transform from inclusive to exclusive bounds : 
+				ComponentKeys.Max += FIntPoint(1, 1);
+				FBox VisualBounds(
+					FVector(ComponentKeys.Min.X * LandscapeComponentLocalSize, ComponentKeys.Min.Y * LandscapeComponentLocalSize, 0.0),
+					FVector(ComponentKeys.Max.X * LandscapeComponentLocalSize, ComponentKeys.Max.Y * LandscapeComponentLocalSize, 0.0));
+				UE_VLOG_WIREOBOX(InLandscape, LogLandscape, Log, VisualBounds, BaseTransform.ToMatrixWithScale(), InRendererRenderInfo.VisualLogColor, TEXT(""));
+			}
+			break;
+		}
+		case FInputWorldArea::EType::SpecificComponent:
+		{
+			FIntRect ComponentKeys = InInputWorldArea.GetSpecificComponentKeys();
+			// Transform from inclusive to exclusive bounds : 
+			ComponentKeys.Max += FIntPoint(1, 1);
+			FBox VisualBounds(
+				FVector(ComponentKeys.Min.X * LandscapeComponentLocalSize, ComponentKeys.Min.Y * LandscapeComponentLocalSize, 0.0),
+				FVector(ComponentKeys.Max.X * LandscapeComponentLocalSize, ComponentKeys.Max.Y * LandscapeComponentLocalSize, 0.0));
+			UE_VLOG_WIREOBOX(InLandscape, LogLandscape, Log, VisualBounds, BaseTransform.ToMatrixWithScale(), InRendererRenderInfo.VisualLogColor, TEXT(""));
+			break;
+		}
+		case FInputWorldArea::EType::OOBox:
+		{
+			FOOBox2D OOBox = InInputWorldArea.GetOOBox();
+			FBox VisualBounds(-FVector(OOBox.Extents, 0.0), FVector(OOBox.Extents, 0.0));
+			FTransform OOBoxLocalTransform = BaseTransform.GetRelativeTransform(OOBox.Transform);
+			OOBoxLocalTransform.SetTranslation(FVector::ZeroVector);
+			FTransform FinalTransform = BaseTransform * OOBoxLocalTransform;
+			UE_VLOG_WIREOBOX(InLandscape, LogLandscape, Log, VisualBounds, FinalTransform.ToMatrixWithScale(), InRendererRenderInfo.VisualLogColor, TEXT(""));
+			break;
+		}
+		default:
+			check(false);
+		}
+	}
+
+	// Log the shapes of the render item output input render items if requested : 
+	void VisLogRenderItemOutput(const ALandscape* InLandscape, bool bInIsHeightmapMerge, bool bInAffectsOutputLayerBitIndices, const TArrayView<const FName>& InRenderItemTargetLayerNames,
+		const FOutputWorldArea& InOutputWorldArea, const FEditLayerRendererRenderInfo& InRendererRenderInfo, const FTransform& InLandscapeTransform, const FBox& InLandsapeLoadedBounds, const TArrayView<ULandscapeComponent*>& InComponentsToRender)
+	{
+		// Display the shapes 1 unreal unit (in world space) under the requested offset :
+		static double OutputLocalOffsetLocalSpace = -1.0 / InLandscapeTransform.GetScale3D().Z;
+		const double LandscapeComponentLocalSize = InLandscape->ComponentSizeQuads;
+		FString LogMessage = InRendererRenderInfo.RendererState.GetRenderer()->GetEditLayerRendererDebugName();
+		if (!bInIsHeightmapMerge)
+		{
+			LogMessage += FString::Printf(TEXT("\n%s%s%s"), *ConvertTargetLayerNamesToString(InRenderItemTargetLayerNames), !bInAffectsOutputLayerBitIndices ? TEXT("(") : TEXT(""), !bInAffectsOutputLayerBitIndices ? TEXT(")") : TEXT(""));
+		}
+
+		FTransform BaseTransform = FTransform(FVector(0.0, 0.0, InRendererRenderInfo.VisualLogOffsetLocalSpace + OutputLocalOffsetLocalSpace)) * InLandscapeTransform;
+		switch (InOutputWorldArea.GetType())
+		{
+		case FOutputWorldArea::EType::LocalComponent:
+		{
+			// Local input area means each of the landscape components : 
+			for (ULandscapeComponent* Component : InComponentsToRender)
+			{
+				FIntPoint ComponentKey = Component->GetComponentKey();
+				FBox VisualBounds(
+					FVector(ComponentKey.X * LandscapeComponentLocalSize, ComponentKey.Y * LandscapeComponentLocalSize, 0.0),
+					FVector((ComponentKey.X + 1) * LandscapeComponentLocalSize, (ComponentKey.Y + 1) * LandscapeComponentLocalSize, 0.0));
+				UE_VLOG_OBOX(InLandscape, LogLandscape, Log, VisualBounds, BaseTransform.ToMatrixWithScale(), InRendererRenderInfo.VisualLogColor, TEXT("%s"), *LogMessage);
+			}
+			break;
+		}
+		case FOutputWorldArea::EType::SpecificComponent:
+		{
+			FIntPoint ComponentKey = InOutputWorldArea.GetSpecificComponentKey();
+			FBox VisualBounds(
+				FVector(ComponentKey.X * LandscapeComponentLocalSize, ComponentKey.Y * LandscapeComponentLocalSize, 0.0),
+				FVector((ComponentKey.X + 1) * LandscapeComponentLocalSize, (ComponentKey.Y + 1) * LandscapeComponentLocalSize, 0.0));
+			UE_VLOG_OBOX(InLandscape, LogLandscape, Log, VisualBounds, BaseTransform.ToMatrixWithScale(), InRendererRenderInfo.VisualLogColor, TEXT("%s"), *LogMessage);
+			break;
+		}
+		case FOutputWorldArea::EType::OOBox:
+		{
+			FOOBox2D OOBox = InOutputWorldArea.GetOOBox();
+			FBox VisualBounds(-FVector(OOBox.Extents, 0.0), FVector(OOBox.Extents, 0.0));
+			FTransform OOBoxLocalTransform = BaseTransform.GetRelativeTransform(OOBox.Transform);
+			OOBoxLocalTransform.SetTranslation(FVector::ZeroVector);
+			FTransform FinalTransform = BaseTransform * OOBoxLocalTransform;
+			UE_VLOG_OBOX(InLandscape, LogLandscape, Log, VisualBounds, FinalTransform.ToMatrixWithScale(), InRendererRenderInfo.VisualLogColor, TEXT("%s"), *LogMessage);
+			break;
+		}
+		default:
+			check(false);
+		}
+	}
+#endif // ENABLE_VISUAL_LOG
+} // UE::Landscape::EditLayers::Private
+
+void ULandscapeHeightmapNormalsEditLayerRenderer::RenderLayer(FRenderParams& InRenderParams)
+{
+	using namespace UE::Landscape::EditLayers;
+	using namespace UE::Landscape::EditLayers::Private;
+
+	FMergeRenderContext* RenderContext = InRenderParams.MergeRenderContext;
+	const FMergeRenderBatch* RenderBatch = RenderContext->GetCurrentRenderBatch();
+	ALandscape* Landscape = RenderContext->GetLandscape();
+
+	TArray<FName> TargetLayerNames = InRenderParams.MergeRenderContext->ConvertTargetLayerBitIndicesToNamesChecked(InRenderParams.RenderGroupBitIndices);
+	checkf(TargetLayerNames.Num() == 1, TEXT("Normals should only be generated on heightmap merge, which should have 1 and only target layer"));
+
+	RenderContext->CycleBlendRenderTargets(/*InDesiredWriteAccess = */ERHIAccess::RTV);
+	ULandscapeScratchRenderTarget* WriteRT = RenderContext->GetBlendRenderTargetWrite();
+	ULandscapeScratchRenderTarget* ReadRT = RenderContext->GetBlendRenderTargetRead();
+	ULandscapeScratchRenderTarget* ValidityRT = RenderContext->GetValidityRenderTarget(TargetLayerNames[0]);
+	check((WriteRT != nullptr) && (ReadRT != nullptr) && (ValidityRT != nullptr));
+
+	ValidityRT->TransitionTo(ERHIAccess::SRVMask);
+
+	FIntPoint NumComponentsInRenderArea = RenderBatch->MaxComponentKey - RenderBatch->MinComponentKey + 1;
+	check((NumComponentsInRenderArea.X > 0) && (NumComponentsInRenderArea.Y > 0));
+
+	ENQUEUE_RENDER_COMMAND(LandscapeLayers_Cmd_HeightmapsGenerateNormals)(
+		[OutputResource = WriteRT->GetRenderTarget2D()->GetResource()
+		, OutputResourceName = WriteRT->GetDebugName()
+		, SourceResource = ReadRT->GetRenderTarget2D()->GetResource()
+		, ValidityResource = ValidityRT->GetRenderTarget2D()->GetResource()
+		, EffectiveTextureSize = RenderBatch->GetRenderTargetResolution(/*bInWithDuplicateBorders = */false)
+		, LandscapeGridScale = Landscape->GetRootComponent()->GetRelativeScale3D()
+		, ComponentSizeQuads = Landscape->ComponentSizeQuads
+		, NumComponentsInRenderArea]
+		(FRHICommandListImmediate& InRHICmdList)
+	{
+		SCOPED_DRAW_EVENTF(InRHICmdList, LandscapeLayers, TEXT("GenerateNormals -> %s"), OutputResourceName);
+		FRDGBuilder GraphBuilder(InRHICmdList, RDG_EVENT_NAME("HeightmapsGenerateNormals"));
+
+		FRDGTextureRef OutputTextureRef = GraphBuilder.RegisterExternalTexture(CreateRenderTarget(OutputResource->TextureRHI, TEXT("OutputTexture")), ERDGTextureFlags::SkipTracking);
+		FRDGTextureRef SourceTextureRef = GraphBuilder.RegisterExternalTexture(CreateRenderTarget(SourceResource->TextureRHI, TEXT("SourceTexture")));
+		FRDGTextureRef ValidityTextureRef = GraphBuilder.RegisterExternalTexture(CreateRenderTarget(ValidityResource->TextureRHI, TEXT("ValidityTexture")));
+
+		FLandscapeEditLayersHeightmapsGenerateNormalsPS::FParameters* PSParams = GraphBuilder.AllocParameters<FLandscapeEditLayersHeightmapsGenerateNormalsPS::FParameters>();
+		PSParams->RenderTargets[0] = FRenderTargetBinding(OutputTextureRef, ERenderTargetLoadAction::ENoAction);
+		PSParams->InTextureSize = FUintVector4(EffectiveTextureSize.X, EffectiveTextureSize.Y, SourceResource->GetSizeX(), SourceResource->GetSizeY());
+		PSParams->InLandscapeGridScale = static_cast<FVector3f>(LandscapeGridScale);
+		PSParams->InComponentSizeQuads = ComponentSizeQuads;
+		PSParams->InNumComponents = FUintVector2(NumComponentsInRenderArea.X, NumComponentsInRenderArea.Y);
+		PSParams->InSourceHeightmapSampler = TStaticSamplerState<SF_Point, AM_Clamp, AM_Clamp>::GetRHI();
+		PSParams->InSourceHeightmap = GraphBuilder.CreateSRV(FRDGTextureSRVDesc::Create(SourceTextureRef));
+		PSParams->InValidityTexture = GraphBuilder.CreateSRV(FRDGTextureSRVDesc::Create(ValidityTextureRef));
+		FLandscapeEditLayersHeightmapsGenerateNormalsPS::GenerateNormalsPS(GraphBuilder, PSParams, EffectiveTextureSize);
+		GraphBuilder.Execute();
+	});
+}
+
+void ULandscapeWeightmapWeightBlendedLayersRenderer::RenderLayer(FRenderParams& InRenderParams)
+{
+	using namespace UE::Landscape::EditLayers;
+	using namespace UE::Landscape::EditLayers::Private;
+
+	check(!InRenderParams.MergeRenderContext->IsHeightmapMerge());
+
+	InRenderParams.MergeRenderContext->CycleBlendRenderTargets(ERHIAccess::RTV);
+	ULandscapeScratchRenderTarget* WriteRT = InRenderParams.MergeRenderContext->GetBlendRenderTargetWrite();
+	ULandscapeScratchRenderTarget* ReadRT = InRenderParams.MergeRenderContext->GetBlendRenderTargetRead();
+
+	TArray<FName> TargetLayerNames = InRenderParams.MergeRenderContext->ConvertTargetLayerBitIndicesToNamesChecked(InRenderParams.RenderGroupBitIndices);
+	TArray<ULandscapeLayerInfoObject*> LayerInfos = InRenderParams.MergeRenderContext->ConvertTargetLayerBitIndicesToLayerInfosChecked(InRenderParams.RenderGroupBitIndices);
+	TArray<FWeightmapPaintLayerInfo> WeightmapPaintLayerInfos;
+	WeightmapPaintLayerInfos.Reserve(LayerInfos.Num());
+	Algo::Transform(LayerInfos, WeightmapPaintLayerInfos, [](ULandscapeLayerInfoObject* InLayerInfo) 
+		{
+			check(InLayerInfo != nullptr); // There should only be valid layer infos at this point
+			FWeightmapPaintLayerInfo WeightmapPaintLayerInfo;
+			if (!InLayerInfo->bNoWeightBlend)
+			{
+				WeightmapPaintLayerInfo.Flags = EWeightmapPaintLayerFlags::IsWeightBlended;
+			}
+			if (InLayerInfo == ALandscapeProxy::VisibilityLayer)
+			{
+				WeightmapPaintLayerInfo.Flags = EWeightmapPaintLayerFlags::IsVisibilityLayer;
+			}
+			return WeightmapPaintLayerInfo;
+		});
+
+	ENQUEUE_RENDER_COMMAND(LandscapeLayers_Cmd_WeightmapsPerformLegacyWeightBlending)(
+		[TargetLayerNames
+		, WeightmapPaintLayerInfos
+		, OutputResource = WriteRT->GetRenderTarget2DArray()->GetResource()
+		, OutputResourceName = WriteRT->GetDebugName()
+		, CurrentEditLayerResource = ReadRT->GetRenderTarget2DArray()->GetResource()
+		, EffectiveTextureSize = InRenderParams.MergeRenderContext->GetCurrentRenderBatch()->GetRenderTargetResolution(/*bInWithDuplicateBorders = */false)]
+		(FRHICommandListImmediate& InRHICmdList)
+	{
+		SCOPED_DRAW_EVENTF(InRHICmdList, LandscapeLayers, TEXT("LegacyWeightBlending -> %s"), OutputResourceName);
+		FRDGBuilder GraphBuilder(InRHICmdList, RDG_EVENT_NAME("WeightmapsPerformLegacyWeightBlending"));
+
+		FRDGTextureRef OutputTextureRef = GraphBuilder.RegisterExternalTexture(CreateRenderTarget(OutputResource->TextureRHI, TEXT("OutputTexture")), ERDGTextureFlags::SkipTracking);
+		FRDGTextureRef CurrentEditLayerTextureRef = GraphBuilder.RegisterExternalTexture(CreateRenderTarget(CurrentEditLayerResource->TextureRHI, TEXT("CurrentEditLayerTexture")));
+		FRDGTextureSRVRef CurrentEditLayerTextureSRVRef = GraphBuilder.CreateSRV(FRDGTextureSRVDesc::Create(CurrentEditLayerTextureRef));
+
+		FRDGBufferRef PaintLayerInfosBuffer = CreateStructuredBuffer(GraphBuilder, TEXT("LandscapePaintLayerInfosBuffer"), WeightmapPaintLayerInfos);
+		FRDGBufferSRVRef PaintLayerInfosBufferSRV = GraphBuilder.CreateSRV(FRDGBufferSRVDesc(PaintLayerInfosBuffer));
+
+		const int32 NumTargetLayers = TargetLayerNames.Num();
+		for (int32 TargetLayerIndex = 0; TargetLayerIndex < NumTargetLayers; ++TargetLayerIndex)
+		{
+			FLandscapeEditLayersWeightmapsPerformLegacyWeightBlendingPS::FParameters* PSParams = GraphBuilder.AllocParameters<FLandscapeEditLayersWeightmapsPerformLegacyWeightBlendingPS::FParameters>();
+			PSParams->RenderTargets[0] = FRenderTargetBinding(OutputTextureRef, ERenderTargetLoadAction::ENoAction, /*InMipIndex = */0, /*InArraySlice = */TargetLayerIndex);
+			PSParams->InPaintLayerIndex = TargetLayerIndex;
+			PSParams->InNumPaintLayers = NumTargetLayers;
+			PSParams->InPaintLayerInfos = PaintLayerInfosBufferSRV;
+			PSParams->InCurrentEditLayerWeightmaps = CurrentEditLayerTextureSRVRef;
+
+			FLandscapeEditLayersWeightmapsPerformLegacyWeightBlendingPS::PerformLegacyWeightBlendingPS(GraphBuilder, PSParams, TargetLayerNames[TargetLayerIndex], EffectiveTextureSize);
+		}
+		GraphBuilder.Execute();
+	});
+}
+
+UE::Landscape::EditLayers::FMergeRenderContext ALandscape::PrepareEditLayersMergeRenderContext(const UE::Landscape::EditLayers::FMergeRenderParams& InMergeRenderParams)
+{
+	TRACE_CPUPROFILER_EVENT_SCOPE(ALandscape::PrepareEditLayersMergeRenderContext);
+
+	using namespace UE::Landscape::EditLayers;
+	using namespace UE::Landscape::EditLayers::Private;
+	using namespace UE::Landscape::Private;
+
+	const int32 VisualLogShowRenderItems = CVarLandscapeBatchedMergeVisualLogShowRenderItemsType.GetValueOnGameThread();
+	const bool bVisualLogShowRenderItemsInput = (VisualLogShowRenderItems == 1) || (VisualLogShowRenderItems == 3);
+	const bool bVisualLogShowRenderItemsOutput = (VisualLogShowRenderItems == 2) || (VisualLogShowRenderItems == 3);
+	const bool bVisualLogShowAllRenderItems = CVarLandscapeBatchedMergeVisualLogShowAllRenderItems.GetValueOnGameThread();
+	const FString VisualLogShowRenderItemsEditLayerRendererFilter = CVarLandscapeBatchedMergeVisualLogShowRenderItemsEditLayerRendererFilter.GetValueOnGameThread();
+	const int32 VisualLogShowComponentDependencies = CVarLandscapeBatchedMergeVisualLogShowComponentDependencies.GetValueOnGameThread();
+	const FString VisualLogShowComponentDependenciesFilter = CVarLandscapeBatchedMergeVisualLogShowComponentDependenciesFilter.GetValueOnGameThread();
+	TOptional<FIntPoint> VisualLogShowComponentDependenciesKey;
+	if (!VisualLogShowComponentDependenciesFilter.IsEmpty())
+	{
+		FIntPoint ComponentKey;
+		if (ComponentKey.InitFromString(VisualLogShowComponentDependenciesFilter))
+		{
+			VisualLogShowComponentDependenciesKey = ComponentKey;
+		}
+		else
+		{
+			UE_LOG(LogLandscape, Warning, TEXT("Cannot parse string \"%s\". Ignoring show component dependencies filter"), *VisualLogShowComponentDependenciesFilter);
+		}
+	}
+
+	ULandscapeInfo* Info = GetLandscapeInfo();
+	check(Info != nullptr);
+	check(!LandscapeEditLayers.IsEmpty()); 
+
+	const FTransform& LandscapeTransform = GetTransform();
+	const FBox LandscapeLoadedBounds = Info->GetLoadedBounds();
+
+	FMergeRenderContext MergeRenderContext(InMergeRenderParams.Landscape, InMergeRenderParams.bIsHeightmapMerge);
+
+	// Start by gathering all possible target layer names on this landscape: 
+	//  This list of all unique target layer names will help accelerate the gathering of output layers on each component (using bit arrays) as well as the target layers intersection tests:
+	if (InMergeRenderParams.bIsHeightmapMerge)
+	{
+		// Only one target layer in the case of heightmap: 
+		const FName TargetLayerName = FName("Height");
+		MergeRenderContext.AllTargetLayerNames = { TargetLayerName };
+		// And it's valid and final :
+		MergeRenderContext.ValidTargetLayerBitIndices = TBitArray<>(true, 1);
+		MergeRenderContext.FinalTargetLayerBitIndices = TBitArray<>(true, 1);
+	}
+	else
+	{
+		// Gather all target layer names and mark those that are valid layers :
+		for (const FLandscapeInfoLayerSettings& LayerSettings : Info->Layers)
+		{ 
+			check(!LayerSettings.LayerName.IsNone());
+			ULandscapeLayerInfoObject*& LayerInfo = MergeRenderContext.WeightmapLayerInfos.Add_GetRef(LayerSettings.LayerInfoObj);
+
+			MergeRenderContext.AllTargetLayerNames.Add(LayerSettings.LayerName);
+			MergeRenderContext.ValidTargetLayerBitIndices.Add(LayerInfo != nullptr);
+			MergeRenderContext.FinalTargetLayerBitIndices.Add(false); // will be computed later...
+		}
+	}
+	const int32 NumAllTargetLayerNames = MergeRenderContext.AllTargetLayerNames.Num();
+	check(MergeRenderContext.ValidTargetLayerBitIndices.Num() == NumAllTargetLayerNames);
+	check(MergeRenderContext.FinalTargetLayerBitIndices.Num() == NumAllTargetLayerNames);
+
+	// InMergeRenderParams.EditLayerRendererStates contains a list of renderers that is not quite final : ULandscapeDefaultEditLayerRenderer is always inserted at the beginning to make sure we always render at 
+	//  least the default value and ULandscapeWeightmapWeightBlendedLayersRenderer can optionally be inserted at the end too :
+	TArray<FEditLayerRendererState> CandidateEditLayerRendererStates;
+	{
+		CandidateEditLayerRendererStates.Reserve(InMergeRenderParams.EditLayerRendererStates.Num() + 2);
+
+		// We always have at least 1 renderer at the start : the default one, whose job is to both provide the default value and act as the final "gatherer" of the component dependencies 
+		//  from all the renderers above (see class comment for more details) :
+		ULandscapeDefaultEditLayerRenderer* InitialEditLayerRenderer = ULandscapeDefaultEditLayerRenderer::StaticClass()->GetDefaultObject<ULandscapeDefaultEditLayerRenderer>();
+		CandidateEditLayerRendererStates.Emplace(InitialEditLayerRenderer, Info);
+
+		// Then append all the provided renderer states
+		CandidateEditLayerRendererStates.Append(InMergeRenderParams.EditLayerRendererStates);
+
+		// Legacy weight-blending requires an additional renderer at the end of the stack, to weight-blend the weightmaps needing it :
+		if (!InMergeRenderParams.bIsHeightmapMerge)
+		{
+			ULandscapeWeightmapWeightBlendedLayersRenderer* WeightmapWeightBlendedLayersRenderer = ULandscapeWeightmapWeightBlendedLayersRenderer::StaticClass()->GetDefaultObject<ULandscapeWeightmapWeightBlendedLayersRenderer>();
+			if (!WeightmapWeightBlendedLayersRenderer->GatherWeightBlendedWeightmapLayerNames(Info).IsEmpty())
+			{
+				CandidateEditLayerRendererStates.Emplace(WeightmapWeightBlendedLayersRenderer, Info);
+			}
+		}
+	}
+
+	// FinalEditLayerRendererStates will contain the renderer states that are actually relevant to this merge :
+	TArray<FEditLayerRendererState> FinalEditLayerRendererStates;
+	FinalEditLayerRendererStates.Reserve(CandidateEditLayerRendererStates.Num());
+
+	// Within each render batch, elements can be processed group by group. For heightmap/visibility, there's only one such group. For weightmaps, there's one group per list of weightmaps
+	//  that need to be processed together for weight-blending. Each group is composed of a list of (weightmap layer) names (it's only a debug name in the case of heightmaps)
+
+	// First, let's work out the weightmaps inter-dependencies (i.e. horizontal dependencies) : weight-blending requires some weightmaps to be processed together: 
+	TArray<TBitArray<>> FinalRenderGroups;
+	if (!InMergeRenderParams.bIsHeightmapMerge)
+	{
+		TRACE_CPUPROFILER_EVENT_SCOPE(PrepareRenderGroups);
+
+		TArray<TArray<TBitArray<>>> RelevantRenderGroupsPerRenderer;
+		TArray<TBitArray<>> RelevantTargetLayerBitIndicesPerRenderer;
+		RelevantRenderGroupsPerRenderer.Reserve(CandidateEditLayerRendererStates.Num());
+		RelevantTargetLayerBitIndicesPerRenderer.Reserve(CandidateEditLayerRendererStates.Num());
+		// Iterate through all renderers to find out what render group they render (only keep those that are relevant for the current state and request) : 
+		for (const FEditLayerRendererState& EditLayerRendererState : CandidateEditLayerRendererStates)
+		{
+			TBitArray<> RendererStateEnabledTargetLayerBitIndices(MergeRenderContext.ConvertTargetLayerNamesToBitIndices(EditLayerRendererState.GetEnabledTargetWeightmaps()));
+			// Retain only the target layer names that are enabled and valid :
+			RendererStateEnabledTargetLayerBitIndices.CombineWithBitwiseAND(MergeRenderContext.ValidTargetLayerBitIndices, EBitwiseOperatorFlags::MinSize);
+
+			// List of all supported render groups for this renderer :
+			TArray<TBitArray<>> RelevantRenderGroupsForThisRenderer;
+			for (const TSet<FName>& RenderGroup : EditLayerRendererState.GetWeightmapRenderGroups())
+			{
+				// Only retain the target layers that have passed the previous test (valid and enabled):
+				TBitArray<> RenderGroupBitIndices = MergeRenderContext.ConvertTargetLayerNamesToBitIndices(RenderGroup.Array());
+				RenderGroupBitIndices.CombineWithBitwiseAND(RendererStateEnabledTargetLayerBitIndices, EBitwiseOperatorFlags::MinSize);
+				if (RenderGroupBitIndices.Find(true) != INDEX_NONE)
+				{
+					checkf(Algo::NoneOf(RelevantRenderGroupsForThisRenderer, [&RenderGroupBitIndices](const TBitArray<>& InOtherRenderGroupBitIndices) 
+						{ 
+							return (TBitArray<>::BitwiseAND(InOtherRenderGroupBitIndices, RenderGroupBitIndices, EBitwiseOperatorFlags::MinSize).Find(true) != INDEX_NONE);
+						}), 
+						TEXT("All of the target layers returned by the renderer must belong to 1 render group of this renderer and 1 only"));
+					RelevantRenderGroupsForThisRenderer.Add(RenderGroupBitIndices);
+				}
+			}
+
+			RelevantRenderGroupsPerRenderer.Add(RelevantRenderGroupsForThisRenderer);
+			RelevantTargetLayerBitIndicesPerRenderer.Add(RendererStateEnabledTargetLayerBitIndices);
+		}
+		check(CandidateEditLayerRendererStates.Num() == RelevantRenderGroupsPerRenderer.Num());
+		check(CandidateEditLayerRendererStates.Num() == RelevantTargetLayerBitIndicesPerRenderer.Num());
+
+		const int32 CandidateNumRenderers = CandidateEditLayerRendererStates.Num();
+		TBitArray<> RequestedTargetLayerBitIndices = MergeRenderContext.ConvertTargetLayerNamesToBitIndices(InMergeRenderParams.WeightmapLayerNames.Array());
+		// No need to retain the invalid target layers : 
+		RequestedTargetLayerBitIndices.CombineWithBitwiseAND(MergeRenderContext.ValidTargetLayerBitIndices, EBitwiseOperatorFlags::MinSize);
+		// Early-out when there's nothing to do :
+		if (RequestedTargetLayerBitIndices.Find(true) == INDEX_NONE)
+		{
+			return MergeRenderContext;
+		}
+
+		FinalRenderGroups.Reserve(MergeRenderContext.ValidTargetLayerBitIndices.CountSetBits());
+		// Start with minimal render groups : one per requested target layer : 
+		for (TConstSetBitIterator It(RequestedTargetLayerBitIndices); It; ++It)
+		{
+			TBitArray<> RenderGroup(false, NumAllTargetLayerNames);
+			RenderGroup[It.GetIndex()] = true;
+			FinalRenderGroups.Add(MoveTemp(RenderGroup));
+		}
+
+		// Then iterate in reverse order on renderers to trace the dependency of each of their render groups towards one another and move target layers from one render group to another as we discover new dependencies:
+		MergeRenderContext.FinalTargetLayerBitIndices = RequestedTargetLayerBitIndices;
+		{
+			TRACE_CPUPROFILER_EVENT_SCOPE(AnalyzeRenderersForRenderGroups);
+
+			for (int32 CurrentRendererIndex = CandidateNumRenderers - 1; CurrentRendererIndex >= 0; --CurrentRendererIndex)
+			{
+				const TArray<TBitArray<>>& RendererRenderGroups = RelevantRenderGroupsPerRenderer[CurrentRendererIndex];
+				const int32 NumRenderGroups = RendererRenderGroups.Num();
+				for (int32 RenderGroupIndex = 0; RenderGroupIndex < NumRenderGroups; ++RenderGroupIndex)
+				{
+					const TBitArray<>& RenderGroup = RendererRenderGroups[RenderGroupIndex];
+
+					TBitArray<> NewMergedRenderGroup(false, NumAllTargetLayerNames);
+					// In all of the final render groups, find the ones that have a layer in common with this render group and merge them all into a single new one
+					FinalRenderGroups.RemoveAllSwap([&NewMergedRenderGroup, &RenderGroup](const TBitArray<>& InFinalRenderGroup)
+						{
+							const bool bShouldMerge = (InFinalRenderGroup != RenderGroup)
+								&& (TBitArray<>::BitwiseAND(InFinalRenderGroup, RenderGroup, EBitwiseOperatorFlags::MinSize).Find(true) != INDEX_NONE);
+							if (bShouldMerge)
+							{
+								NewMergedRenderGroup.CombineWithBitwiseOR(InFinalRenderGroup, EBitwiseOperatorFlags::MinSize);
+							}
+							return bShouldMerge;
+						});
+					// Now add it back to the list of final render groups if it's valid, so that it can be merged by the next renderer if required :
+					if (NewMergedRenderGroup.CountSetBits() > 0)
+					{
+						FinalRenderGroups.Add(NewMergedRenderGroup);
+						MergeRenderContext.FinalTargetLayerBitIndices.CombineWithBitwiseOR(NewMergedRenderGroup, EBitwiseOperatorFlags::MinSize);
+					}
+				}
+			}
+		}
+
+		{
+			TRACE_CPUPROFILER_EVENT_SCOPE(FinalizeRenderGroups);
+
+			// Now that we have the list of all target layers required for rendering the requested ones, we trim the list of edit layer renderers that just don't do anything with those:
+			for (int32 RendererIndex = 0; RendererIndex < CandidateNumRenderers; ++RendererIndex)
+			{
+				const TBitArray<> RelevantTargetLayerBitIndices = RelevantTargetLayerBitIndicesPerRenderer[RendererIndex];
+				if (TBitArray<>::BitwiseAND(RelevantTargetLayerBitIndices, MergeRenderContext.FinalTargetLayerBitIndices, EBitwiseOperatorFlags::MinSize).Find(true) != INDEX_NONE)
+				{
+					FinalEditLayerRendererStates.Add(CandidateEditLayerRendererStates[RendererIndex]);
+				}
+			}
+
+			// TODO [jonathan.bard] : revisit : this shouldn't be true now that we support additively resolving channels with blending : 
+#if !SUPPORTS_LANDSCAPE_EDITORONLY_UBER_MATERIAL
+			// When 4 weightmaps are packed in 1 RGBA channel, we cannot guarantee that weightmaps will be fully resolvable when a group is done, since there's no guarantee that a given component's 4 
+			//  allocations will be contained in that group. So we create a single group instead, that contains all weightmap layers. This consumes more memory since we need 3 texture arrays for a batch and 
+			//  the number of weightmaps in a group is the number of slices of the array : 
+			FinalRenderGroups = { MergeRenderContext.FinalTargetLayerBitIndices };
+#endif // !SUPPORTS_LANDSCAPE_EDITORONLY_UBER_MATERIAL
+		}
+	}
+	else
+	{
+		// For heightmap, there's no other dependency so simply trim all renderers that are not enabled : 
+		FinalEditLayerRendererStates = CandidateEditLayerRendererStates.FilterByPredicate([&InMergeRenderParams](const FEditLayerRendererState& InRendererState)
+		{
+			return InRendererState.IsTargetEnabled(ELandscapeToolTargetType::Heightmap);
+		});
+
+		// Only one group in the case of heightmap: 
+		const FName TargetLayerName = FName("Height");
+		FinalRenderGroups = { TBitArray<>(true, 1) };
+	}
+
+	const int32 FinalNumRenderers = FinalEditLayerRendererStates.Num();
+	
+	// Early-out when there's nothing to do :
+	if (FinalNumRenderers == 0)
+	{
+		return MergeRenderContext;
+	}
+
+	const ELandscapeToolTargetTypeFlags MergeTypeMask = InMergeRenderParams.bIsHeightmapMerge ? ELandscapeToolTargetTypeFlags::Heightmap : (ELandscapeToolTargetTypeFlags::Weightmap | ELandscapeToolTargetTypeFlags::Visibility);
+
+	FLandscapeComponent2DIndexer Component2DIndexer = CreateLandscapeComponent2DIndexer(Info);
+	const TArray<ULandscapeComponent*> AllComponents = Component2DIndexer.GetAllValues();
+	const TBitArray<> AllValidComponentBitIndices = Component2DIndexer.GetValidValueBitIndices();
+	const int32 NumAllComponents = AllComponents.Num();
+
+	// The list of all components that will end up being rendered across all renderers (one bit per component) :
+	TBitArray<> FinalComponentsToRenderInfoBitIndices(false, NumAllComponents);
+
+	// Pre-allocate a working list of all landscape components render info. Some of which we might not end up rendering, but at least, that allows to associate a component with an index, 
+	//  which allows to turn intersection/union of components (which we do a lot in this function) into simple bit array bitwise AND/OR operations : 
+	TArray<FComponentToRenderInfo> AllComponentsToRenderInfos;
+	TArray<FEditLayerRendererRenderInfo> OrderedEditLayerRendererRenderInfos;
+
+#if ENABLE_VISUAL_LOG
+	// Helper for debugging component dependencies : only if the CVar requires it
+	TOptional<FComponentDependenciesVisLogHelper> VisualLogDependencyHelper;
+	UE_IFVLOG(
+			if ((MergeRenderContext.IsVisualLogEnabled()) && ((VisualLogShowComponentDependencies > 0) || VisualLogShowComponentDependenciesKey.IsSet()))
+			{
+				// Force the display of all info when we show the dependencies of one component in particular : 
+				const FComponentDependenciesVisLogHelper::EShowNodeInfo ShowNodeInfo = VisualLogShowComponentDependenciesKey.IsSet() 
+					? FComponentDependenciesVisLogHelper::EShowNodeInfo::Detailed
+					: static_cast<FComponentDependenciesVisLogHelper::EShowNodeInfo>(VisualLogShowComponentDependencies);
+				VisualLogDependencyHelper.Emplace(this, InMergeRenderParams.bIsHeightmapMerge, ShowNodeInfo, MergeRenderContext);
+			}
+		);
+	auto VisLogDependency = [&VisualLogDependencyHelper, &VisualLogShowComponentDependenciesKey, &AllComponentsToRenderInfos, &OrderedEditLayerRendererRenderInfos]
+		(int32 InSourceComponentIndex, int32 InSourceRendererIndex, int32 InDestinationComponentIndex, int32 InDestinationRendererIndex)
+		{
+			if (VisualLogDependencyHelper.IsSet() && (InSourceRendererIndex >= 0) && VisualLogShowComponentDependenciesKey.IsSet())
+			{
+				VisualLogDependencyHelper->AddDependency(AllComponentsToRenderInfos[InSourceComponentIndex], OrderedEditLayerRendererRenderInfos[InSourceRendererIndex],
+					AllComponentsToRenderInfos[InDestinationComponentIndex], OrderedEditLayerRendererRenderInfos[InDestinationRendererIndex]);
+			}
+		};
+#endif // ENABLE_VISUAL_LOG
+
+	{
+		TRACE_CPUPROFILER_EVENT_SCOPE(PrepareRenderersAnalysis);
+
+		AllComponentsToRenderInfos.AddDefaulted(NumAllComponents);
+		for (TConstSetBitIterator It(AllValidComponentBitIndices); It; ++It)
+		{
+			const int32 ComponentIndex = It.GetIndex();
+			AllComponentsToRenderInfos[ComponentIndex] = FComponentToRenderInfo(AllComponents[ComponentIndex], ComponentIndex, NumAllComponents, NumAllTargetLayerNames);
+			MergeRenderContext.MaxLocalHeight = FMath::Max(MergeRenderContext.MaxLocalHeight, AllComponentsToRenderInfos[ComponentIndex].LocalBounds.Max.Z);
+		}
+
+		// Initiate the process by flipping a bit for each component to merge, for the last renderer in the stack, then we'll register the dependencies 
+		//  between components on this renderer and those on the previous renderer by iterating on renderers in reverse stack order : 
+		for (ULandscapeComponent* Component : InMergeRenderParams.ComponentsToMerge)
+		{
+			const int32 ComponentIndex = Component2DIndexer.GetValueIndexChecked(Component);
+			FinalComponentsToRenderInfoBitIndices[ComponentIndex] = true;
+		}
+
+		// Prepare the render infos of all these renderers : 
+		OrderedEditLayerRendererRenderInfos.Reserve(FinalNumRenderers);
+		for (int32 RendererIndex = 0; RendererIndex < FinalNumRenderers; ++RendererIndex)
+		{
+			OrderedEditLayerRendererRenderInfos.Emplace(FinalEditLayerRendererStates[RendererIndex], RendererIndex, LandscapeTransform, MergeRenderContext.MaxLocalHeight, NumAllComponents);
+		}
+	}
+
+	// Iterate over all renderers in inverse order to compute which landscape component needs to be included in the render. This way, the renderers on top are able to 
+	//  request potentially more components from renderers underneath (e.g. if the renderer performs a blur, it will require an area around the component's area, thus additional components, potentially, 
+	//  which will then request potentially more components on the renderer underneath, etc.)
+	// For each renderer we'll only iterate on nodes that correspond to it. Since they're added in reverse order, this is just a matter of starting the iteration from the first of the renderer's nodes :
+	{
+		TRACE_CPUPROFILER_EVENT_SCOPE(AnalyzeRenderers);
+		for (int32 CurrentRendererIndex = FinalNumRenderers - 1; CurrentRendererIndex >= 0; --CurrentRendererIndex)
+		{
+			FEditLayerRendererRenderInfo& EditLayerRendererRenderInfo = OrderedEditLayerRendererRenderInfos[CurrentRendererIndex];
+
+			const FString EditLayerRendererDebugName = EditLayerRendererRenderInfo.RendererState.GetRenderer()->GetEditLayerRendererDebugName();
+			const bool bPassesVisualLogRenderItemFilter = VisualLogShowRenderItemsEditLayerRendererFilter.IsEmpty() || EditLayerRendererDebugName.Contains(VisualLogShowRenderItemsEditLayerRendererFilter);
+
+			// This renderer will eventually append new components to render along the way so we add them to a new list and swap at the end : 
+			TBitArray<> ComponentsToRenderInfoBitIndicesAfterRenderer = FinalComponentsToRenderInfoBitIndices;
+
+			// Intersect this renderer's render items' outputs with the components to render to find out which ones should participate to the merge :
+			for (const FEditLayerRenderItem& RenderItem : EditLayerRendererRenderInfo.RendererState.GetRenderer()->GetRenderItems(Info))
+			{
+				const FEditLayerTargetTypeState& RenderItemTargetTypeState = RenderItem.GetTargetTypeState();
+				checkf(RenderItemTargetTypeState.Intersect(EditLayerRendererRenderInfo.RendererState.GetSupportedTargetTypeState()) == RenderItemTargetTypeState,
+					TEXT("One of edit layer renderer %s's render items target type state is not supported by the renderer's supported target type state. \
+					This is invalid : the renderer's supported target type state should be a superset of its render item's target type state. \n\
+					- Render item state: %s \n\
+					- Renderer supported state: %s \n"), *EditLayerRendererDebugName, *RenderItemTargetTypeState.ToString(), *EditLayerRendererRenderInfo.RendererState.GetSupportedTargetTypeState().ToString());
+
+				bool bOutputRenderItem = false;
+				const FOutputWorldArea& OutputWorldArea = RenderItem.GetOutputWorldArea();
+				// Only consider render items which affect the type of merge being performed: 
+				const bool bRenderItemAffectsMergeType = EnumHasAnyFlags(RenderItemTargetTypeState.GetTargetTypeMask(), MergeTypeMask);
+				if (bRenderItemAffectsMergeType)
+				{
+					switch (OutputWorldArea.GetType())
+					{
+					case FOutputWorldArea::EType::LocalComponent:
+					{
+						// This render item writes to every component : 
+						bOutputRenderItem = true;
+						break;
+					}
+					case FOutputWorldArea::EType::SpecificComponent:
+					{
+						// This render item writes to only 1 component :
+						if (int32 ComponentIndex = Component2DIndexer.GetValueIndexForKeySafe(OutputWorldArea.GetSpecificComponentKey()); ComponentIndex != INDEX_NONE)
+						{
+							// Retain the render item if the component it renders to overlaps with one the components to render :
+							bOutputRenderItem = FinalComponentsToRenderInfoBitIndices[ComponentIndex];
+						}
+						break;
+					}
+					case FOutputWorldArea::EType::OOBox:
+					{
+						FIntRect ComponentIndicesBoundingRect;
+						TMap<FIntPoint, ULandscapeComponent*> DependentComponents;
+						// TODO [jonathan.bard] : Change this to OOBB to OOBB test to cull more components : look at FOrientedBoxHelpers
+						const FOOBox2D& OOBox = OutputWorldArea.GetOOBox();
+						Info->GetOverlappedComponents(OOBox.Transform, FBox2D(-OOBox.Extents / 2.0, OOBox.Extents / 2.0), DependentComponents, ComponentIndicesBoundingRect);
+						for (auto It = DependentComponents.CreateConstIterator(); It; ++It && !bOutputRenderItem)
+						{
+							int32 ComponentIndex = Component2DIndexer.GetValueIndexForKeyChecked(It.Key());
+							// Retain the render item if one of the components it renders to overlaps with one the components to render :
+							bOutputRenderItem = FinalComponentsToRenderInfoBitIndices[ComponentIndex];
+						}
+						break;
+					}
+					default:
+						check(false);
+					}
+				}
+
+				// Mark which output layers of the component this render item will only need to affect :
+				// For a weightmap merge, it's possible the render item will only end up modifying the existing ones (as opposed to "generating" new ones).
+				//  In that case, restrain the render item from modifying the target layer mask for this component :
+				const bool bAffectsOutputLayerBitIndices = MergeRenderContext.IsHeightmapMerge() || !RenderItem.GetModifyExistingWeightmapsOnly();
+
+				UE_IFVLOG(
+					{
+						if (MergeRenderContext.IsVisualLogEnabled() && bRenderItemAffectsMergeType && bPassesVisualLogRenderItemFilter)
+						{
+							if (bVisualLogShowRenderItemsInput && (bOutputRenderItem || bVisualLogShowAllRenderItems))
+							{
+								TArray<ULandscapeComponent*> AllComponentsToVisLog = bVisualLogShowAllRenderItems ? Component2DIndexer.GetValidValues() : Component2DIndexer.GetValidValuesForBitIndices(FinalComponentsToRenderInfoBitIndices);
+								VisLogRenderItemInput(this, RenderItem.GetInputWorldArea(), EditLayerRendererRenderInfo, LandscapeTransform, LandscapeLoadedBounds, MakeArrayView(AllComponentsToVisLog));
+							}
+
+							if (bVisualLogShowRenderItemsOutput && (bOutputRenderItem || bVisualLogShowAllRenderItems))
+							{
+								TArray<ULandscapeComponent*> AllComponentsToVisLog = bVisualLogShowAllRenderItems ? Component2DIndexer.GetValidValues() : Component2DIndexer.GetValidValuesForBitIndices(FinalComponentsToRenderInfoBitIndices);
+								VisLogRenderItemOutput(this, MergeRenderContext.IsHeightmapMerge(), bAffectsOutputLayerBitIndices, MakeArrayView(RenderItemTargetTypeState.GetActiveWeightmaps()),
+									OutputWorldArea, EditLayerRendererRenderInfo, LandscapeTransform, LandscapeLoadedBounds, MakeArrayView(AllComponentsToVisLog));
+							}
+						}
+					});
+
+				if (bOutputRenderItem)
+				{
+					TBitArray<> InputComponentBitIndices(false, NumAllComponents);
+					const TBitArray<>* ComponentsToIterateBitIndices = nullptr;
+
+					const FInputWorldArea& InputWorldArea = RenderItem.GetInputWorldArea();
+					switch (InputWorldArea.GetType())
+					{
+					case FInputWorldArea::EType::Infinite:
+					{
+						InputComponentBitIndices = Component2DIndexer.GetValidValueBitIndices();
+						ComponentsToIterateBitIndices = &InputComponentBitIndices;
+						break;
+					}
+					case FInputWorldArea::EType::LocalComponent:
+					{
+						// This render item requires the component itself and potentially its neighbors, so we need to iterate on all the components currently being processed : 
+						ComponentsToIterateBitIndices = &FinalComponentsToRenderInfoBitIndices;
+						break;
+					}
+					case FInputWorldArea::EType::SpecificComponent:
+					{
+						InputComponentBitIndices = Component2DIndexer.GetValidValueBitIndicesInBounds(InputWorldArea.GetSpecificComponentKeys(), /* bInInclusiveBounds = */true);
+						ComponentsToIterateBitIndices = &InputComponentBitIndices;
+						break;
+					}
+					case FInputWorldArea::EType::OOBox:
+					{
+						FIntRect ComponentIndicesBoundingRect;
+						TMap<FIntPoint, ULandscapeComponent*> DependentComponents;
+						// TODO [jonathan.bard] : Change this to OOBB to OOBB test to cull more components
+						const FOOBox2D& OOBox = InputWorldArea.GetOOBox();
+						Info->GetOverlappedComponents(OOBox.Transform, FBox2D(-OOBox.Extents / 2.0, OOBox.Extents / 2.0), DependentComponents, ComponentIndicesBoundingRect);
+						for (auto It : DependentComponents)
+						{
+							int32 ComponentIndex = Component2DIndexer.GetValueIndexForKeyChecked(It.Key);
+							InputComponentBitIndices[ComponentIndex] = true;
+						}
+						ComponentsToIterateBitIndices = &InputComponentBitIndices;
+						break;
+					}
+					default:
+						check(false);
+					}
+					check(ComponentsToIterateBitIndices != nullptr);
+
+					// List all target layers written by this render item :
+					TBitArray<> OutputLayerBitIndices;
+					if (bAffectsOutputLayerBitIndices)
+					{
+						if (MergeRenderContext.IsHeightmapMerge())
+						{
+							// Only one target layer in the case of a heightmap merge
+							check(MergeRenderContext.AllTargetLayerNames.Num() == 1);
+							OutputLayerBitIndices.SetNum(1, true);
+						}
+						else
+						{
+							OutputLayerBitIndices = MergeRenderContext.ConvertTargetLayerNamesToBitIndices(RenderItemTargetTypeState.GetActiveWeightmaps());
+						}
+					}
+
+					// Iterate on all the required components :
+					for (TConstSetBitIterator It(*ComponentsToIterateBitIndices); It; ++It)
+					{
+						const int32 ComponentToRenderIndex = It.GetIndex();
+						FComponentToRenderInfo& ComponentToRenderInfo = AllComponentsToRenderInfos[ComponentToRenderIndex];
+						check(ComponentToRenderInfo.Component != nullptr);
+
+						// Add the render item's target layers to the component's own : 
+						if (bAffectsOutputLayerBitIndices)
+						{
+							ComponentToRenderInfo.TargetLayerBitIndices.CombineWithBitwiseOR(OutputLayerBitIndices, EBitwiseOperatorFlags::MinSize);
+						}
+
+						// Special case for FInputWorldArea::EType::LocalComponent, where the input components are specific to the component being iterated : 
+						if (InputWorldArea.GetType() == FInputWorldArea::EType::LocalComponent)
+						{
+							InputComponentBitIndices = Component2DIndexer.GetValidValueBitIndicesInBounds(InputWorldArea.GetLocalComponentKeys(ComponentToRenderInfo.Component->GetComponentKey()), /* bInInclusiveBounds = */true);
+						}
+
+						// There should always be a dependency between this renderer and the previous in the stack for the component itself : 
+						InputComponentBitIndices[ComponentToRenderIndex] = true;
+
+						// If these components are not yet in the final list of components to render, add them: 
+						ComponentsToRenderInfoBitIndicesAfterRenderer.CombineWithBitwiseOR(InputComponentBitIndices, EBitwiseOperatorFlags::MinSize);
+
+						// Add these components to the list that this renderer needs to render: 
+						EditLayerRendererRenderInfo.ComponentToRenderInfoBitIndices.CombineWithBitwiseOR(InputComponentBitIndices, EBitwiseOperatorFlags::MinSize);
+
+						// Finally add these components as dependencies to the component we're trying to render : 
+						ComponentToRenderInfo.DependentComponentBitIndices.CombineWithBitwiseOR(InputComponentBitIndices, EBitwiseOperatorFlags::MinSize);
+
+#if ENABLE_VISUAL_LOG
+						if (VisualLogDependencyHelper.IsSet())
+						{
+							for (TConstSetBitIterator ItInputComponent(InputComponentBitIndices); ItInputComponent; ++ItInputComponent)
+							{
+								// Register a dependency from the component we want to render towards all of the components its input area overlaps with on the previous renderer : 
+								const int32 DependentComponentIndex = ItInputComponent.GetIndex();
+								VisLogDependency(DependentComponentIndex, CurrentRendererIndex - 1, ComponentToRenderIndex, CurrentRendererIndex);
+							}
+						}
+#endif // ENABLE_VISUAL_LOG
+					}
+				}
+				else // !bOutputRenderItem
+				{
+#if ENABLE_VISUAL_LOG
+					// Declare a passthrough dependency between each component on this renderer to the next, to display the full chain of dependencies : 
+					if (VisualLogDependencyHelper.IsSet())
+					{
+						for (TConstSetBitIterator It(ComponentsToRenderInfoBitIndicesAfterRenderer); It; ++It)
+						{
+							const int32 ComponentToRenderIndex = It.GetIndex();
+							VisLogDependency(ComponentToRenderIndex, CurrentRendererIndex - 1, ComponentToRenderIndex, CurrentRendererIndex);
+						}
+					}
+#endif // ENABLE_VISUAL_LOG
+				}
+			}
+
+			// The renderer has been fully processed, now we can update the list of components to render for the next renderer in line : 
+			Swap(ComponentsToRenderInfoBitIndicesAfterRenderer, FinalComponentsToRenderInfoBitIndices);
+		}
+	}
+
+	for (TConstSetBitIterator It(FinalComponentsToRenderInfoBitIndices); It; ++It)
+	{
+		const int32 ComponentIndex = It.GetIndex();
+		FComponentToRenderInfo& ComponentToRenderInfo = AllComponentsToRenderInfos[ComponentIndex];
+		// Now compute the bounds to finalize this component render info (it's faster to do via Component2DIndexer than iterating through components) : 
+		FIntRect DependentComponentsInclusiveBounds = Component2DIndexer.GetValidValuesBoundsForBitIndices(ComponentToRenderInfo.DependentComponentBitIndices, /*bInInclusiveBounds = */true);
+		ComponentToRenderInfo.Finalize(DependentComponentsInclusiveBounds, ComponentSizeQuads);
+	}
+
+#if ENABLE_VISUAL_LOG
+	if (VisualLogDependencyHelper.IsSet() && (VisualLogShowComponentDependencies > 0))
+	{
+		// Display a node for every component that will be rendered in the end: 
+		for (TConstSetBitIterator It(FinalComponentsToRenderInfoBitIndices); It; ++It)
+		{
+			VisualLogDependencyHelper->AddNode(AllComponentsToRenderInfos[It.GetIndex()], OrderedEditLayerRendererRenderInfos[0]);
+		}
+	}
+#endif // ENABLE_VISUAL_LOG
+
+	// Now divide the work into batches as large as possible (but fitting in the desired max batch resolution, if possible): 
+	TArray<FComponentToRenderBatchInfo> AllComponentToRenderBatchInfos = DivideIntoBatches(FinalComponentsToRenderInfoBitIndices, AllComponentsToRenderInfos);
+	{
+		TRACE_CPUPROFILER_EVENT_SCOPE(PrepareBatches);
+		MergeRenderContext.RenderBatches.Reserve(AllComponentToRenderBatchInfos.Num());
+		MergeRenderContext.TargetLayersToComponents.AddDefaulted(MergeRenderContext.AllTargetLayerNames.Num());
+		for (const FComponentToRenderBatchInfo& ComponentToRenderBatchInfo : AllComponentToRenderBatchInfos)
+		{
+			int32 MergeBatchIndex = MergeRenderContext.RenderBatches.Num();
+			FMergeRenderBatch& MergeRenderBatch = MergeRenderContext.RenderBatches.Emplace_GetRef();
+			MergeRenderBatch.Landscape = this;
+			MergeRenderBatch.SectionRect = ComponentToRenderBatchInfo.CombinedSectionRect;
+			MergeRenderBatch.MinComponentKey = MergeRenderBatch.MinComponentKey.ComponentMin(ComponentToRenderBatchInfo.MinComponentKey);
+			MergeRenderBatch.MaxComponentKey = MergeRenderBatch.MaxComponentKey.ComponentMax(ComponentToRenderBatchInfo.MaxComponentKey);
+			// We'll reuse the same merge render targets in order to generate the mips, which include the duplicate borders, so we need to expand the render target's size to accommodate for this : 
+			MergeRenderBatch.Resolution = (MergeRenderBatch.MaxComponentKey - MergeRenderBatch.MinComponentKey + 1) * NumSubsections * (SubsectionSizeQuads + 1);
+			MergeRenderBatch.TargetLayersToComponents.AddDefaulted(MergeRenderContext.AllTargetLayerNames.Num());
+			MergeRenderBatch.TargetLayerNameBitIndices.SetNum(MergeRenderContext.AllTargetLayerNames.Num(), false);
+
+			MergeRenderContext.MaxNeededResolution = MergeRenderContext.MaxNeededResolution.ComponentMax(MergeRenderBatch.Resolution);
+		
+			// For each renderer, find the list of components actually needed and store that in a separate array, to avoid recomputing it for each target group : 
+			int32 LocalNumRenderers = OrderedEditLayerRendererRenderInfos.Num();
+			TArray<TArray<ULandscapeComponent*>> ComponentsForRenderers;
+			ComponentsForRenderers.SetNum(LocalNumRenderers);
+			for (int32 EditLayerRendererIndex = 0; EditLayerRendererIndex < LocalNumRenderers; ++EditLayerRendererIndex)
+			{
+				const FEditLayerRendererRenderInfo& EditLayerRendererRenderInfo = OrderedEditLayerRendererRenderInfos[EditLayerRendererIndex];
+				// Find out all components that are in common between the renderer's list and the batch's list : only these will need to be rendered in that render step : 
+				TBitArray<> ComponentToRenderInBatchBitIndices = TBitArray<>::BitwiseAND(EditLayerRendererRenderInfo.ComponentToRenderInfoBitIndices, ComponentToRenderBatchInfo.ComponentToRenderInfoBitIndices, EBitwiseOperatorFlags::MinSize);
+				TArray<ULandscapeComponent*>& ComponentsToRender = ComponentsForRenderers[EditLayerRendererIndex];
+				ComponentsToRender.Reserve(ComponentToRenderInBatchBitIndices.CountSetBits());
+				// Transform the bit indices back into a proper component list : 
+				for (TConstSetBitIterator ItComponent(ComponentToRenderBatchInfo.ComponentToRenderInfoBitIndices); ItComponent; ++ItComponent)
+				{
+					FComponentToRenderInfo& ComponentToRenderInfo = AllComponentsToRenderInfos[ItComponent.GetIndex()];
+					ComponentsToRender.Add(ComponentToRenderInfo.Component);
+
+					// Inform the render batch and context of the target layer names associated with each component :
+
+					// Declare these target layers as being in use for the batch :
+					MergeRenderBatch.TargetLayerNameBitIndices.CombineWithBitwiseOR(ComponentToRenderInfo.TargetLayerBitIndices, EBitwiseOperatorFlags::MinSize);
+					MergeRenderBatch.ComponentToTargetLayerBitIndices.FindOrAdd(ComponentToRenderInfo.Component).CombineWithBitwiseOR(ComponentToRenderInfo.TargetLayerBitIndices, EBitwiseOperatorFlags::MaxSize); // Use EBitwiseOperatorFlags::MaxSize here in order to allocate NumTargetLayerNames entries to the resulting bit array in case FindOrAdd is an add
+					MergeRenderContext.ComponentToTargetLayerBitIndices.FindOrAdd(ComponentToRenderInfo.Component).CombineWithBitwiseOR(ComponentToRenderInfo.TargetLayerBitIndices, EBitwiseOperatorFlags::MaxSize); // Use EBitwiseOperatorFlags::MaxSize here in order to allocate NumTargetLayerNames entries to the resulting bit array in case FindOrAdd is an add
+					for (TConstSetBitIterator ItTargetLayer(ComponentToRenderInfo.TargetLayerBitIndices); ItTargetLayer; ++ItTargetLayer)
+					{
+						const int32 TargetLayerIndex = ItTargetLayer.GetIndex();
+						MergeRenderBatch.TargetLayersToComponents[TargetLayerIndex].Add(ComponentToRenderInfo.Component);
+						MergeRenderContext.TargetLayersToComponents[TargetLayerIndex].Add(ComponentToRenderInfo.Component);
+					}
+				}
+
+				// Declare these components as being in use for the batch : 
+				MergeRenderBatch.ComponentsToRender.Append(ComponentsToRender);
+			}
+
+			// Now, we have all the info to build our list of successive render steps : process group by group :
+			for (const TBitArray<>& RenderGroup : FinalRenderGroups)
+			{
+				// Execute the required operations for the whole stack of renderers for this group :
+				for (int32 RendererIndex = 0; RendererIndex < LocalNumRenderers; ++RendererIndex)
+				{
+					const FEditLayerRendererRenderInfo& EditLayerRendererRenderInfo = OrderedEditLayerRendererRenderInfos[RendererIndex];
+					TArrayView<ULandscapeComponent*> ComponentsToRender = MakeArrayView(ComponentsForRenderers[RendererIndex]);
+
+					if (!ComponentsToRender.IsEmpty() && EditLayerRendererRenderInfo.RendererState.GetRenderer()->CanRender())
+					{
+						// Render the content of this render layer : 
+						MergeRenderBatch.RenderSteps.Emplace(FMergeRenderStep::EType::RenderLayer, EditLayerRendererRenderInfo.RendererState, RenderGroup, ComponentsToRender);
+					}
+				}
+
+				// Finally, signal the group as done for this batch
+				MergeRenderBatch.RenderSteps.Emplace(FMergeRenderStep::EType::SignalBatchMergeGroupDone, FEditLayerRendererState(), RenderGroup, MergeRenderBatch.ComponentsToRender.Array());
+
+				MergeRenderContext.MaxNeededNumSlices = FMath::Max(MergeRenderContext.MaxNeededNumSlices, RenderGroup.CountSetBits());
+			}
+		}
+
+		// Sort the batches for easier debugging : 
+		MergeRenderContext.RenderBatches.Sort();
+	}
+
+	return MergeRenderContext;
+}
+
 int32 ALandscape::PerformLayersHeightmapsGlobalMerge(const FUpdateLayersContentContext& InUpdateLayersContentContext, const FEditLayersHeightmapMergeParams& InMergeParams)
@@ -4941,19 +7257,19 @@
 			for (int32 i = 0; i < Layer.Brushes.Num(); ++i)
 			{
 				// TODO: handle conversion from float to RG8 by using material params to write correct values
-				// TODO: handle conversion/handling of RT not same size as internal size
 
 				FLandscapeLayerBrush& Brush = Layer.Brushes[i];
-				FLandscapeBrushParameters BrushParameters(ELandscapeToolTargetType::Heightmap, CombinedHeightmapNonAtlasRT);
+				FLandscapeBrushParameters BrushParameters = BuildLandscapeBrushParameters(/*bIsHeightmapMerge = */true, LandscapeExtent, CombinedHeightmapNonAtlasRT);
 
-				UTextureRenderTarget2D* BrushOutputNonAtlasRT = Brush.RenderLayer(LandscapeExtent, BrushParameters);
-				if ((BrushOutputNonAtlasRT == nullptr) 
-					|| (BrushOutputNonAtlasRT->SizeX != CombinedHeightmapNonAtlasRT->SizeX)
-					|| (BrushOutputNonAtlasRT->SizeY != CombinedHeightmapNonAtlasRT->SizeY))
+				UTextureRenderTarget2D* BrushOutputNonAtlasRT = Brush.RenderLayer(BrushParameters);
+				if (BrushOutputNonAtlasRT == nullptr) 
 				{
 					continue;
 				}
 
+				// TODO: handle conversion/handling of RT not same size as internal size
+				check((BrushOutputNonAtlasRT->SizeX == CombinedHeightmapNonAtlasRT->SizeX) && (BrushOutputNonAtlasRT->SizeY == CombinedHeightmapNonAtlasRT->SizeY));
+
 				ALandscapeBlueprintBrushBase* LandscapeBrush = Brush.GetBrush();
 				check(LandscapeBrush != nullptr); // If we managed to render, the brush should be valid
 
@@ -4990,11 +7306,7 @@
 				FLandscapeBrushParameters BrushParameters = BuildLandscapeBrushParameters(
 					/*bIsHeightmapMerge = */true, LandscapeExtent, CombinedHeightmapNonAtlasRT);
 
-				UTextureRenderTarget2D* BrushOutputNonAtlasRT = RenderCallAdapter->RenderAsBlueprintBrush(
-					BrushParameters, 
-					LandscapeEditLayersLocals::GetLandscapeExtentTransform(*this, LandscapeExtent),
-					// TODO: Someday this will be part of BrushParameters
-					LandscapeExtent.Max - LandscapeExtent.Min);
+				UTextureRenderTarget2D* BrushOutputNonAtlasRT = RenderCallAdapter->RenderAsBlueprintBrush(BrushParameters, GetTransform());
 
 				// Do the same conditional copy that we do for blueprint brushes
 				if (BrushOutputNonAtlasRT != CombinedHeightmapNonAtlasRT
@@ -5026,7 +7338,7 @@
 	// Broadcast Event of the Full Render
 	if ((InMergeParams.HeightmapUpdateModes & Update_Heightmap_All) == Update_Heightmap_All)
 	{
-		LandscapeFullHeightmapRenderDoneDelegate.Broadcast(LandscapeScratchRT3);
+		OnEditLayersMergedDelegate.Broadcast(FOnLandscapeEditLayersMergedParams(LandscapeScratchRT3, /*InRenderAreaResolution = */LandscapeExtent.Max - LandscapeExtent.Min, /*bInIsHeightmapMerge = */true));
 	}
 
 	// compute Normals into LandscapeScratchRT1
@@ -5064,23 +7376,9 @@
 			const FIntPoint Mip0CopySize = Heightmap.SectionRect.Size();
 			const FIntPoint Mip0SourcePosition = Heightmap.SectionRect.Min;
 
-			FString SourceTextureName = CombinedHeightmapAtlasRT->GetName();
-			FString DestinationTextureName = Heightmap.Texture->GetName();
-			if (const TArray<ULandscapeComponent*>* Components = InUpdateLayersContentContext.MapHelper.HeightmapToComponents.Find(Heightmap.Texture))
-			{
-				FString ComponentName;
-				FString ActorName;
-				for (ULandscapeComponent* Component : *Components)
-				{
-					ComponentName += Component->GetName() + TEXT(", ");
-					ActorName = Component->GetLandscapeProxy()->GetActorLabel();
-				}
-				DestinationTextureName = FString::Format(TEXT("{0}-{1}-{2}"), { ActorName, ComponentName, Heightmap.Texture->GetName() });
-			}
-
 			// Mip 0
 			{
-				FLandscapeLayersCopyTextureParams& CopyTextureParams = DeferredCopyTextures.Add_GetRef(FLandscapeLayersCopyTextureParams(SourceTextureName, CombinedHeightmapAtlasRT->GetResource(), DestinationTextureName, Heightmap.Texture->GetResource()));
+				FLandscapeLayersCopyTextureParams& CopyTextureParams = DeferredCopyTextures.Add_GetRef(FLandscapeLayersCopyTextureParams(CombinedHeightmapAtlasRT, Heightmap.Texture));
 				// Only copy the size that's actually needed : 
 				CopyTextureParams.CopySize = Mip0CopySize;
 				// Copy from the composited texture's position to the top-left corner of the heightmap
@@ -5094,8 +7392,7 @@
 				UTextureRenderTarget2D* RenderTargetMip = HeightmapRTList[MipRTIndex];
 				if (RenderTargetMip != nullptr)
 				{
-					SourceTextureName = RenderTargetMip->GetName();
-					FLandscapeLayersCopyTextureParams& CopyTextureParams = DeferredCopyTextures.Add_GetRef(FLandscapeLayersCopyTextureParams(SourceTextureName, RenderTargetMip->GetResource(), DestinationTextureName, Heightmap.Texture->GetResource()));
+					FLandscapeLayersCopyTextureParams& CopyTextureParams = DeferredCopyTextures.Add_GetRef(FLandscapeLayersCopyTextureParams(RenderTargetMip, Heightmap.Texture));
 					CopyTextureParams.CopySize.X = Mip0CopySize.X >> MipIndex;
 					CopyTextureParams.CopySize.Y = Mip0CopySize.Y >> MipIndex;
 					CopyTextureParams.SourcePosition.X = Mip0SourcePosition.X >> MipIndex;
@@ -5119,22 +7416,249 @@
 	return InMergeParams.HeightmapUpdateModes;
 }
 
-FLandscapeBrushParameters ALandscape::BuildLandscapeBrushParameters(bool bInIsHeightmapMerge, 
-	const FIntRect& InRenderAreaExtents, UTextureRenderTarget2D* InRenderTarget, FName InWeightmapLayerName)
+int32 ALandscape::PerformLayersHeightmapsBatchedMerge(const FUpdateLayersContentContext& InUpdateLayersContentContext, const FEditLayersHeightmapMergeParams& InMergeParams)
 {
+	using namespace UE::Landscape::EditLayers;
+	using namespace UE::Landscape::EditLayers::Private;
+
+	TRACE_CPUPROFILER_EVENT_SCOPE(ALandscape::PerformLayersHeightmapsBatchedMerge);
+	SCOPED_DRAW_EVENT_GAMETHREAD(PerformLayersHeightmapsBatchedMerge);
+
+	ULandscapeInfo* Info = GetLandscapeInfo();
+	check(Info != nullptr);
+
+	// Prepare the heightmap merge operations for all components that need to be updated :
+	TArray<FEditLayerRendererState> RendererStates = GetEditLayerRendererStates(Info, InMergeParams.bSkipBrush);
+	// Add an edit layer renderer at the top of the stack in order to add dependencies between each landscape component and its immediate neighbors in order to ensure they end up
+	//  in the same render batch. The renderer is responsible for computing the normals at the end of the batch : 
+	ULandscapeHeightmapNormalsEditLayerRenderer* HeightmapNormalsRenderer = ULandscapeHeightmapNormalsEditLayerRenderer::StaticClass()->GetDefaultObject<ULandscapeHeightmapNormalsEditLayerRenderer>();
+	RendererStates.Emplace(HeightmapNormalsRenderer, Info);
+
+	FMergeRenderParams MergeRenderParams(/*bInIsHeightmapMerge = */true, /*InLandscape = */this, InUpdateLayersContentContext.LandscapeComponentsHeightmapsToRender, RendererStates);
+	FMergeRenderContext MergeRenderContext = PrepareEditLayersMergeRenderContext(MergeRenderParams);
+	if (!MergeRenderContext.IsValid())
+	{
+		// Nothing to do : 
+		return InMergeParams.HeightmapUpdateModes;
+	}
+
+	// For each batch, render and resolve the raw heightmaps into the individual textures : 
+	TSet<ULandscapeComponent*> ResolvedComponents;
+	ResolvedComponents.Reserve(InUpdateLayersContentContext.LandscapeComponentsHeightmapsToRender.Num());
+
+	// Callback executed each time a render batch is done computing the requested into, just before releasing the render resources : 
+	auto OnRenderBatchGroupDone = [&ResolvedComponents, &InUpdateLayersContentContext, OnEditLayersMergedDelegate = &OnEditLayersMergedDelegate](const FMergeRenderContext::FOnRenderBatchTargetGroupDoneParams& InParams)
+	{
+		ALandscape* Landscape = InParams.MergeRenderContext->GetLandscape();
+
+		//  Note: thanks to HeightmapNormalsRenderer, we have the guarantee that the (up to) 8 neighbors of each of the components requested 
+		//  for are present in the batch, which means we have all the data to generate the normals already 
+		TSet<ULandscapeComponent*> ComponentsToResolveThisBatch;
+		{
+			TRACE_CPUPROFILER_EVENT_SCOPE(PrepareResolve);
+			ComponentsToResolveThisBatch.Reserve(InParams.RenderBatch->ComponentsToRender.Num());
+			for (ULandscapeComponent* RenderedComponent : InParams.RenderBatch->ComponentsToRender)
+			{
+				checkf(InParams.SortedComponentMergeRenderInfos.FindByPredicate([RenderedComponent](const FComponentMergeRenderInfo& InComponentMergeInfo) { return InComponentMergeInfo.Component == RenderedComponent; }) != nullptr,
+					TEXT("All components in the batch must be present in SortedComponentMergeRenderInfos"));
+
+				ULandscapeComponent** ResolvedComponent = ResolvedComponents.Find(RenderedComponent);
+				if (ResolvedComponent == nullptr)
+				{
+					TStaticArray<ULandscapeComponent*, 9> NeighborComponents;
+					RenderedComponent->GetLandscapeComponentNeighbors3x3(NeighborComponents);
+					TSet<ULandscapeComponent*> ValidNeighborComponents;
+					for (ULandscapeComponent* NeighborComponent : NeighborComponents)
+					{
+						if ((NeighborComponent != nullptr) && (NeighborComponent != RenderedComponent))
+						{
+							ValidNeighborComponents.Add(NeighborComponent);
+						}
+					}
+
+					// We need all neighbors to be present in this batch in order to be able to finalize that component :
+					if ((ValidNeighborComponents.Intersect(InParams.RenderBatch->ComponentsToRender).Num() == ValidNeighborComponents.Num())
+						&& InUpdateLayersContentContext.LandscapeComponentsHeightmapsToResolve.Contains(RenderedComponent))
+					{
+						ResolvedComponents.Add(RenderedComponent);
+						ComponentsToResolveThisBatch.Add(RenderedComponent);
+					}
+				}
+			}
+		}
+
+		// Copy to mip0 of the final textures and expand the vertices on borders so that we can generate the mips from it:
+		{
+			TRACE_CPUPROFILER_EVENT_SCOPE(CopyMip0AndExpand);
+			SCOPED_DRAW_EVENT_GAMETHREAD(CopyMip0AndExpand);
+
+			// Recompose mip0 of the final heightmaps, subsection by subsection, to duplicate borders : 
+			InParams.MergeRenderContext->CycleBlendRenderTargets(/*InDesiredWriteAccess = */ERHIAccess::SRVMask); // TODO [jonathan.bard] : This should be CopyDst but ExecuteCopyLayersTexture doesn't allow for it ATM
+			ULandscapeScratchRenderTarget* ReadRT = InParams.MergeRenderContext->GetBlendRenderTargetRead();
+
+			OnEditLayersMergedDelegate->Broadcast(FOnLandscapeEditLayersMergedParams(ReadRT->GetRenderTarget(),
+				/*InRenderAreaResolution = */InParams.RenderBatch->GetRenderTargetResolution(/*bInWithDuplicateBorders = */false), /*bInIsHeightmapMerge = */true));
+
+			ULandscapeScratchRenderTarget* WriteRT = InParams.MergeRenderContext->GetBlendRenderTargetWrite();
+
+			struct FComponentCopyInfo
+			{
+				UTexture2D* Texture = nullptr;
+				ULandscapeComponent* Component = nullptr;
+				FIntPoint TextureOffset = FIntPoint::ZeroValue;
+				TArray<FIntRect, TInlineAllocator<4>> SourceSubsectionRects;
+				TArray<FIntRect, TInlineAllocator<4>> DestinationSubsectionRects;
+			};
+			TArray<FComponentCopyInfo> ComponentCopyInfos;
+			const int32 TotalNumSubsections = Landscape->NumSubsections * Landscape->NumSubsections;
+			const int32 ComponentSubsectionVerts = Landscape->SubsectionSizeQuads + 1;
+
+			for (ULandscapeComponent* Component : InParams.RenderBatch->ComponentsToRender)
+			{
+				FComponentCopyInfo& ComponentCopyInfo = ComponentCopyInfos.Emplace_GetRef();
+				ComponentCopyInfo.Component = Component;
+				ComponentCopyInfo.Texture = Component->GetHeightmap(false);
+				// Effective area of the texture affecting this component (because of texture sharing) :
+				ComponentCopyInfo.TextureOffset = FIntPoint(static_cast<int32>(Component->HeightmapScaleBias.Z * ComponentCopyInfo.Texture->Source.GetSizeX()), static_cast<int32>(Component->HeightmapScaleBias.W * ComponentCopyInfo.Texture->Source.GetSizeY()));
+		
+				InParams.RenderBatch->ComputeSubsectionRects(Component, /*OutSubsectionRects = */ComponentCopyInfo.SourceSubsectionRects, /*OutSubsectionRectsWithDuplicateBorders = */ComponentCopyInfo.DestinationSubsectionRects);
+				check(ComponentCopyInfo.SourceSubsectionRects.Num() == TotalNumSubsections);
+				check(ComponentCopyInfo.DestinationSubsectionRects.Num() == TotalNumSubsections);
+			}
+
+			// TODO [jonathan.bard] : move this after expand (and rename "Expand" to "Generate mip 0")
+			{
+				SCOPED_DRAW_EVENT_GAMETHREAD(CopyToMip0);
+				// Copy sub-section by sub-section in order to duplicate borders :
+				TArray<FLandscapeLayersCopyTextureParams> DeferredCopyTextures;
+				for (const FComponentCopyInfo& ComponentCopyInfo : ComponentCopyInfos)
+				{
+					for (int32 SubSectionIndex = 0; SubSectionIndex < TotalNumSubsections; ++SubSectionIndex)
+					{
+						FIntPoint SubSection(SubSectionIndex % Landscape->NumSubsections, SubSectionIndex / Landscape->NumSubsections);
+						const FIntRect& SourceSubSectionRect = ComponentCopyInfo.SourceSubsectionRects[SubSectionIndex];
+		
+						// Copy to mip0 of the final texture if requested : 
+						if (ComponentsToResolveThisBatch.Contains(ComponentCopyInfo.Component))
+						{
+							FLandscapeLayersCopyTextureParams& CopyTextureParams = DeferredCopyTextures.Add_GetRef(FLandscapeLayersCopyTextureParams(ReadRT->GetRenderTarget(), ComponentCopyInfo.Texture));
+							CopyTextureParams.SourcePosition = SourceSubSectionRect.Min;
+							CopyTextureParams.CopySize = SourceSubSectionRect.Size();
+							CopyTextureParams.DestPosition = ComponentCopyInfo.TextureOffset + FIntPoint(SubSection.X * ComponentSubsectionVerts, SubSection.Y * ComponentSubsectionVerts);
+						}
+					}
+				}
+				ExecuteCopyLayersTexture(MoveTemp(DeferredCopyTextures));
+			}
+
+			// "Expand" the scratch render target : 
+			// Copy sub-section by sub-section in order to duplicate borders :
+			InParams.MergeRenderContext->RenderExpandedRenderTarget(*InParams.RenderBatch);
+		}
+
+		// Generate the mips from the expanded RT and copy to the final texture mips
+		{
+			TRACE_CPUPROFILER_EVENT_SCOPE(GenerateMips);
+			const int32 NumMips = (int32)FMath::CeilLogTwo(Landscape->SubsectionSizeQuads) + 1;
+			SCOPED_DRAW_EVENTF_GAMETHREAD(LandscapeLayers, TEXT("Generate %i remaining mips"), NumMips - 1);
+
+			FIntPoint CurrentMipResolution = InParams.RenderBatch->GetRenderTargetResolution(/*bInWithDuplicateBorders = */true); // Mips are generated after the borders have been duplicated
+			FIntPoint CurrentMipSubsectionSize = Landscape->SubsectionSizeQuads + 1;
+			for (int32 MipIndex = 1; MipIndex < NumMips; ++MipIndex)
+			{
+				InParams.MergeRenderContext->CycleBlendRenderTargets(/*InDesiredWriteAccess = */ERHIAccess::RTV);
+				ULandscapeScratchRenderTarget* WriteRT = InParams.MergeRenderContext->GetBlendRenderTargetWrite();
+				ULandscapeScratchRenderTarget* ReadRT = InParams.MergeRenderContext->GetBlendRenderTargetRead();
+				CurrentMipResolution /= 2;
+				check(CurrentMipResolution.X > 0 && CurrentMipResolution.Y > 0);
+				CurrentMipSubsectionSize /= 2;
+				check(CurrentMipSubsectionSize.X > 0 && CurrentMipSubsectionSize.Y > 0);
+
+				{
+					SCOPED_DRAW_EVENTF_GAMETHREAD(LandscapeLayers, TEXT("Generate mip %i"), MipIndex);
+
+					ENQUEUE_RENDER_COMMAND(LandscapeLayers_Cmd_HeightmapsGenerateMips)(
+						[ OutputResource = WriteRT->GetRenderTarget2D()->GetResource()
+						, SourceResource = ReadRT->GetRenderTarget2D()->GetResource()
+						, CurrentMipResolution
+						, CurrentMipSubsectionSize
+						, MipIndex ]
+						(FRHICommandListImmediate& InRHICmdList)
+					{
+						FRDGBuilder GraphBuilder(InRHICmdList, RDG_EVENT_NAME("HeightmapsGenerateMips"));
+
+						FRDGTextureRef OutputTextureRef = GraphBuilder.RegisterExternalTexture(CreateRenderTarget(OutputResource->TextureRHI, TEXT("OutputTexture")), ERDGTextureFlags::SkipTracking);
+						FRDGTextureRef SourceTextureRef = GraphBuilder.RegisterExternalTexture(CreateRenderTarget(SourceResource->TextureRHI, TEXT("SourceTexture")));
+
+						FLandscapeEditLayersHeightmapsGenerateMipsPS::FParameters* PSParams = GraphBuilder.AllocParameters<FLandscapeEditLayersHeightmapsGenerateMipsPS::FParameters>();
+						PSParams->RenderTargets[0] = FRenderTargetBinding(OutputTextureRef, ERenderTargetLoadAction::ENoAction);
+						PSParams->InCurrentMipSubsectionSize = FUintVector2(CurrentMipSubsectionSize.X, CurrentMipSubsectionSize.Y);
+						PSParams->InSourceHeightmap = GraphBuilder.CreateSRV(FRDGTextureSRVDesc::Create(SourceTextureRef));
+						FLandscapeEditLayersHeightmapsGenerateMipsPS::GenerateMipsPS(GraphBuilder, PSParams, CurrentMipResolution);
+						GraphBuilder.Execute();
+					});
+				}
+
+				// Then copy the appropriate regions to the destination texture mips : 
+				// TODO [jonathan.bard] : add this when we don't auto-transition to SRV in the copy texture thing : WriteRT->TransitionTo(ERHIAccess::CopySrc);
+				{
+					SCOPED_DRAW_EVENTF_GAMETHREAD(LandscapeLayers, TEXT("Copy mip %i"), MipIndex);
+
+					WriteRT->TransitionTo(ERHIAccess::SRVMask);
+					TArray<FLandscapeLayersCopyTextureParams> DeferredCopyTextures;
+					for (ULandscapeComponent* Component : ComponentsToResolveThisBatch)
+					{
+						UTexture2D* ComponentHeightmap = Component->GetHeightmap(false);
+
+						FIntRect SourceSectionRect = InParams.RenderBatch->ComputeSectionRect(Component, /*bInWithDuplicateBorders = */true);
+						SourceSectionRect.Min.X >>= MipIndex;
+						SourceSectionRect.Min.Y >>= MipIndex;
+						SourceSectionRect.Max.X >>= MipIndex;
+						SourceSectionRect.Max.Y >>= MipIndex;
+
+						// Effective area of the texture affecting this component (because of texture sharing) :
+						FIntPoint TextureOffset = FIntPoint(static_cast<int32>(Component->HeightmapScaleBias.Z * ComponentHeightmap->Source.GetSizeX()), static_cast<int32>(Component->HeightmapScaleBias.W * ComponentHeightmap->Source.GetSizeY()));
+						TextureOffset.X >>= MipIndex;
+						TextureOffset.Y >>= MipIndex;
+
+						FLandscapeLayersCopyTextureParams& CopyTextureParams = DeferredCopyTextures.Add_GetRef(FLandscapeLayersCopyTextureParams(WriteRT->GetRenderTarget(), ComponentHeightmap));
+						CopyTextureParams.SourcePosition = SourceSectionRect.Min;
+						CopyTextureParams.CopySize = SourceSectionRect.Size();
+						CopyTextureParams.DestPosition = TextureOffset;
+						CopyTextureParams.DestMip = MipIndex;
+					}
+					ExecuteCopyLayersTexture(MoveTemp(DeferredCopyTextures));
+				}
+			}
+		}
+	};
+
+	// Render everything now. Every time a group from a batch is done (there's only one group per batch for heightmaps), the OnRenderBatchGroupDone callback is called : 
+	MergeRenderContext.Render(OnRenderBatchGroupDone);
+
+	// All requested components must have been resolved by now :
+	check(ResolvedComponents.Num() == InUpdateLayersContentContext.LandscapeComponentsHeightmapsToResolve.Num());
+
+	// Prepare the UTexture2D readbacks we'll need to perform :
+	{
+		TRACE_CPUPROFILER_EVENT_SCOPE(CopyToReadback);
+		SCOPED_DRAW_EVENT_GAMETHREAD(CopyToReadback);
+
+		TArray<FLandscapeLayersCopyReadbackTextureParams> DeferredCopyReadbackTextures = PrepareLandscapeLayersCopyReadbackTextureParams(InUpdateLayersContentContext.MapHelper, InUpdateLayersContentContext.HeightmapsToResolve.Array(), /*bWeightmaps = */false);
+		ExecuteCopyToReadbackTexture(DeferredCopyReadbackTextures);
+	}
+
+	return InMergeParams.HeightmapUpdateModes;
+}
+
+FLandscapeBrushParameters ALandscape::BuildLandscapeBrushParameters(bool bInIsHeightmapMerge, const FIntRect& InRenderAreaExtents, UTextureRenderTarget2D* InRenderTarget, FName InWeightmapLayerName)
+{
 	FTransform RenderAreaWorldTransform = GetTransform();
 	FVector OffsetVector(InRenderAreaExtents.Min.X, InRenderAreaExtents.Min.Y, 0.f);
 	FVector Translation = RenderAreaWorldTransform.TransformFVector4(OffsetVector);
 	RenderAreaWorldTransform.SetTranslation(Translation);
 	FIntPoint RenderAreaSize = InRenderAreaExtents.Max - InRenderAreaExtents.Min;
-
-	// TODO: Someday we will build the final return struct like this:
-	// return FLandscapeBrushParameters(bInIsHeightmapMerge, RenderAreaWorldTransform, RenderAreaSize, InRenderTarget, InWeightmapLayerName);
-
-	return FLandscapeBrushParameters(bInIsHeightmapMerge ? ELandscapeToolTargetType::Heightmap 
-		: (InWeightmapLayerName == UMaterialExpressionLandscapeVisibilityMask::ParameterName) ? ELandscapeToolTargetType::Visibility 
-		: ELandscapeToolTargetType::Weightmap,
-		InRenderTarget, InWeightmapLayerName);
+	return FLandscapeBrushParameters(bInIsHeightmapMerge, RenderAreaWorldTransform, RenderAreaSize, InRenderTarget, InWeightmapLayerName);
 }
 
 int32 ALandscape::RegenerateLayersHeightmaps(const FUpdateLayersContentContext& InUpdateLayersContentContext)
@@ -5187,18 +7711,26 @@
 		MergeParams.bForceRender = bForceRender;
 		MergeParams.bSkipBrush = bSkipBrush;
 
-		if (SupportsEditLayersLocalMerge())
+		switch (CurrentEditLayersMergeMode)
 		{
-			return PerformLayersHeightmapsLocalMerge(InUpdateLayersContentContext, MergeParams);
-		}
-		else
+		case ELandscapeEditLayersMergeMode::GlobalMerge:
 		{
 			return PerformLayersHeightmapsGlobalMerge(InUpdateLayersContentContext, MergeParams);
 		}
-
+		case ELandscapeEditLayersMergeMode::LocalMerge:
+		{
+			return PerformLayersHeightmapsLocalMerge(InUpdateLayersContentContext, MergeParams);
+		}
+		case ELandscapeEditLayersMergeMode::BatchedMerge:
+		{
+			return PerformLayersHeightmapsBatchedMerge(InUpdateLayersContentContext, MergeParams);
+		}
+		default:
+			check(false);
+		}
 	}
 
-	return HeightmapUpdateModes;
+	return 0;
 }
 
 void ALandscape::UpdateForChangedHeightmaps(const TArrayView<FLandscapeEditLayerComponentReadbackResult>& InComponentReadbackResults)
@@ -5332,6 +7864,8 @@
 
 void ALandscape::OnDirtyWeightmap(FTextureToComponentHelper const& MapHelper, UTexture2D const* InWeightmap, FColor const* InOldData, FColor const* InNewData, int32 InMipLevel, uint8 ChangedChannelsMask)
 {
+	using namespace UE::Landscape::Private;
+
 	int32 DumpWeightmapDiff = CVarLandscapeDumpWeightmapDiff.GetValueOnGameThread();
 	const bool bDumpDiff = (DumpWeightmapDiff > 0);
 	const bool bDumpDiffAllMips = (DumpWeightmapDiff > 1);
@@ -5522,66 +8056,70 @@
 				const int32 ComponentWidth = ((SubsectionSizeQuads + 1) * NumSubsections) >> InMipLevel;
 				FIntRect SubRegion(HeightmapOffsetX, HeightmapOffsetY, HeightmapOffsetX + ComponentWidth, HeightmapOffsetY + ComponentWidth);
 
-				FFileHelper::CreateBitmap(*(FilePattern + "_a(pre).bmp"), SizeU, SizeV, InOldData, &SubRegion, &IFileManager::Get(), /*OutFilename = */nullptr, /*bool bInWriteAlpha = */true);
-				FFileHelper::CreateBitmap(*(FilePattern + "_b(post).bmp"), SizeU, SizeV, InNewData, &SubRegion, &IFileManager::Get(), /*OutFilename = */nullptr, /*bool bInWriteAlpha = */true);
-
-				if (bDumpDiffDetails)
+				int32 NumDifferentPixels = 0;
+				uint16 MaxHeightDiff = 0;
+				FIntPoint MaxHeightDiffUV(ForceInit);
+				uint8 MaxNormalDiff = 0;
+				FIntPoint MaxNormalDiffUV(ForceInit);
+				FStringBuilderBase StrBuilder;
+				const FColor* OldDataStartPtr = InOldData + (HeightmapOffsetY * SizeU + HeightmapOffsetX);
+				const FColor* NewDataStartPtr = InNewData + (HeightmapOffsetY * SizeU + HeightmapOffsetX);
+				for (int32 V = 0; V < ComponentWidth; ++V)
 				{
-					int32 NumDifferentPixels = 0;
-					uint16 MaxHeightDiff = 0;
-					FIntPoint MaxHeightDiffUV(ForceInit);
-					uint8 MaxNormalDiff = 0;
-					FIntPoint MaxNormalDiffUV(ForceInit);
-					FStringBuilderBase StrBuilder;
-					const FColor* OldDataStartPtr = InOldData + (HeightmapOffsetY * SizeU + HeightmapOffsetX);
-					const FColor* NewDataStartPtr = InNewData + (HeightmapOffsetY * SizeU + HeightmapOffsetX);
-					for (int32 V = 0; V < ComponentWidth; ++V)
+					for (int32 U = 0; U < ComponentWidth; ++U)
 					{
-						for (int32 U = 0; U < ComponentWidth; ++U)
+						const FColor* OldDataPtr = OldDataStartPtr + (V * SizeU + U);
+						const FColor* NewDataPtr = NewDataStartPtr + (V * SizeU + U);
+						if (*OldDataPtr != *NewDataPtr)
 						{
-							const FColor* OldDataPtr = OldDataStartPtr + (V * SizeU + U);
-							const FColor* NewDataPtr = NewDataStartPtr + (V * SizeU + U);
-							if (*OldDataPtr != *NewDataPtr)
+							uint16 OldHeight = ((static_cast<uint16>(OldDataPtr->R) << 8) | static_cast<uint16>(OldDataPtr->G));
+							uint16 NewHeight = ((static_cast<uint16>(NewDataPtr->R) << 8) | static_cast<uint16>(NewDataPtr->G));
+							uint16 HeightDiff = (NewHeight > OldHeight) ? NewHeight - OldHeight : OldHeight - NewHeight;
+							if (HeightDiff > MaxHeightDiff)
 							{
-								uint16 OldHeight = ((static_cast<uint16>(OldDataPtr->R) << 8) | static_cast<uint16>(OldDataPtr->G));
-								uint16 NewHeight = ((static_cast<uint16>(NewDataPtr->R) << 8) | static_cast<uint16>(NewDataPtr->G));
-								uint16 HeightDiff = (NewHeight > OldHeight) ? NewHeight - OldHeight : OldHeight - NewHeight;
-								if (HeightDiff > MaxHeightDiff)
-								{
-									MaxHeightDiffUV = FIntPoint(U, V);
-									MaxHeightDiff = HeightDiff;
-								}
+								MaxHeightDiffUV = FIntPoint(U, V);
+								MaxHeightDiff = HeightDiff;
+							}
 
-								uint8 OldNormalX = (OldDataPtr->B);
-								uint8 NewNormalX = (NewDataPtr->B);
-								uint16 NormalXDiff = (NewNormalX > OldNormalX) ? NewNormalX - OldNormalX : OldNormalX - NewNormalX;
-								if (NormalXDiff > MaxNormalDiff)
-								{
-									MaxNormalDiffUV = FIntPoint(U, V);
-									MaxNormalDiff = NormalXDiff;
-								}
+							uint8 OldNormalX = (OldDataPtr->B);
+							uint8 NewNormalX = (NewDataPtr->B);
+							uint16 NormalXDiff = (NewNormalX > OldNormalX) ? NewNormalX - OldNormalX : OldNormalX - NewNormalX;
+							if (NormalXDiff > MaxNormalDiff)
+							{
+								MaxNormalDiffUV = FIntPoint(U, V);
+								MaxNormalDiff = NormalXDiff;
+							}
 
-								uint8 OldNormalY = (OldDataPtr->A);
-								uint8 NewNormalY = (NewDataPtr->A);
-								uint16 NormalYDiff = (NewNormalY > OldNormalY) ? NewNormalY - OldNormalY : OldNormalY - NewNormalY;
-								if (NormalYDiff > MaxNormalDiff)
-								{
-									MaxNormalDiffUV = FIntPoint(U, V);
-									MaxNormalDiff = NormalYDiff;
-								}
+							uint8 OldNormalY = (OldDataPtr->A);
+							uint8 NewNormalY = (NewDataPtr->A);
+							uint16 NormalYDiff = (NewNormalY > OldNormalY) ? NewNormalY - OldNormalY : OldNormalY - NewNormalY;
+							if (NormalYDiff > MaxNormalDiff)
+							{
+								MaxNormalDiffUV = FIntPoint(U, V);
+								MaxNormalDiff = NormalYDiff;
+							}
 
-								StrBuilder.Appendf(TEXT("Pixel (%4u,%4u) : Height (%5u -> %5u, absdiff %5u), Normal ((%3u,%3u) -> (%3u,%3u), absdiff %3u)\n"), 
-									U, V, OldHeight, NewHeight, HeightDiff, OldNormalX, OldNormalY, NewNormalX, NewNormalY, FMath::Max(NormalXDiff, NormalYDiff));
+							StrBuilder.Appendf(TEXT("Pixel (%4u,%4u) : Height (%5u -> %5u, absdiff %5u), Normal ((%3u,%3u) -> (%3u,%3u), absdiff %3u)\n"),
+								U, V, OldHeight, NewHeight, HeightDiff, OldNormalX, OldNormalY, NewNormalX, NewNormalY, FMath::Max(NormalXDiff, NormalYDiff));
 
-								++NumDifferentPixels;
-							}
+							++NumDifferentPixels;
 						}
 					}
-					StrBuilder.InsertAt(0, FString::Printf(TEXT("----------------------------------------\n")));
-					StrBuilder.InsertAt(0, FString::Printf(TEXT("Max normal diff (at %s) = %u (%1.3f%%)\n"), *MaxNormalDiffUV.ToString(), MaxNormalDiff, 100.0f * static_cast<float>(MaxNormalDiff) / MAX_uint8));
-					StrBuilder.InsertAt(0, FString::Printf(TEXT("Max height diff (at %s) = %u (%1.3f%%)\n"), *MaxHeightDiffUV.ToString(), MaxHeightDiff, 100.0f * static_cast<float>(MaxHeightDiff) / MAX_uint16));
-					StrBuilder.InsertAt(0, FString::Printf(TEXT("Num diffs = %u\n"), NumDifferentPixels));
-					FFileHelper::SaveStringToFile(StrBuilder.ToView(), *(FilePattern + "_diff.txt"));
 				}
+
+				if (NumDifferentPixels > 0)
+				{
+					FFileHelper::CreateBitmap(*(FilePattern + "_a(pre).bmp"), SizeU, SizeV, InOldData, &SubRegion, &IFileManager::Get(), /*OutFilename = */nullptr, /*bool bInWriteAlpha = */true);
+					FFileHelper::CreateBitmap(*(FilePattern + "_b(post).bmp"), SizeU, SizeV, InNewData, &SubRegion, &IFileManager::Get(), /*OutFilename = */nullptr, /*bool bInWriteAlpha = */true);
+
+					if (bDumpDiffDetails)
+					{
+						StrBuilder.InsertAt(0, FString::Printf(TEXT("----------------------------------------\n")));
+						StrBuilder.InsertAt(0, FString::Printf(TEXT("Max normal diff (at %s) = %u (%1.3f%%)\n"), *MaxNormalDiffUV.ToString(), MaxNormalDiff, 100.0f * static_cast<float>(MaxNormalDiff) / MAX_uint8));
+						StrBuilder.InsertAt(0, FString::Printf(TEXT("Max height diff (at %s) = %u (%1.3f%%)\n"), *MaxHeightDiffUV.ToString(), MaxHeightDiff, 100.0f * static_cast<float>(MaxHeightDiff) / MAX_uint16));
+						StrBuilder.InsertAt(0, FString::Printf(TEXT("Num diffs = %u\n"), NumDifferentPixels));
+						FFileHelper::SaveStringToFile(StrBuilder.ToView(), *(FilePattern + "_diff.txt"));
+					}
+				}
 			}
 		}
@@ -6081,7 +8619,8 @@
 	}
 }
 
-void ALandscape::ReallocateLayersWeightmaps(FUpdateLayersContentContext& InUpdateLayersContentContext, const TArray<ULandscapeLayerInfoObject*>& InBrushRequiredAllocations)
+void ALandscape::ReallocateLayersWeightmaps(FUpdateLayersContentContext& InUpdateLayersContentContext, const TArray<ULandscapeLayerInfoObject*>& InBrushRequiredAllocations, 
+	const TMap<ULandscapeComponent*, TArray<ULandscapeLayerInfoObject*>>* InPerComponentAllocations, bool bInRestrictTextureSharingToTheseComponents)
 {
 	TRACE_CPUPROFILER_EVENT_SCOPE(LandscapeLayers_ReallocateLayersWeightmaps);
 
@@ -6113,6 +8652,13 @@
 
 	// Build a map of all the allocation per components
 	TMap<ULandscapeComponent*, TArray<ULandscapeLayerInfoObject*>> LayerAllocsPerComponent;
+	// If the job of building the per-component allocations has already been done, just use them : 
+	if (InPerComponentAllocations != nullptr)
+	{ 
+		LayerAllocsPerComponent = *InPerComponentAllocations;
+	}
+	else
+	{
 	for (FLandscapeLayer& Layer : LandscapeEditLayers)
 	{
 		for (ULandscapeComponent* Component : InUpdateLayersContentContext.LandscapeComponentsWeightmapsToResolve)
@@ -6149,9 +8695,10 @@
 			}
 		}
 	}
+	}
 
 	// Trim the components that don't need weightmaps anymore (e.g. all edit layers are made invisible : there were some components LandscapeComponentsWeightmapsToResolve but there aren't anymore now):
-	InUpdateLayersContentContext.LandscapeComponentsWeightmapsToResolve.RemoveAll([&](ULandscapeComponent* Component)
+	InUpdateLayersContentContext.LandscapeComponentsWeightmapsToResolve.RemoveAll([&LayerAllocsPerComponent](ULandscapeComponent* Component)
 	{
 		TArray<ULandscapeLayerInfoObject*>* ComponentLayerAlloc = LayerAllocsPerComponent.Find(Component);
 		check(ComponentLayerAlloc != nullptr);
@@ -6189,16 +8736,23 @@
 	}
 
 	// Realloc the weightmap so it will create proper texture (if needed) and will set the allocations information
-	TArray<UTexture*> NewCreatedTextures;
+	TSet<UTexture*> NewCreatedTextures;
 	for (ULandscapeComponent* Component : InUpdateLayersContentContext.LandscapeComponentsWeightmapsToResolve)
 	{
-		Component->ReallocateWeightmaps(nullptr, false, false, false, nullptr, &NewCreatedTextures);
+		// If requested, don't allow the component to use textures that have a weightmap in another component than those provided when reallocating weightmaps :
+		TSet<ULandscapeComponent*> RestrictSharingToComponents;
+		if (bInRestrictTextureSharingToTheseComponents)
+		{
+			LayerAllocsPerComponent.GetKeys(RestrictSharingToComponents);
+		}
+		NewCreatedTextures.Append(Component->ReallocateWeightmaps(/*DataInterface = */nullptr, /*InEditLayerGuid = */FGuid(), /*bInSaveToTransactionBuffer = */false, /*bool bInForceReallocate = */false, 
+			/*InTargetProxy = */nullptr, /*InRestrictSharingToComponents = */bInRestrictTextureSharingToTheseComponents ? &RestrictSharingToComponents : nullptr));
 	}
 
 	// TODO: correctly only recreate what is required instead of everything..
 	//GDisableAutomaticTextureMaterialUpdateDependencies = true;
 
-	FTextureCompilingManager::Get().FinishCompilation(NewCreatedTextures);
+	FTextureCompilingManager::Get().FinishCompilation(NewCreatedTextures.Array());
 	FLandscapeTextureStreamingManager* TextureStreamingManager = GetWorld()->GetSubsystem<ULandscapeSubsystem>()->GetTextureStreamingManager();
 	for (UTexture* Texture : NewCreatedTextures)
 	{
@@ -6315,7 +8869,7 @@
 	BeginInitResource(WeightmapScratchPackLayerTextureResource);
 }
 
-// Little struct that holds information common to PerformLayersWeightmapsLocalMerge and PerformLayersWeightmapsGlobalMerge
+// Little struct that holds information common to PerformLayersWeightmapsLocalMerge, PerformLayersWeightmapsBatchedMerge and PerformLayersWeightmapsGlobalMerge
 struct FEditLayersWeightmapMergeParams
 {
 	int32 WeightmapUpdateModes;
@@ -6326,11 +8880,13 @@
 // Render-thread version of the data / functions we need for the local merge of edit layers : 
 namespace EditLayersWeightmapLocalMerge_RenderThread
 {
+	using namespace UE::Landscape::EditLayers;
+
 	struct FEditLayerInfo
 	{
 		int32 SourceWeightmapTextureIndex = INDEX_NONE; // The index in VisibleEditLayerWeightmapTextures of the texture to read from for this layer
 		int32 SourceWeightmapTextureChannel = INDEX_NONE; // The channel of the texture to read from for this layer
-		ELandscapeEditLayerWeightmapBlendMode BlendMode = ELandscapeEditLayerWeightmapBlendMode::Num; // See EEditLayerWeightmapBlendMode
+		EWeightmapBlendMode BlendMode = EWeightmapBlendMode::Num; // See EEditLayerWeightmapBlendMode
 		float Alpha = 1.0f; // Alpha value to be used in the blend
 	};
 
@@ -6374,9 +8930,8 @@
 	struct FTextureResolveInfo
 	{
 		FTextureResolveInfo(int32 InNumComponentsToRender)
-		{
-			ComponentToRenderInfoBitIndices.Init(false, InNumComponentsToRender);
-		}
+			: ComponentToRenderInfoBitIndices(false, InNumComponentsToRender)
+		{}
 
 		void ValidatePerChannelSourceInfo(int32 InChannelIndex, const FComponentAndPaintLayerRenderInfoIdentifier& InComponentAndPaintLayerIdentifier)
 		{
@@ -6423,9 +8978,8 @@
 	{
 		FTextureResolveBatchInfo(int32 InNumComponentsToRender, int32 InBatchIndex)
 			: BatchIndex(InBatchIndex)
-		{
-			ComponentToRenderInfoBitIndices.Init(false, InNumComponentsToRender);
-		}
+			, ComponentToRenderInfoBitIndices(false, InNumComponentsToRender)
+		{}
 
 		void AddTexture(const FTextureResolveInfo& InTextureResolveInfo)
 		{
@@ -6483,7 +9037,7 @@
 		TArray<FTextureResolveInfo> TextureToResolveInfos;
 
 		// List of infos for each individual paint layer involved in the merge operation (including the visibility layer)
-		TArray<FLandscapeWeightmapPaintLayerInfo> PaintLayerInfos;
+		TArray<FWeightmapPaintLayerInfo> PaintLayerInfos;
 
 		// Not truly render-thread data because it references UTextures but it's just because FLandscapeEditLayerReadback were historically game-thread initiated so for as long as we'll use those for readback, we need to store this here : 
 		TArray<FLandscapeLayersCopyReadbackTextureParams> DeferredCopyReadbackTextures;
@@ -6911,7 +9465,7 @@
 							int32* PaintLayerInfoIndex = PaintLayerToPaintLayerInfoIndex.Find(ComponentEditLayerAllocation.LayerInfo);
 							if (PaintLayerInfoIndex == nullptr)
 							{
-								FLandscapeWeightmapPaintLayerInfo NewPaintLayerInfo;
+								FWeightmapPaintLayerInfo NewPaintLayerInfo;
 								if (!ComponentEditLayerAllocation.LayerInfo->bNoWeightBlend)
 								{
 									NewPaintLayerInfo.Flags = EWeightmapPaintLayerFlags::IsWeightBlended;
@@ -6959,7 +9513,7 @@
 							NewEditLayerInfo.SourceWeightmapTextureChannel = ComponentEditLayerAllocation.WeightmapTextureChannel;
 							NewEditLayerInfo.Alpha = (ComponentEditLayerAllocation.LayerInfo == ALandscapeProxy::VisibilityLayer) ? 1.0f : Layer.WeightmapAlpha; // visibility can't affect or be affected by other paint layer weights
 							const bool* BlendSubstractive = Layer.WeightmapLayerAllocationBlend.Find(ComponentEditLayerAllocation.LayerInfo);
-							NewEditLayerInfo.BlendMode = ((BlendSubstractive != nullptr) && *BlendSubstractive) ? ELandscapeEditLayerWeightmapBlendMode::Substractive : ELandscapeEditLayerWeightmapBlendMode::Additive;
+							NewEditLayerInfo.BlendMode = ((BlendSubstractive != nullptr) && *BlendSubstractive) ? EWeightmapBlendMode::Subtractive : EWeightmapBlendMode::Additive;
 
 							ComponentPaintLayerRenderInfo.VisibleEditLayerInfos.Add(NewEditLayerInfo);
 						}
@@ -7132,7 +9686,7 @@
 	// TODO [jonathan.bard] : this is only possible to do so in local merge, since we don't yet support BP brushes, for which we need to call Render() in order to be able to know which layer allocations
 	//  they need. As such, the global merge path is broken and cannot be fixed unless the BP brush interface is changed in order to inform the system, by advance, about what paint layer it needs
 	// Make sure we have proper textures+allocations for all the final weightmaps we're about to resolve :
-	ReallocateLayersWeightmaps(InUpdateLayersContentContext, /*InBrushRequiredAllocations = */TArray<ULandscapeLayerInfoObject*>());
+	ReallocateLayersWeightmaps(InUpdateLayersContentContext, /*InBrushRequiredAllocations = */TArray<ULandscapeLayerInfoObject*>(), /*InPerComponentAllocations = */nullptr, /*bInRestrictTextureSharingToTheseComponents = */false);
 
 	FMergeInfo RenderThreadData;
 	PrepareLayersWeightmapsLocalMergeRenderThreadData(InUpdateLayersContentContext, InMergeParams, RenderThreadData);
@@ -7422,17 +9976,15 @@
 					// Handle brush blending
 					if (Layer.bVisible && !InMergeParams.bSkipBrush)
 					{
-						const ELandscapeToolTargetType LayerType = UE::Landscape::IsVisibilityLayer(LayerInfoObj) ? ELandscapeToolTargetType::Visibility : ELandscapeToolTargetType::Weightmap;
-
-						// Draw each brushes				
+						// Draw each brush				
 						for (int32 i = 0; i < Layer.Brushes.Num(); ++i)
 						{
 							// TODO: handle conversion/handling of RT not same size as internal size
 
 							FLandscapeLayerBrush& Brush = Layer.Brushes[i];
-							FLandscapeBrushParameters BrushParameters(LayerType, LandscapeScratchRT3, LayerInfoObj->LayerName);
+							FLandscapeBrushParameters BrushParameters = BuildLandscapeBrushParameters(/*bIsHeightmapMerge = */false, LandscapeExtent, LandscapeScratchRT3, LayerInfoObj->LayerName);
 
-							UTextureRenderTarget2D* BrushOutputRT = Brush.RenderLayer( LandscapeExtent, BrushParameters);
+							UTextureRenderTarget2D* BrushOutputRT = Brush.RenderLayer(BrushParameters);
 							if (BrushOutputRT == nullptr || BrushOutputRT->SizeX != LandscapeScratchRT3->SizeX || BrushOutputRT->SizeY != LandscapeScratchRT3->SizeY)
 							{
 								continue;
@@ -7466,11 +10018,7 @@
 							FLandscapeBrushParameters BrushParameters = BuildLandscapeBrushParameters(
 								/*bIsHeightmapMerge = */false, LandscapeExtent, LandscapeScratchRT3, LayerInfoObj->LayerName);
 							
-							UTextureRenderTarget2D* BrushOutputRT = RenderCallAdapter->RenderAsBlueprintBrush(
-								BrushParameters, 
-								LandscapeEditLayersLocals::GetLandscapeExtentTransform(*this, LandscapeExtent),
-								// TODO: Someday this will be part of BrushParameters
-								LandscapeExtent.Max - LandscapeExtent.Min);
+							UTextureRenderTarget2D* BrushOutputRT = RenderCallAdapter->RenderAsBlueprintBrush(BrushParameters, GetTransform());
 
 							if (BrushOutputRT != nullptr && BrushOutputRT->SizeX == LandscapeScratchRT3->SizeX 
 								&& BrushOutputRT->SizeY == LandscapeScratchRT3->SizeY)
@@ -7514,7 +10062,7 @@
 			bFirstLayer = false;
 		}
 
-		ReallocateLayersWeightmaps(InUpdateLayersContentContext, BrushRequiredAllocations);
+		ReallocateLayersWeightmaps(InUpdateLayersContentContext, BrushRequiredAllocations, /*InPerComponentAllocations = */nullptr, /*bInRestrictTextureSharingToTheseComponents = */false);
 
 		// List of UTexture2D that we need to kick off readbacks for :
 		TArray<UTexture2D*> TexturesNeedingReadback;
@@ -7721,6 +10269,335 @@
 	return InMergeParams.WeightmapUpdateModes;
 }
 
+int32 ALandscape::PerformLayersWeightmapsBatchedMerge(FUpdateLayersContentContext& InUpdateLayersContentContext, const FEditLayersWeightmapMergeParams& InMergeParams)
+{
+	using namespace UE::Landscape;
+	using namespace UE::Landscape::EditLayers;
+	using namespace UE::Landscape::EditLayers::Private;
+
+	TRACE_CPUPROFILER_EVENT_SCOPE(ALandscape::PerformLayersWeightmapsBatchedMerge);
+	SCOPED_DRAW_EVENT_GAMETHREAD(PerformLayersWeightmapsBatchedMerge);
+
+	ULandscapeInfo* Info = GetLandscapeInfo();
+	check(Info != nullptr);
+
+	TArray<FEditLayerRendererState> RendererStates = GetEditLayerRendererStates(Info, InMergeParams.bSkipBrush);
+
+	const TSet<FName> RequestedWeightmapLayerNames(GetTargetLayerNames(/*bInIncludeVisibilityLayer = */ true));
+
+	// Prepare the merge : 
+	TSet<ULandscapeComponent*> LandscapeComponentsWeightmapsToRenderBefore(InUpdateLayersContentContext.LandscapeComponentsWeightmapsToRender);
+	FMergeRenderParams MergeRenderParams(/*bInIsHeightmapMerge = */false, /*InLandscape = */this, InUpdateLayersContentContext.LandscapeComponentsWeightmapsToRender, RendererStates, RequestedWeightmapLayerNames);
+	FMergeRenderContext MergeRenderContext = PrepareEditLayersMergeRenderContext(MergeRenderParams); 
+	if (!MergeRenderContext.IsValid())
+	{
+		return InMergeParams.WeightmapUpdateModes;
+	}
+
+	TArray<ULandscapeComponent*> FinalComponentsToResolve;
+	FinalComponentsToResolve.Reserve(InUpdateLayersContentContext.LandscapeComponentsWeightmapsToResolve.Num());
+	{
+		TRACE_CPUPROFILER_EVENT_SCOPE(ReallocateWeightmaps);
+		TMap<ULandscapeComponent*, TArray<ULandscapeLayerInfoObject*>> PerComponentAllocations;
+		{
+			TRACE_CPUPROFILER_EVENT_SCOPE(PrepareComponentAllocations);
+			PerComponentAllocations.Reserve(MergeRenderContext.ComponentToTargetLayerBitIndices.Num());
+			for (auto ItPair : MergeRenderContext.ComponentToTargetLayerBitIndices)
+			{
+				ULandscapeComponent* Component = ItPair.Key;
+				const TBitArray<>& ComponentTargetLayerBitIndices = ItPair.Value;
+				TArray<ULandscapeLayerInfoObject*>& ComponentAllocations = PerComponentAllocations.Add(Component, MergeRenderContext.ConvertTargetLayerBitIndicesToLayerInfos(ComponentTargetLayerBitIndices));
+				if (!ComponentAllocations.IsEmpty())
+				{
+					FinalComponentsToResolve.Add(Component);
+				}
+			}
+		}
+		ReallocateLayersWeightmaps(InUpdateLayersContentContext, /*InBrushRequiredAllocations = */{}, &PerComponentAllocations, /*bInRestrictTextureSharingToTheseComponents = */true);
+
+		checkf(Algo::AllOf(InUpdateLayersContentContext.LandscapeComponentsWeightmapsToRender, [&LandscapeComponentsWeightmapsToRenderBefore](ULandscapeComponent* InComponent) { return LandscapeComponentsWeightmapsToRenderBefore.Contains(InComponent); }), 
+			TEXT("If this asserts, the list of components involved in the weightmaps being merged has changed, which would require re-running the merge preparation step. bInRestrictTextureSharingToTheseComponents should prevent this from happening"));
+	}
+
+	TSet<ULandscapeComponent*> ResolvedLandscapeComponents;
+	ResolvedLandscapeComponents.Reserve(FinalComponentsToResolve.Num());
+	// Key = texture, Value = channels mask resolved so far. If all channels have been resolved, we remove the entry as the texture is already fully resolved :
+	TMap<UTexture2D*, uint8> TexturesNeedingResolve; 
+	TexturesNeedingResolve.Reserve(InUpdateLayersContentContext.WeightmapsToResolve.Num());
+	for (UTexture2D* Weightmap : InUpdateLayersContentContext.WeightmapsToResolve)
+	{ 
+		TexturesNeedingResolve.Add(Weightmap, 0);
+	}
+
+	// Callback executed each time a render batch is done computing the requested into, just before releasing the render resources : 
+	auto OnRenderBatchGroupDone = [&ResolvedLandscapeComponents, &TexturesNeedingResolve, &InUpdateLayersContentContext, OnEditLayersMergedDelegate = &OnEditLayersMergedDelegate](const FMergeRenderContext::FOnRenderBatchTargetGroupDoneParams& InParams)
+	{
+		// We can now finalize the weightmaps : since we don't use SUPPORTS_LANDSCAPE_EDITORONLY_UBER_MATERIAL yet, each component that has been rendered should have a complete set of weightmaps that has been regenerated
+		//  already so we need to repack those into the appropriate weightmap channels
+
+		ALandscape* Landscape = InParams.MergeRenderContext->GetLandscape();
+
+		// Cycle render targets one last time so we can read from the last RT we've written to :
+		InParams.MergeRenderContext->CycleBlendRenderTargets();
+		ULandscapeScratchRenderTarget* ReadRT = InParams.MergeRenderContext->GetBlendRenderTargetRead();
+
+		OnEditLayersMergedDelegate->Broadcast(FOnLandscapeEditLayersMergedParams(ReadRT->GetRenderTarget(),
+			/*InRenderAreaResolution = */InParams.RenderBatch->GetRenderTargetResolution(/*bInWithDuplicateBorders = */false), /*bInIsHeightmapMerge = */false));
+
+		const int32 TotalNumSubsections = Landscape->NumSubsections * Landscape->NumSubsections;
+		FIntPoint MinWeightmapResolution(MAX_int32, MAX_int32);
+		FIntPoint MaxWeightmapResolution(MIN_int32, MIN_int32);
+		TArray<FName> TargetLayerNames = InParams.MergeRenderContext->ConvertTargetLayerBitIndicesToNamesChecked(InParams.RenderGroupBitIndices);
+
+		struct FWeightmapResolveInfo
+		{
+			// Weightmap to resolve : 
+			FTextureResource* TextureResource = nullptr;
+			FString TextureResourceDebugName;
+			// Slice index (in the batch's source texture array) that needs to be copied onto each individual channel of this weightmap:
+			FIntVector4 SourceSliceIndexPerChannel = FIntVector4(INDEX_NONE);
+			// The rects that correspond to the component to read (without border expansion) in the batch (source) texture (one per channel)
+			TStaticArray<TArray<FIntRect, TInlineAllocator<4>>, 4> SourceSubsectionRectsPerChannel;
+			// Indicates which channel(s) to resolve : 
+			uint8 ChannelMask = 0;
+			// Indicates whether this resolve operation should be done additively (e.g. there could be a first resolve on texture T for channels .rgb and another later on for channel .a : the former 
+			//  would *not* be done additively but the latter, yes... you can thank weightmap channel sharing for all that complexity)
+			bool bIsAdditiveResolve = false;
+			// Indicates this resolve operation is the final one, i.e. all channels from this weightmap are now resolved so we can finalize the texture entirely (generate mips, final copy and readback) :
+			bool bIsFinalResolve = false;
+		};
+		// List of weightmaps to resolve this batch and how to resolve them: 
+		TMap<UTexture2D*, FWeightmapResolveInfo> WeightmapResolveInfosForBatch;
+		TSet<ULandscapeComponent*> ComponentsResolvedInBatch;
+
+		// Process the list of textures that are not yet resolved and find if there are components in this batch that we participate to it so we can resolve it either partially or fully : 
+		{
+			TRACE_CPUPROFILER_EVENT_SCOPE(PrepareResolve);
+			for (auto ItWeightmapNeedingResolve = TexturesNeedingResolve.CreateIterator(); ItWeightmapNeedingResolve; ++ItWeightmapNeedingResolve)
+			{
+				UTexture2D* Weightmap = ItWeightmapNeedingResolve.Key();
+				uint8& AlreadyResolvedChannelMask = ItWeightmapNeedingResolve.Value();
+
+				uint8 WeightmapChannelMask = InUpdateLayersContentContext.MapHelper.WeightmapToChannelMask.FindChecked(Weightmap);
+				check((WeightmapChannelMask > 0) && (WeightmapChannelMask <= 15)); // a mask of 0 means nothing needs to be resolved, so it shouldn't ever be in the TexturesNeedingResolve list
+			
+				const TArray<ULandscapeComponent*>& ComponentsForWeightmap = InUpdateLayersContentContext.MapHelper.WeightmapToComponents.FindChecked(Weightmap);
+				for (ULandscapeComponent* RenderedComponent : ComponentsForWeightmap)
+				{
+					// If the component is present in this batch, we can resolve its channels :
+					if (InParams.RenderBatch->ComponentsToRender.Contains(RenderedComponent) 
+						// The component was possibly rendered in a previous batch, in which case we don't have to resolved it again : 
+						&& !ResolvedLandscapeComponents.Contains(RenderedComponent))
+					{
+						TArray<FIntRect, TInlineAllocator<4>> SourceSubsectionRects, DummySubsectionRects;
+						InParams.RenderBatch->ComputeSubsectionRects(RenderedComponent, /*OutSubsectionRects = */SourceSubsectionRects, /*OutSubsectionRectsWithDuplicateBorders = */DummySubsectionRects);
+						check(SourceSubsectionRects.Num() == TotalNumSubsections);
+
+						FWeightmapResolveInfo* WeightmapResolveInfo = WeightmapResolveInfosForBatch.Find(Weightmap);
+						if (WeightmapResolveInfo == nullptr)
+						{
+							// This is a new weightmap to resolve this batch, let's create the info : 
+							WeightmapResolveInfo = &WeightmapResolveInfosForBatch.Add(Weightmap, FWeightmapResolveInfo());
+							WeightmapResolveInfo->TextureResource = Weightmap->GetResource();
+							WeightmapResolveInfo->TextureResourceDebugName = Weightmap->GetName();
+							// We need the resolve operation to be additive if some channels have already been resolved in a previous batch: 
+							WeightmapResolveInfo->bIsAdditiveResolve = (AlreadyResolvedChannelMask != 0);
+
+							ALandscapeProxy* Proxy = RenderedComponent->GetLandscapeProxy();
+							check(Proxy != nullptr);
+							// Setup the CPU readback if it does not already exist:
+							FLandscapeEditLayerReadback** CPUReadback = Proxy->WeightmapsCPUReadback.Find(Weightmap);
+							if (CPUReadback == nullptr)
+							{
+								// Lazily create the readback objects as required (ReallocateLayersWeightmaps might have created new weightmaps)
+								FLandscapeEditLayerReadback* NewCPUReadback = new FLandscapeEditLayerReadback();
+								const uint8* LockedMip = Weightmap->Source.LockMipReadOnly(0);
+								const uint64 Hash = FLandscapeEditLayerReadback::CalculateHash(LockedMip, Weightmap->Source.GetSizeX() * Weightmap->Source.GetSizeY() * sizeof(FColor));
+								NewCPUReadback->SetHash(Hash);
+								Weightmap->Source.UnlockMip(0);
+								CPUReadback = &Proxy->WeightmapsCPUReadback.Add(Weightmap, NewCPUReadback);
+							}
+							check(*CPUReadback != nullptr);
+
+							FIntPoint WeightmapResolution(Weightmap->Source.GetSizeX(), Weightmap->Source.GetSizeY());
+							MinWeightmapResolution = MinWeightmapResolution.ComponentMin(WeightmapResolution);
+							MaxWeightmapResolution = MaxWeightmapResolution.ComponentMax(WeightmapResolution);
+						}
+
+
+						// Select only the allocations of this component that involve this texture : 
+						const TArray<UTexture2D*>& ComponentTextures = RenderedComponent->GetWeightmapTextures();
+						const int32 WeightmapIndex = ComponentTextures.IndexOfByKey(Weightmap);
+						check(WeightmapIndex != INDEX_NONE);
+						const TArray<FWeightmapLayerAllocationInfo>& ComponentAllocationInfos = RenderedComponent->GetWeightmapLayerAllocations();
+						TArray<FWeightmapLayerAllocationInfo, TInlineAllocator<4>> AllocationInfosForTexture;
+						Algo::TransformIf(ComponentAllocationInfos, AllocationInfosForTexture,
+							[WeightmapIndex](const FWeightmapLayerAllocationInfo& InAllocationInfo) { return (InAllocationInfo.LayerInfo != nullptr) && InAllocationInfo.IsAllocated() && (InAllocationInfo.WeightmapTextureIndex == WeightmapIndex); },
+							[](const FWeightmapLayerAllocationInfo& InAllocationInfo) { return InAllocationInfo; });
+						check(!AllocationInfosForTexture.IsEmpty() && (AllocationInfosForTexture.Num() <= 4));
+
+						for (const FWeightmapLayerAllocationInfo& AllocationInfo : AllocationInfosForTexture)
+						{
+							check((AllocationInfo.WeightmapTextureChannel >= 0) && (AllocationInfo.WeightmapTextureChannel < 4));
+							checkf(((WeightmapResolveInfo->ChannelMask & (1 << AllocationInfo.WeightmapTextureChannel)) == 0), TEXT("This channel has already been resolved, it shouldn't happen, it would mean that 2 allocations are using the same channel"));
+							int32 SliceIndex = TargetLayerNames.IndexOfByKey(AllocationInfo.GetLayerName());
+							checkf(SliceIndex != INDEX_NONE, TEXT("Couldn't find %s in the list of weightmaps that have been produced"), *AllocationInfo.GetLayerName().ToString());
+							WeightmapResolveInfo->SourceSliceIndexPerChannel[AllocationInfo.WeightmapTextureChannel] = SliceIndex;
+							WeightmapResolveInfo->SourceSubsectionRectsPerChannel[AllocationInfo.WeightmapTextureChannel] = SourceSubsectionRects;
+							WeightmapResolveInfo->ChannelMask |= (1 << AllocationInfo.WeightmapTextureChannel);
+						}
+						check(WeightmapResolveInfo->ChannelMask != 0);
+						// We can now consider these channels resolved for this texture : 
+						AlreadyResolvedChannelMask |= WeightmapResolveInfo->ChannelMask;
+
+						// If all of this weightmap's channels have been resolved, we can finalize it and remove it from our list of weightmaps to resolve :
+						if (AlreadyResolvedChannelMask == WeightmapChannelMask)
+						{
+							WeightmapResolveInfo->bIsFinalResolve = true;
+							ItWeightmapNeedingResolve.RemoveCurrent();
+						}
+
+						ComponentsResolvedInBatch.Add(RenderedComponent);
+					}
+				}
+			}
+		}
+
+		if (!WeightmapResolveInfosForBatch.IsEmpty())
+		{
+			TRACE_CPUPROFILER_EVENT_SCOPE(FinalizeWeightmaps);
+			SCOPED_DRAW_EVENT_GAMETHREAD(FinalizeWeightmaps);
+
+			check(MinWeightmapResolution == MaxWeightmapResolution);
+
+			ENQUEUE_RENDER_COMMAND(LandscapeLayers_Cmd_WeightmapsPackWeightmap)(
+				[ SourceResource = ReadRT->GetRenderTarget2DArray()->GetResource()
+				, WeightmapResolveInfosForBatch
+				, WeightmapResolution = MinWeightmapResolution
+				, NumMips = (int32)FMath::CeilLogTwo(Landscape->SubsectionSizeQuads) + 1
+				, TotalNumSubsections = Landscape->NumSubsections * Landscape->NumSubsections
+				, NumSubsections = Landscape->NumSubsections
+				, ComponentSubsectionVerts = Landscape->SubsectionSizeQuads + 1 ]
+				(FRHICommandListImmediate& InRHICmdList)
+			{
+				FRDGBuilder GraphBuilder(InRHICmdList, RDG_EVENT_NAME("WeightmapsFinalizeWeightmaps"));
+
+				FRDGTextureSRVRef BlackDummySRVRef = GraphBuilder.CreateSRV(FRDGTextureSRVDesc::Create(GSystemTextures.GetBlackDummy(GraphBuilder)));
+				FRDGTextureRef SourceTextureRef = GraphBuilder.RegisterExternalTexture(CreateRenderTarget(SourceResource->TextureRHI, TEXT("SourceTexture")));
+				FRDGTextureSRVRef SourceTextureSRVRef = GraphBuilder.CreateSRV(FRDGTextureSRVDesc::Create(SourceTextureRef));
+				FRDGTextureDesc Desc = FRDGTextureDesc::Create2D(WeightmapResolution, PF_B8G8R8A8, FClearValueBinding::Black, TexCreate_RenderTargetable | TexCreate_ShaderResource, NumMips);
+				FRDGTextureRef PackedTextureRef = GraphBuilder.CreateTexture(Desc, TEXT("PackedWeightmap"));
+
+				for (auto It : WeightmapResolveInfosForBatch)
+				{
+					const FWeightmapResolveInfo& WeightmapResolveInfo = It.Value;
+					FRDGTextureRef DestinationTextureRef = GraphBuilder.RegisterExternalTexture(CreateRenderTarget(WeightmapResolveInfo.TextureResource->TextureRHI, TEXT("DestinationTexture")));
+
+					RDG_EVENT_SCOPE(GraphBuilder, "%sFinalize %s", WeightmapResolveInfo.bIsFinalResolve ? TEXT("") : TEXT("(partially)"), *WeightmapResolveInfo.TextureResourceDebugName);
+					{
+						RDG_EVENT_SCOPE(GraphBuilder, "Pack %d channels %s", FMath::CountBits(WeightmapResolveInfo.ChannelMask), WeightmapResolveInfo.bIsAdditiveResolve ? TEXT("(additive)") : TEXT(""));
+
+						// If the resolve is additive, it means we need to use the destination texture as an input to the pack operation, because we've stored the previous (partial) resolve in it :
+						FRDGTextureSRVRef WeightmapBeingPackedSRVRef = WeightmapResolveInfo.bIsAdditiveResolve ? GraphBuilder.CreateSRV(FRDGTextureSRVDesc::Create(DestinationTextureRef)) : BlackDummySRVRef;
+
+						// Operate sub-section by sub-section in order to both pack the 4 channels and duplicate borders :
+						for (int32 SubSectionIndex = 0; SubSectionIndex < TotalNumSubsections; ++SubSectionIndex)
+						{
+							FIntPoint SubSection(SubSectionIndex % NumSubsections, SubSectionIndex / NumSubsections);
+							const FIntRect OutputRect = FIntRect(SubSection * ComponentSubsectionVerts, SubSection * ComponentSubsectionVerts + FIntPoint(ComponentSubsectionVerts, ComponentSubsectionVerts));
+
+							FLandscapeEditLayersWeightmapsPackWeightmapPS::FParameters* PSParams = GraphBuilder.AllocParameters<FLandscapeEditLayersWeightmapsPackWeightmapPS::FParameters>();
+							PSParams->RenderTargets[0] = FRenderTargetBinding(PackedTextureRef, ERenderTargetLoadAction::ENoAction);
+							PSParams->InSourceSliceIndices = WeightmapResolveInfo.SourceSliceIndexPerChannel;
+							for (int32 ChannelIndex = 0; ChannelIndex < 4; ++ChannelIndex)
+							{
+								if ((WeightmapResolveInfo.SourceSliceIndexPerChannel[ChannelIndex] != INDEX_NONE))
+								{
+									const FIntRect& SourceSubSectionRect = WeightmapResolveInfo.SourceSubsectionRectsPerChannel[ChannelIndex][SubSectionIndex];
+									PSParams->InSourcePixelOffsets[ChannelIndex] = FUintVector4(SourceSubSectionRect.Min.X, SourceSubSectionRect.Min.Y, 0, 0);
+								}
+								else
+								{
+									PSParams->InSourcePixelOffsets[ChannelIndex] = FUintVector4(0, 0, 0, 0);
+								}
+							}
+							PSParams->InSubsectionPixelOffset = FUintVector2(OutputRect.Min.X, OutputRect.Min.Y);
+							PSParams->InIsAdditive = WeightmapResolveInfo.bIsAdditiveResolve ? 1 : 0;
+							PSParams->InSourceWeightmaps = SourceTextureSRVRef;
+							PSParams->InWeightmapBeingPacked = WeightmapBeingPackedSRVRef;
+							FLandscapeEditLayersWeightmapsPackWeightmapPS::PackWeightmapPS(GraphBuilder, PSParams, OutputRect);
+						}
+					}
+
+					if (WeightmapResolveInfo.bIsFinalResolve)
+					{
+						// Only generate the mips if it's the final resolve : 
+						if (NumMips > 1)
+						{
+							RDG_EVENT_SCOPE(GraphBuilder, "Generate %d remaining mips", NumMips - 1);
+
+							FIntPoint CurrentMipSize = WeightmapResolution;
+							for (int32 MipLevel = 1; MipLevel < NumMips; ++MipLevel)
+							{
+								CurrentMipSize.X >>= 1;
+								CurrentMipSize.Y >>= 1;
+
+								// Read from scratch weightmap texture (mip N - 1) -> write to scratch weightmap texture (mip N) :
+								FLandscapeEditLayersWeightmapsGenerateMipsPS::FParameters* PSParams = GraphBuilder.AllocParameters<FLandscapeEditLayersWeightmapsGenerateMipsPS::FParameters>();
+								PSParams->RenderTargets[0] = FRenderTargetBinding(PackedTextureRef, ERenderTargetLoadAction::ENoAction, static_cast<uint8>(MipLevel));
+								PSParams->InCurrentMipSubsectionSize = FUintVector2(CurrentMipSize.X / NumSubsections, CurrentMipSize.Y / NumSubsections);
+								PSParams->InSourceWeightmap = GraphBuilder.CreateSRV(FRDGTextureSRVDesc::CreateForMipLevel(PackedTextureRef, MipLevel - 1));
+
+								FLandscapeEditLayersWeightmapsGenerateMipsPS::GenerateMipsPS(GraphBuilder, PSParams, CurrentMipSize);
+							}
+						}
+					}
+
+					{
+						// We use the final texture as a temporary buffer when it's a partial resolve, so we only need to copy mip 0 then. All mips will be generated/copy when the resolve step is final :
+						const int32 NumMipsToCopy = WeightmapResolveInfo.bIsFinalResolve ? NumMips : 1;
+						RDG_EVENT_SCOPE(GraphBuilder, "Copy %d mips", NumMips);
+						for (int32 MipLevel = 0; MipLevel < NumMips; ++MipLevel)
+						{
+							FRHICopyTextureInfo CopyInfo;
+							CopyInfo.SourceMipIndex = MipLevel;
+							CopyInfo.DestMipIndex = MipLevel;
+							AddCopyTexturePass(GraphBuilder, PackedTextureRef, DestinationTextureRef, CopyInfo);
+						}
+					}
+				}
+
+				GraphBuilder.Execute();
+			});
+		}
+
+		// Remember all components resolved this batch so that we don't have to resolve them ever again : 
+		ResolvedLandscapeComponents.Append(ComponentsResolvedInBatch);
+	};
+
+	// Render everything now. Every time a group from a batch is done, the OnRenderBatchGroupDone callback is called : 
+	MergeRenderContext.Render(OnRenderBatchGroupDone);
+
+	check(ResolvedLandscapeComponents.Intersect(TSet<ULandscapeComponent*>(FinalComponentsToResolve)).Num() == ResolvedLandscapeComponents.Num());
+	check(TexturesNeedingResolve.IsEmpty());
+
+	// Prepare the UTexture2D readbacks we'll need to perform :
+	{
+		TRACE_CPUPROFILER_EVENT_SCOPE(CopyToReadback);
+		SCOPED_DRAW_EVENT_GAMETHREAD(CopyToReadback);
+
+		TArray<FLandscapeLayersCopyReadbackTextureParams> DeferredCopyReadbackTextures = PrepareLandscapeLayersCopyReadbackTextureParams(InUpdateLayersContentContext.MapHelper, InUpdateLayersContentContext.WeightmapsToResolve.Array(), /*bWeightmaps = */true);
+		ExecuteCopyToReadbackTexture(DeferredCopyReadbackTextures);
+	}
+
+	// Finally, update the material instances to take into account potentially new material combinations : 
+	UpdateLayersMaterialInstances(InUpdateLayersContentContext.LandscapeComponentsWeightmapsToResolve);
+
+	return InMergeParams.WeightmapUpdateModes;
+}
+
+
+
 int32 ALandscape::RegenerateLayersWeightmaps(FUpdateLayersContentContext& InUpdateLayersContentContext)
 {
 	TRACE_CPUPROFILER_EVENT_SCOPE(LandscapeLayers_RegenerateLayersWeightmaps);
@@ -7750,13 +10627,22 @@
 		MergeParams.bForceRender = bForceRender;
 		MergeParams.bSkipBrush = bSkipBrush;
 
-		if (SupportsEditLayersLocalMerge())
+		switch (CurrentEditLayersMergeMode)
 		{
-			return PerformLayersWeightmapsLocalMerge(InUpdateLayersContentContext, MergeParams);
-		}
-		else
+		case ELandscapeEditLayersMergeMode::GlobalMerge:
 		{
 			return PerformLayersWeightmapsGlobalMerge(InUpdateLayersContentContext, MergeParams);
 		}
+		case ELandscapeEditLayersMergeMode::LocalMerge:
+		{
+			return PerformLayersWeightmapsLocalMerge(InUpdateLayersContentContext, MergeParams);
+		}
+		case ELandscapeEditLayersMergeMode::BatchedMerge:
+		{
+			return PerformLayersWeightmapsBatchedMerge(InUpdateLayersContentContext, MergeParams);
+		}
+		default:
+			check(false);
+		}
 	}
 
@@ -7926,7 +10812,6 @@
 		const TArray<FWeightmapLayerAllocationInfo>& WeightmapBaseLayerAllocation = Component->GetWeightmapLayerAllocations();
 
 		const TArray<UTexture2D*>& ComponentWeightmapTextures = Component->GetWeightmapTextures();
-		UTexture2D* Heightmap = Component->GetHeightmap();
 
 		for (auto& ItPair : Component->MaterialPerLOD)
 		{
@@ -8337,7 +11222,7 @@
 void ULandscapeComponent::GetLandscapeComponentNeighborsToRender(TSet<ULandscapeComponent*>& OutNeighborComponents) const
 {
 	ULandscapeInfo* LandscapeInfo = GetLandscapeInfo();
-	FIntPoint ComponentKey = GetSectionBase() / ComponentSizeQuads;
+	FIntPoint ComponentKey = GetComponentKey();
 
 	for (int32 IndexX = ComponentKey.X - 1; IndexX <= ComponentKey.X + 1; ++IndexX)
 	{
@@ -8355,7 +11240,7 @@
 void ULandscapeComponent::GetLandscapeComponentNeighbors3x3(TStaticArray<ULandscapeComponent*, 9>& OutNeighborComponents) const
 {
 	ULandscapeInfo* LandscapeInfo = GetLandscapeInfo();
-	FIntPoint ComponentKey = GetSectionBase() / ComponentSizeQuads;
+	FIntPoint ComponentKey = GetComponentKey();
 
 	int32 LinearIndex = 0;
 	for (int32 IndexY = ComponentKey.Y - 1; IndexY <= ComponentKey.Y + 1; ++IndexY)
@@ -8448,13 +11333,14 @@
 	FLandscapeDirtyOnlyInModeScope DirtyOnlyInMode(LandscapeInfo);
 
 	// If we went from local merge to global merge or vice versa, we need to reinitialize layers : 
+	ELandscapeEditLayersMergeMode EditLayersMergeMode = GetEditLayersMergeMode();
 	// If we toggled normal capture, we need to reinitialize layers
-	bool bUseLocalMerge = SupportsEditLayersLocalMerge();
 	bool bUseNormalCapture = HasNormalCaptureBPBrushLayer();
-	if ((bLandscapeLayersAreUsingLocalMerge != bUseLocalMerge) || (bLandscapeLayersAreInitializedForNormalCapture != bUseNormalCapture))
+	if ((CurrentEditLayersMergeMode != EditLayersMergeMode)
+		|| (bLandscapeLayersAreInitializedForNormalCapture != bUseNormalCapture))
 	{
 		RequestLayersInitialization(/*bInRequestContentUpdate = */true);
-		bLandscapeLayersAreUsingLocalMerge = bUseLocalMerge;
+		CurrentEditLayersMergeMode = EditLayersMergeMode;
 		bLandscapeLayersAreInitializedForNormalCapture = bUseNormalCapture;
 	}
 
@@ -9424,6 +12310,11 @@
 	return nullptr;
 }
 
+int32 ALandscape::GetLayerIndex(const FGuid& InLayerGuid) const
+{
+	return LandscapeEditLayers.IndexOfByPredicate([&InLayerGuid](const FLandscapeLayer& Other) { return Other.Guid == InLayerGuid; });
+}
+
 const FLandscapeLayer* ALandscape::GetLayerConst(const FGuid& InLayerGuid) const
 {
 	return LandscapeEditLayers.FindByPredicate([&InLayerGuid](const FLandscapeLayer& Other) { return Other.Guid == InLayerGuid; });
@@ -9482,26 +12373,28 @@
 	}
 }
 
-void ALandscape::DeleteLayer(int32 InLayerIndex)
+bool ALandscape::DeleteLayer(int32 InLayerIndex)
 {
 	ensure(HasLayersContent());
 
-	ULandscapeInfo* LandscapeInfo = GetLandscapeInfo();
 	const FLandscapeLayer* LayerStruct = GetLayerConst(InLayerIndex);
-	if (!LandscapeInfo || !LayerStruct)
+	if (!LayerStruct)
 	{
-		return;
+		return false;
 	}
 
 	Modify();
 	FGuid LayerGuid = LayerStruct->Guid;
 
 	// Clean up Weightmap usage in LandscapeProxies
-	LandscapeInfo->ForEachLandscapeProxy([&LayerGuid](ALandscapeProxy* Proxy)
+	if (ULandscapeInfo* LandscapeInfo = GetLandscapeInfo())
 	{
-		Proxy->DeleteLayer(LayerGuid);
-		return true;
-	});
+		LandscapeInfo->ForEachLandscapeProxy([&LayerGuid](ALandscapeProxy* Proxy)
+		{
+			Proxy->DeleteLayer(LayerGuid);
+			return true;
+		});
+	}
 
 	// We're about to remove the layer from our list, which will invalidate our LayerStruct pointer.
 	// We'll need to call OnLayerRemoved afterward, though, so keep pointer to the UObject.
@@ -9518,6 +12411,8 @@
 
 	// Request Update
 	RequestLayersContentUpdateForceAll();
+
+	return true;
 }
 
 void ALandscape::CollapseLayer(int32 InLayerIndex)
@@ -9993,8 +12888,7 @@
 
 void ALandscape::CreateDefaultLayer()
 {
-	ULandscapeInfo* LandscapeInfo = GetLandscapeInfo();
-	if (!LandscapeInfo || !CanHaveLayersContent())
+	if (!CanHaveLayersContent())
 	{
 		return;
 	}
@@ -10002,8 +12896,6 @@
 	check(LandscapeEditLayers.Num() == 0); // We can only call this function if we have no layers
 
 	CreateLayer(FName(TEXT("Layer")));
-	// Force update rendering resources
-	RequestLayersInitialization();
 }
 
 FLandscapeLayer* ALandscape::DuplicateLayerAndMoveBrushes(const FLandscapeLayer& InOtherLayer)
@@ -10043,8 +12935,7 @@
 
 int32 ALandscape::CreateLayer(FName InName, const TSubclassOf<ULandscapeEditLayerBase>& InEditLayerClass)
 {
-	ULandscapeInfo* LandscapeInfo = GetLandscapeInfo();
-	if (!LandscapeInfo || IsMaxLayersReached() || !CanHaveLayersContent())
+	if (IsMaxLayersReached() || !CanHaveLayersContent())
 	{
 		return INDEX_NONE;
 	}
@@ -10062,8 +12953,7 @@
 
 int32 ALandscape::CreateLayerFrom(const FLandscapeLayer& InLayer)
 {
-	ULandscapeInfo* LandscapeInfo = GetLandscapeInfo();
-	if (!LandscapeInfo || IsMaxLayersReached() || !CanHaveLayersContent())
+	if (IsMaxLayersReached() || !CanHaveLayersContent())
 	{
 		return INDEX_NONE;
 	}
@@ -10085,27 +12975,26 @@
 
 void ALandscape::OnLayerCreatedInternal(FLandscapeLayer& Layer)
 {
-	ULandscapeInfo* LandscapeInfo = GetLandscapeInfo();
-	check(LandscapeInfo != nullptr); 
-
 	// TODO: Might not be necessary eventually, if EditLayer has ability to trigger landscape updates
 	// and has access to its guid.
 	Layer.EditLayer->SetBackPointer(this);
-	
 	Layer.EditLayer->SetFlags(RF_Transactional);
-
 	Layer.EditLayer->OnLayerCreated(Layer);
 
 	// Create associated layer data in each landscape proxy
-	LandscapeInfo->ForEachLandscapeProxy([&Layer](ALandscapeProxy* Proxy)
+	if (ULandscapeInfo* LandscapeInfo = GetLandscapeInfo())
 	{
-		Proxy->AddLayer(Layer.Guid);
-		return true;
-	});
+		LandscapeInfo->ForEachLandscapeProxy([&Layer](ALandscapeProxy* Proxy)
+		{
+			Proxy->AddLayer(Layer.Guid);
+			return true;
+		});
+	}
 
 	// Request Update
+	// Force update rendering resources
+	RequestLayersInitialization(/*bInRequestContentUpdate = */true);
 	RequestSplineLayerUpdate(); // Request a spline update as well, in case we already had spline actors and a spline layers was just created
-	RequestLayersContentUpdateForceAll();
 }
 
 void ALandscape::AddLayersToProxy(ALandscapeProxy* InProxy)
@@ -10241,7 +13130,7 @@
 	return false;
 }
 
-int32 ALandscape::GetBrushLayer(class ALandscapeBlueprintBrushBase* InBrush) const
+int32 ALandscape::GetBrushLayer(const ALandscapeBlueprintBrushBase* InBrush) const
 {
 	for (int32 LayerIndex = 0; LayerIndex < LandscapeEditLayers.Num(); ++LayerIndex)
 	{
@@ -10406,6 +13295,33 @@
 #endif
 }
 
+bool FLandscapeLayerBrush::Affects(ELandscapeToolTargetType InTargetType, const FName& InWeightmapLayerName) const
+{
+#if WITH_EDITORONLY_DATA
+	if ((BlueprintBrush != nullptr) && BlueprintBrush->IsVisible())
+	{
+		switch (InTargetType)
+		{
+			case ELandscapeToolTargetType::Heightmap:
+			{
+				return BlueprintBrush->AffectsHeightmap();
+			}
+			case ELandscapeToolTargetType::Weightmap:
+			{
+				return BlueprintBrush->AffectsWeightmap() && ((InWeightmapLayerName == NAME_None) || BlueprintBrush->AffectsWeightmapLayer(InWeightmapLayerName));
+			}
+			case ELandscapeToolTargetType::Visibility:
+			{
+				return BlueprintBrush->AffectsVisibilityLayer();
+			}
+			default:
+				check(false);
+		}
+	}
+#endif // WITH_EDITORONLY_DATA
+	return false;
+}
+
 bool FLandscapeLayerBrush::AffectsHeightmap() const
 {
 #if WITH_EDITORONLY_DATA
@@ -10433,84 +13349,26 @@
 #endif
 }
 
-UTextureRenderTarget2D* FLandscapeLayerBrush::Render(bool InIsHeightmap, const FIntRect& InLandscapeExtent, UTextureRenderTarget2D* InLandscapeRenderTarget, const FName& InWeightmapLayerName)
-{
-	const ELandscapeToolTargetType LayerType = InIsHeightmap ? ELandscapeToolTargetType::Heightmap : ELandscapeToolTargetType::Weightmap;
-	FLandscapeBrushParameters BrushParameters(LayerType, InLandscapeRenderTarget, InWeightmapLayerName);
-
-	return RenderLayer(InLandscapeExtent, BrushParameters);
-}
-
-UTextureRenderTarget2D* FLandscapeLayerBrush::RenderLayer(const FIntRect& InLandscapeSize, const FLandscapeBrushParameters& InParameters)
+UTextureRenderTarget2D* FLandscapeLayerBrush::RenderLayer(const FLandscapeBrushParameters& InParameters)
 {
 #if WITH_EDITORONLY_DATA
 	TRACE_CPUPROFILER_EVENT_SCOPE(LandscapeLayers_LayerBrushRender);
-	if (((InParameters.LayerType == ELandscapeToolTargetType::Heightmap) && !AffectsHeightmap()) ||
-		((InParameters.LayerType == ELandscapeToolTargetType::Weightmap) && !AffectsWeightmapLayer(InParameters.WeightmapLayerName)) ||
-		((InParameters.LayerType == ELandscapeToolTargetType::Visibility) && !AffectsVisibilityLayer()))
+	if ((BlueprintBrush == nullptr)
+		|| ((InParameters.LayerType == ELandscapeToolTargetType::Heightmap) && !AffectsHeightmap())
+		|| ((InParameters.LayerType == ELandscapeToolTargetType::Weightmap) && !AffectsWeightmapLayer(InParameters.WeightmapLayerName))
+		|| ((InParameters.LayerType == ELandscapeToolTargetType::Visibility) && !AffectsVisibilityLayer()))
 	{
 		return nullptr;
 	}
 
-	if (Initialize(InLandscapeSize, InParameters.CombinedResult))
-	{
-		const TCHAR* LayerName = TEXT("");
-		
-		switch (InParameters.LayerType)
-		{
-		case ELandscapeToolTargetType::Heightmap:
-			LayerName = TEXT("LS Height");
-			break;
-		case ELandscapeToolTargetType::Weightmap:
-			LayerName = TEXT("LS Weight");
-			break;
-		case ELandscapeToolTargetType::Visibility:
-			LayerName = TEXT("LS VisibilityLayer");
-			break;
-		default:
-			check(false);
-		}
-
-		SCOPED_DRAW_EVENTF_GAMETHREAD(LandscapeLayers, TEXT("LandscapeLayers_RenderLayerBrush %s: %s")
-			, RHI_BREADCRUMB_FORCE_STRING_LITERAL(LayerName)
-			, BlueprintBrush->GetFName()
-		);
-
-		TGuardValue<bool> AutoRestore(GAllowActorScriptExecutionInEditor, true);
-		UTextureRenderTarget2D* Result = BlueprintBrush->RenderLayer(InParameters);
-
-		return Result;
-	}
-#endif
+	return BlueprintBrush->Execute(InParameters);
+#else  // WITH_EDITORONLY_DATA
 	return nullptr;
+#endif // !WITH_EDITORONLY_DATA
+
 }
 
-bool FLandscapeLayerBrush::Initialize(const FIntRect& InLandscapeExtent, UTextureRenderTarget2D* InLandscapeRenderTarget)
-{
-#if WITH_EDITORONLY_DATA
-	if (BlueprintBrush && InLandscapeRenderTarget)
-	{
-		if (ALandscape* Landscape = BlueprintBrush->GetOwningLandscape())
-		{
-			const FIntPoint NewLandscapeRenderTargetSize = FIntPoint(InLandscapeRenderTarget->SizeX, InLandscapeRenderTarget->SizeY);
-			FTransform NewLandscapeTransform = LandscapeEditLayersLocals::GetLandscapeExtentTransform(*Landscape, InLandscapeExtent);
-			FIntPoint NewLandscapeSize = InLandscapeExtent.Max - InLandscapeExtent.Min;
-			if (!LandscapeTransform.Equals(NewLandscapeTransform) || (LandscapeSize != NewLandscapeSize) || LandscapeRenderTargetSize != NewLandscapeRenderTargetSize)
-			{
-				LandscapeTransform = NewLandscapeTransform;
-				LandscapeRenderTargetSize = NewLandscapeRenderTargetSize;
-				LandscapeSize = NewLandscapeSize;
-
-				TGuardValue<bool> AutoRestore(GAllowActorScriptExecutionInEditor, true);
-				BlueprintBrush->Initialize(LandscapeTransform, LandscapeSize, LandscapeRenderTargetSize);
-			}
-			return true;
-		}
-	}
-#endif
-	return false;
-}
-
 #endif // WITH_EDITOR
 
+
 #undef LOCTEXT_NAMESPACE
