Introducing landscape batched merge (activated by CVar : landscape.EditLayersLocalMerge.Enable 2): complete refactor of the landscape edit layers merge algorithm to solve the taxing GPU memory require

diff --git a/Engine/Source/Runtime/Landscape/Private/LandscapeRender.cpp b/Engine/Source/Runtime/Landscape/Private/LandscapeRender.cpp
--- a/Engine/Source/Runtime/Landscape/Private/LandscapeRender.cpp
+++ b/Engine/Source/Runtime/Landscape/Private/LandscapeRender.cpp
@@ -54,7 +54,7 @@ LandscapeRender.cpp: New terrain rendering
 #include "RenderGraphBuilder.h"
 #include "Scalability.h"
 #include "Rendering/CustomRenderPass.h"
-#include "LandscapeUtils.h"
+#include "LandscapeUtilsPrivate.h"
 #include "SceneRendererInterface.h"
 
 using namespace UE::Landscape;
@@ -2950,7 +2950,7 @@ void FLandscapeComponentSceneProxy::GetDynamicMeshElements(const TArray<const FS
 		{
 			for (int32 DestinationMipIndex = SourceMipIndex + 1; DestinationMipIndex < NumRelevantMips; ++DestinationMipIndex)
 			{
-				int32 MipToMipDeltaIndex = UE::Landscape::ComputeMipToMipMaxDeltasIndex(SourceMipIndex, DestinationMipIndex, NumRelevantMips);
+				int32 MipToMipDeltaIndex = UE::Landscape::Private::ComputeMipToMipMaxDeltasIndex(SourceMipIndex, DestinationMipIndex, NumRelevantMips);
 				MipToMipInfoString += FString::Printf(TEXT("- %i->%i: %f\n"), SourceMipIndex, DestinationMipIndex, WorldSpaceMipToMipMaxDeltas[MipToMipDeltaIndex]);
 			}
 		}
@@ -4566,13 +4566,13 @@ bool FLandscapeComponentSceneProxy::ShouldInvalidateShadows(const FSceneView& In
 	check(DestinationMipIndex + 1 < NumRelevantMips);
 
 	// Evaluate the max delta for both SourceLODValue and DestinationLODValue against SourceMipIndex :
-	const int32 SourceMipToMipMaxDeltaIndex = UE::Landscape::ComputeMipToMipMaxDeltasIndex(SourceMipIndex, SourceMipIndex + 1, NumRelevantMips);
+	const int32 SourceMipToMipMaxDeltaIndex = UE::Landscape::Private::ComputeMipToMipMaxDeltasIndex(SourceMipIndex, SourceMipIndex + 1, NumRelevantMips);
 	const double SourceMipToMipMaxDelta = WorldSpaceMipToMipMaxDeltas[SourceMipToMipMaxDeltaIndex];
 	// MipToMipMaxDelta represents the maximum delta if we were to transition from SourceMipIndex to SourceMipIndex + 1 but we want to compute the error at SourceLODValue
 	//  so re-scale the delta within that range to evaluate the actual error : 
 	const double SourceMaxDelta = SourceMipToMipMaxDelta * (SourceLODValue - SourceMipIndex);
 
-	const int32 DestinationMipToMipMaxDeltaIndex = UE::Landscape::ComputeMipToMipMaxDeltasIndex(SourceMipIndex, DestinationMipIndex + 1, NumRelevantMips);
+	const int32 DestinationMipToMipMaxDeltaIndex = UE::Landscape::Private::ComputeMipToMipMaxDeltasIndex(SourceMipIndex, DestinationMipIndex + 1, NumRelevantMips);
 	const double DestinationMipToMipMaxDelta = WorldSpaceMipToMipMaxDeltas[DestinationMipToMipMaxDeltaIndex];
 	// MipToMipMaxDelta represents the maximum delta if we were to transition from SourceMipIndex to DestinationMipIndex + 1 but we want to compute the error at DestinationLODValue
 	//  so re-scale the delta within that range to evaluate the actual error : 
diff --git a/Engine/Source/Runtime/Landscape/Private/LandscapeUtils.cpp b/Engine/Source/Runtime/Landscape/Private/LandscapeUtils.cpp
--- a/Engine/Source/Runtime/Landscape/Private/LandscapeUtils.cpp
+++ b/Engine/Source/Runtime/Landscape/Private/LandscapeUtils.cpp
@@ -6,6 +6,8 @@
 #include "DataDrivenShaderPlatformInfo.h"
 #include "LandscapeLayerInfoObject.h"
 #include "LandscapeProxy.h"
+#include "LandscapeEditTypes.h"
+#include "Algo/Transform.h"
 
 #if WITH_EDITOR
 #include "EditorDirectories.h"
@@ -29,36 +31,27 @@ bool DoesPlatformSupportEditLayers(EShaderPlatform InShaderPlatform)
 		&& !IsMobilePlatform(InShaderPlatform);
 }
 
-int32 ComputeMaxDeltasOffsetForMip(int32 InMipIndex, int32 InNumRelevantMips)
+ELandscapeToolTargetTypeFlags GetLandscapeToolTargetTypeAsFlags(ELandscapeToolTargetType InTargetType)
 {
-	int32 Offset = 0;
-	for (int32 X = 0; X < InMipIndex; ++X)
-	{
-		Offset += InNumRelevantMips - 1 - X;
-	}
-	return Offset;
+	uint8 TargetTypeValue = static_cast<uint8>(InTargetType);
+	check(TargetTypeValue < static_cast<uint8>(ELandscapeToolTargetType::Count));
+	return static_cast<ELandscapeToolTargetTypeFlags>(1 << TargetTypeValue);
 }
 
-int32 ComputeMaxDeltasCountForMip(int32 InMipIndex, int32 InNumRelevantMips)
+ELandscapeToolTargetType GetLandscapeToolTargetTypeSingleFlagAsType(ELandscapeToolTargetTypeFlags InSingleFlag)
 {
-	return InNumRelevantMips - 1 - InMipIndex;
+	check(FMath::CountBits(static_cast<uint64>(InSingleFlag)) == 1);
+	uint32 Index = FMath::FloorLog2(static_cast<uint8>(InSingleFlag));
+	check(Index < static_cast<uint32>(ELandscapeToolTargetType::Count));
+	return static_cast<ELandscapeToolTargetType>(Index);
 }
 
-int32 ComputeMipToMipMaxDeltasIndex(int32 InSourceMipIndex, int32 InDestinationMipIndex, int32 InNumRelevantMips)
+FString GetLandscapeToolTargetTypeFlagsAsString(ELandscapeToolTargetTypeFlags InTargetTypeFlags)
 {
-	check((InSourceMipIndex >= 0) && (InSourceMipIndex < InNumRelevantMips));
-	check((InDestinationMipIndex > InSourceMipIndex) && (InDestinationMipIndex < InNumRelevantMips));
-	return ComputeMaxDeltasOffsetForMip(InSourceMipIndex, InNumRelevantMips) + InDestinationMipIndex - InSourceMipIndex - 1;
-}
-
-int32 ComputeMipToMipMaxDeltasCount(int32 InNumRelevantMips)
-{
-	int32 Count = 0;
-	for (int32 MipIndex = 0; MipIndex < InNumRelevantMips - 1; ++MipIndex)
-	{
-		Count += InNumRelevantMips - 1 - MipIndex;
-	}
-	return Count;
+	TArray<FString> TargetTypeStrings;
+	Algo::Transform(MakeFlagsRange(InTargetTypeFlags), TargetTypeStrings, [](ELandscapeToolTargetTypeFlags InTargetTypeFlag) 
+		{ return UEnum::GetValueAsString(GetLandscapeToolTargetTypeSingleFlagAsType(InTargetTypeFlag)); });
+	return *FString::Join(TargetTypeStrings, TEXT(","));
 }
 
 #if WITH_EDITOR
@@ -179,7 +172,7 @@ bool FBatchTextureCopy::ProcessTextureCopies()
 		return false;
 	}
 
-	// Populate source/destination maps to filter unique occurences
+	// Populate source/destination maps to filter unique occurrences
 	for (const TPair<FTextureCopyRequest, FTextureCopyChannelMapping>& CopyRequest : CopyRequests)
 	{
 		FSourceDataMipNumber& SourceData = Sources.Add(CopyRequest.Key.Source);
@@ -267,23 +260,6 @@ bool FBatchTextureCopy::ProcessTextureCopies()
 	return true;
 }
 
-int32 LandscapeMobileWeightTextureArray = 0;
-static FAutoConsoleVariableRef CVarLandscapeMobileWeightTextureArray(
-	TEXT("landscape.MobileWeightTextureArray"),
-	LandscapeMobileWeightTextureArray,
-	TEXT("Use Texture Arrays for weights on Mobile platforms"),
-	ECVF_ReadOnly | ECVF_MobileShaderChange);
-
-bool IsMobileWeightmapTextureArrayEnabled()
-{
-	return LandscapeMobileWeightTextureArray != 0;	
-}
-	
-bool UseWeightmapTextureArray(EShaderPlatform InPlatform)
-{
-	return IsMobilePlatform(InPlatform) && (LandscapeMobileWeightTextureArray != 0);	
-}
-
 FLayerInfoFinder::FLayerInfoFinder()
 {
 	const UClass* AssetClass = ULandscapeLayerInfoObject::StaticClass();
diff --git a/Engine/Source/Runtime/Landscape/Private/Materials/MaterialExpressionLandscapeLayerBlend.cpp b/Engine/Source/Runtime/Landscape/Private/Materials/MaterialExpressionLandscapeLayerBlend.cpp
--- a/Engine/Source/Runtime/Landscape/Private/Materials/MaterialExpressionLandscapeLayerBlend.cpp
+++ b/Engine/Source/Runtime/Landscape/Private/Materials/MaterialExpressionLandscapeLayerBlend.cpp
@@ -6,7 +6,7 @@
 #include "EngineGlobals.h"
 #include "MaterialCompiler.h"
 #include "DataDrivenShaderPlatformInfo.h"
-#include "LandscapeUtils.h"
+#include "LandscapeUtilsPrivate.h"
 #include "Materials/MaterialAttributeDefinitionMap.h"
 #include "Materials/MaterialExpressionShadingModel.h"
 
@@ -159,7 +159,7 @@ int32 UMaterialExpressionLandscapeLayerBlend::Compile(class FMaterialCompiler* C
 	TArray<int32> WeightCodes;
 	WeightCodes.Empty(Layers.Num());
 
-	const bool bTextureArrayEnabled = UE::Landscape::UseWeightmapTextureArray(Compiler->GetShaderPlatform());
+	const bool bTextureArrayEnabled = UE::Landscape::Private::UseWeightmapTextureArray(Compiler->GetShaderPlatform());
 	
 	for (int32 LayerIdx = 0; LayerIdx<Layers.Num(); LayerIdx++)
 	{
@@ -335,7 +335,7 @@ bool UMaterialExpressionLandscapeLayerBlend::GenerateHLSLExpression(FMaterialHLS
 		if (Layer.BlendType != LB_AlphaBlend)
 		{
 			const FExpression* WeightExpression = nullptr;
-			const bool bTextureArrayEnabled = UE::Landscape::IsMobileWeightmapTextureArrayEnabled();
+			const bool bTextureArrayEnabled = UE::Landscape::Private::IsMobileWeightmapTextureArrayEnabled();
 			verify(GenerateStaticTerrainLayerWeightExpression(Layer.LayerName, Layer.PreviewWeight, bTextureArrayEnabled, Generator, WeightExpression));
 
 			if (WeightExpression)
@@ -399,7 +399,7 @@ bool UMaterialExpressionLandscapeLayerBlend::GenerateHLSLExpression(FMaterialHLS
 		if (Layer.BlendType == LB_AlphaBlend)
 		{
 			const FExpression* WeightExpression = nullptr;
-			const bool bTextureArrayEnabled = UE::Landscape::IsMobileWeightmapTextureArrayEnabled();
+			const bool bTextureArrayEnabled = UE::Landscape::Private::IsMobileWeightmapTextureArrayEnabled();
 			verify(GenerateStaticTerrainLayerWeightExpression(Layer.LayerName, Layer.PreviewWeight, bTextureArrayEnabled, Generator, WeightExpression));
 
 			if (WeightExpression)
diff --git a/Engine/Source/Runtime/Landscape/Private/Materials/MaterialExpressionLandscapeLayerSample.cpp b/Engine/Source/Runtime/Landscape/Private/Materials/MaterialExpressionLandscapeLayerSample.cpp
--- a/Engine/Source/Runtime/Landscape/Private/Materials/MaterialExpressionLandscapeLayerSample.cpp
+++ b/Engine/Source/Runtime/Landscape/Private/Materials/MaterialExpressionLandscapeLayerSample.cpp
@@ -6,7 +6,7 @@
 #include "EngineGlobals.h"
 #include "MaterialCompiler.h"
 #include "Materials/Material.h"
-#include "LandscapeUtils.h"
+#include "LandscapeUtilsPrivate.h"
 
 #if WITH_EDITOR
 #include "MaterialHLSLGenerator.h"
@@ -42,7 +42,7 @@ UMaterialExpressionLandscapeLayerSample::UMaterialExpressionLandscapeLayerSample
 #if WITH_EDITOR
 int32 UMaterialExpressionLandscapeLayerSample::Compile(class FMaterialCompiler* Compiler, int32 OutputIndex)
 {
-	const bool bTextureArrayEnabled = UE::Landscape::UseWeightmapTextureArray(Compiler->GetShaderPlatform());
+	const bool bTextureArrayEnabled = UE::Landscape::Private::UseWeightmapTextureArray(Compiler->GetShaderPlatform());
 	const int32 WeightCode = Compiler->StaticTerrainLayerWeight(ParameterName, Compiler->Constant(PreviewWeight), bTextureArrayEnabled);
 	if (WeightCode == INDEX_NONE)
 	{
@@ -57,7 +57,7 @@ int32 UMaterialExpressionLandscapeLayerSample::Compile(class FMaterialCompiler*
 
 bool UMaterialExpressionLandscapeLayerSample::GenerateHLSLExpression(FMaterialHLSLGenerator& Generator, UE::HLSLTree::FScope& Scope, int32 OutputIndex, UE::HLSLTree::FExpression const*& OutExpression) const
 {
-	const bool bTextureArrayEnabled = UE::Landscape::IsMobileWeightmapTextureArrayEnabled();
+	const bool bTextureArrayEnabled = UE::Landscape::Private::IsMobileWeightmapTextureArrayEnabled();
 	return GenerateStaticTerrainLayerWeightExpression(ParameterName, PreviewWeight, bTextureArrayEnabled, Generator, OutExpression);
 }
 
diff --git a/Engine/Source/Runtime/Landscape/Private/Materials/MaterialExpressionLandscapeLayerSwitch.cpp b/Engine/Source/Runtime/Landscape/Private/Materials/MaterialExpressionLandscapeLayerSwitch.cpp
--- a/Engine/Source/Runtime/Landscape/Private/Materials/MaterialExpressionLandscapeLayerSwitch.cpp
+++ b/Engine/Source/Runtime/Landscape/Private/Materials/MaterialExpressionLandscapeLayerSwitch.cpp
@@ -8,7 +8,7 @@
 #include "MaterialHLSLGenerator.h"
 #include "MaterialHLSLTree.h"
 #include "Materials/Material.h"
-#include "LandscapeUtils.h"
+#include "LandscapeUtilsPrivate.h"
 
 #include UE_INLINE_GENERATED_CPP_BY_NAME(MaterialExpressionLandscapeLayerSwitch)
 
@@ -52,7 +52,7 @@ bool UMaterialExpressionLandscapeLayerSwitch::IsResultMaterialAttributes(int32 O
 
 int32 UMaterialExpressionLandscapeLayerSwitch::Compile(class FMaterialCompiler* Compiler, int32 OutputIndex)
 {
-	const bool bTextureArrayEnabled = UE::Landscape::UseWeightmapTextureArray(Compiler->GetShaderPlatform());
+	const bool bTextureArrayEnabled = UE::Landscape::Private::UseWeightmapTextureArray(Compiler->GetShaderPlatform());
 	const int32 WeightCode = Compiler->StaticTerrainLayerWeight(
 		ParameterName,
 		PreviewUsed ? Compiler->Constant(1.0f) : INDEX_NONE,
diff --git a/Engine/Source/Runtime/Landscape/Private/Materials/MaterialExpressionLandscapeLayerWeight.cpp b/Engine/Source/Runtime/Landscape/Private/Materials/MaterialExpressionLandscapeLayerWeight.cpp
--- a/Engine/Source/Runtime/Landscape/Private/Materials/MaterialExpressionLandscapeLayerWeight.cpp
+++ b/Engine/Source/Runtime/Landscape/Private/Materials/MaterialExpressionLandscapeLayerWeight.cpp
@@ -6,7 +6,7 @@
 #include "EngineGlobals.h"
 #include "MaterialCompiler.h"
 #include "Materials/Material.h"
-#include "LandscapeUtils.h"
+#include "LandscapeUtilsPrivate.h"
 
 #include UE_INLINE_GENERATED_CPP_BY_NAME(MaterialExpressionLandscapeLayerWeight)
 
@@ -59,7 +59,7 @@ bool UMaterialExpressionLandscapeLayerWeight::IsResultMaterialAttributes(int32 O
 
 int32 UMaterialExpressionLandscapeLayerWeight::Compile(class FMaterialCompiler* Compiler, int32 OutputIndex)
 {
-	const bool bTextureArrayEnabled = UE::Landscape::UseWeightmapTextureArray(Compiler->GetShaderPlatform());
+	const bool bTextureArrayEnabled = UE::Landscape::Private::UseWeightmapTextureArray(Compiler->GetShaderPlatform());
 	const int32 BaseCode = Base.Expression
 		                       ? Base.Compile(Compiler)
 		                       : Compiler->Constant3(static_cast<float>(ConstBase.X), static_cast<float>(ConstBase.Y), static_cast<float>(ConstBase.Z));
diff --git a/Engine/Source/Runtime/Landscape/Private/Materials/MaterialExpressionLandscapeVisibilityMask.cpp b/Engine/Source/Runtime/Landscape/Private/Materials/MaterialExpressionLandscapeVisibilityMask.cpp
--- a/Engine/Source/Runtime/Landscape/Private/Materials/MaterialExpressionLandscapeVisibilityMask.cpp
+++ b/Engine/Source/Runtime/Landscape/Private/Materials/MaterialExpressionLandscapeVisibilityMask.cpp
@@ -6,7 +6,7 @@
 #include "EngineGlobals.h"
 #include "MaterialCompiler.h"
 #include "Materials/Material.h"
-#include "LandscapeUtils.h"
+#include "LandscapeUtilsPrivate.h"
 
 #if WITH_EDITOR
 #include "MaterialHLSLGenerator.h"
@@ -45,7 +45,7 @@ UMaterialExpressionLandscapeVisibilityMask::UMaterialExpressionLandscapeVisibili
 #if WITH_EDITOR
 int32 UMaterialExpressionLandscapeVisibilityMask::Compile(class FMaterialCompiler* Compiler, int32 OutputIndex)
 {
-	const bool bTextureArrayEnabled = UE::Landscape::UseWeightmapTextureArray(Compiler->GetShaderPlatform());
+	const bool bTextureArrayEnabled = UE::Landscape::Private::UseWeightmapTextureArray(Compiler->GetShaderPlatform());
 	int32 MaskLayerCode = Compiler->StaticTerrainLayerWeight(ParameterName, Compiler->Constant(0.f), bTextureArrayEnabled);
 	return MaskLayerCode == INDEX_NONE ? Compiler->Constant(1.f) : Compiler->Sub(Compiler->Constant(1.f), MaskLayerCode);
 }
@@ -55,7 +55,7 @@ bool UMaterialExpressionLandscapeVisibilityMask::GenerateHLSLExpression(FMateria
 	using namespace UE::HLSLTree;
 
 	const FExpression* MaskLayerExpression = nullptr;
-	const bool bTextureArrayEnabled = UE::Landscape::IsMobileWeightmapTextureArrayEnabled();
+	const bool bTextureArrayEnabled = UE::Landscape::Private::IsMobileWeightmapTextureArrayEnabled();
 	verify(GenerateStaticTerrainLayerWeightExpression(ParameterName, 0.f, bTextureArrayEnabled, Generator, MaskLayerExpression));
 
 	FTree& Tree = Generator.GetTree();
diff --git a/Engine/Source/Runtime/Landscape/Public/GlobalMergeLegacySupportUtil.h b/Engine/Source/Runtime/Landscape/Public/GlobalMergeLegacySupportUtil.h
--- a/Engine/Source/Runtime/Landscape/Public/GlobalMergeLegacySupportUtil.h
+++ b/Engine/Source/Runtime/Landscape/Public/GlobalMergeLegacySupportUtil.h
@@ -38,9 +38,7 @@ public:
 	 * is called. This function calls the virtual InitializeAsBlueprintBrush and RenderLayerAsBlueprintBrush
 	 * functions, which are equivalent to a blueprint brush actor's Initialize_Native and RenderLayer_Native.
 	 */
-	UTextureRenderTarget2D* RenderAsBlueprintBrush(const FLandscapeBrushParameters& InParameters, const FTransform& LandscapeTransform, 
-		// TODO: Someday will be part of FLandscapeBrushParameters
-		const FIntPoint& RenderAreaSize);
+	UTextureRenderTarget2D* RenderAsBlueprintBrush(const FLandscapeBrushParameters& InParameters, const FTransform& LandscapeTransform);
 
 	// Called in the same places as the equivalent methods on FLandscapeLayerBrush.
 	virtual bool AffectsHeightmapAsBlueprintBrush() const { return false; }
diff --git a/Engine/Source/Runtime/Landscape/Public/LandscapeBlueprintBrushBase.h b/Engine/Source/Runtime/Landscape/Public/LandscapeBlueprintBrushBase.h
--- a/Engine/Source/Runtime/Landscape/Public/LandscapeBlueprintBrushBase.h
+++ b/Engine/Source/Runtime/Landscape/Public/LandscapeBlueprintBrushBase.h
@@ -6,6 +6,7 @@
 #include "UObject/ObjectMacros.h"
 #include "GameFramework/Actor.h"
 #include "LandscapeEditTypes.h"
+#include "LandscapeEditLayerRenderer.h"
 
 #include "LandscapeBlueprintBrushBase.generated.h"
 
@@ -16,24 +17,29 @@ struct FLandscapeBrushParameters
 {
 	GENERATED_BODY()
 
-	FLandscapeBrushParameters()
-		: LayerType(ELandscapeToolTargetType::Invalid)
-		, CombinedResult(nullptr)
-		, WeightmapLayerName()
-	{}
+	FLandscapeBrushParameters() = default;
 
+	UE_DEPRECATED(5.4, "Use the other constructor")
 	FLandscapeBrushParameters(const ELandscapeToolTargetType& InLayerType, TObjectPtr<UTextureRenderTarget2D> InCombinedResult, const FName& InWeightmapLayerName = FName())
-		: LayerType(InLayerType)
-		, CombinedResult(InCombinedResult)
+		: CombinedResult(InCombinedResult)
+		, LayerType(InLayerType)
 		, WeightmapLayerName(InWeightmapLayerName)
 	{}
 
+	FLandscapeBrushParameters(bool bInIsHeightmapMerge, const FTransform& InRenderAreaWorldTransform, const FIntPoint& InRenderAreaSize, UTextureRenderTarget2D* InCombinedResult, const FName& InWeightmapLayerName = FName());
+
+	UPROPERTY(Category = "Settings", EditAnywhere, BlueprintReadWrite)
+	FTransform RenderAreaWorldTransform;
+
 	UPROPERTY(Category = "Settings", EditAnywhere, BlueprintReadWrite)
-	ELandscapeToolTargetType LayerType;
+	FIntPoint RenderAreaSize = FIntPoint(MAX_int32, MAX_int32);
 
 	UPROPERTY(Category = "Settings", EditAnywhere, BlueprintReadWrite)
 	TObjectPtr<UTextureRenderTarget2D> CombinedResult;
 
+	UPROPERTY(Category = "Settings", EditAnywhere, BlueprintReadWrite)
+	ELandscapeToolTargetType LayerType = ELandscapeToolTargetType::Invalid;
+
 	UPROPERTY(Category = "Settings", EditAnywhere, BlueprintReadWrite)
 	FName WeightmapLayerName;
 };
@@ -41,6 +47,10 @@ struct FLandscapeBrushParameters
 
 UCLASS(Abstract, NotBlueprintable, MinimalAPI)
 class ALandscapeBlueprintBrushBase : public AActor
+#if CPP && WITH_EDITOR // UHT doesn't support inheriting from namespaced class
+	, public ILandscapeEditLayerRenderer
+	, public UE::Landscape::EditLayers::IEditLayerRendererProvider
+#endif // CPP && WITH_EDITOR
 {
 	GENERATED_UCLASS_BODY()
 
@@ -69,22 +79,24 @@ protected:
 
 	uint32 LastRequestLayersContentUpdateFrameNumber;
 
+	FTransform CurrentRenderAreaWorldTransform;
+	FIntPoint CurrentRenderAreaSize = FIntPoint(MAX_int32, MAX_int32);
+	FIntPoint CurrentRenderTargetSize = FIntPoint(MAX_int32, MAX_int32);
+
 	bool bCaptureBoundaryNormals = false;	// HACK [chris.tchou] remove once we have a better boundary normal solution
 #endif
 
 public:
-	virtual void Initialize_Native(const FTransform& InLandscapeTransform, const FIntPoint& InLandscapeSize, const FIntPoint& InLandscapeRenderTargetSize) {}
-
 	UFUNCTION(BlueprintNativeEvent, meta = (DeprecatedFunction, DeprecationMessage = "Please use RenderLayer instead."))
 	LANDSCAPE_API UTextureRenderTarget2D* Render(bool InIsHeightmap, UTextureRenderTarget2D* InCombinedResult, const FName& InWeightmapLayerName);
 
 	UFUNCTION(BlueprintNativeEvent)
 	LANDSCAPE_API UTextureRenderTarget2D* RenderLayer(const FLandscapeBrushParameters& InParameters);
-
 	LANDSCAPE_API virtual UTextureRenderTarget2D* RenderLayer_Native(const FLandscapeBrushParameters& InParameters);
 
 	UFUNCTION(BlueprintNativeEvent)
 	LANDSCAPE_API void Initialize(const FTransform& InLandscapeTransform, const FIntPoint& InLandscapeSize, const FIntPoint& InLandscapeRenderTargetSize);
+	LANDSCAPE_API virtual void Initialize_Native(const FTransform& InLandscapeTransform, const FIntPoint& InLandscapeSize, const FIntPoint& InLandscapeRenderTargetSize) {}
 
 	UFUNCTION(BlueprintCallable, Category = "Landscape")
 	LANDSCAPE_API void RequestLandscapeUpdate(bool bInUserTriggered = false);
@@ -97,6 +109,20 @@ public:
 	LANDSCAPE_API void SetCanAffectVisibilityLayer(bool bInCanAffectVisibilityLayer);
 
 #if WITH_EDITOR
+	//~ Begin ILandscapeEditLayerRenderer implementation
+	LANDSCAPE_API virtual void GetRendererStateInfo(const ULandscapeInfo* InLandscapeInfo, 
+		UE::Landscape::EditLayers::FEditLayerTargetTypeState& OutSupportedTargetTypeState, UE::Landscape::EditLayers::FEditLayerTargetTypeState& OutEnabledTargetTypeState, TArray<TSet<FName>>& OutRenderGroups) const override;
+	LANDSCAPE_API virtual TArray<UE::Landscape::EditLayers::FEditLayerRenderItem> GetRenderItems(const ULandscapeInfo* InLandscapeInfo) const override;
+	LANDSCAPE_API virtual void RenderLayer(ILandscapeEditLayerRenderer::FRenderParams& InRenderParams) override;
+	LANDSCAPE_API virtual FString GetEditLayerRendererDebugName() const override;
+	//~ End ILandscapeEditLayerRenderer implementation
+
+	//~ Begin UE::Landscape::EditLayers::IEditLayerRendererProvider implementation
+	LANDSCAPE_API TArray<UE::Landscape::EditLayers::FEditLayerRendererState> GetEditLayerRendererStates(const ULandscapeInfo* InLandscapeInfo, bool bInSkipBrush);
+	//~ End UE::Landscape::EditLayers::IEditLayerRendererProvider implementation
+
+	UTextureRenderTarget2D* Execute(const FLandscapeBrushParameters& InParameters);
+
 	LANDSCAPE_API virtual void CheckForErrors() override;
 
 	LANDSCAPE_API virtual void GetRenderDependencies(TSet<UObject*>& OutDependencies);
@@ -104,14 +130,17 @@ public:
 	LANDSCAPE_API virtual void SetOwningLandscape(class ALandscape* InOwningLandscape);
 	LANDSCAPE_API class ALandscape* GetOwningLandscape() const;
 
+	/** CanAffect... methods indicate the brush has the _capacity_ to affect this or that aspect of the landscape. 
+	*  Note: it doesn't mean the brush currently affects it : the Affects... methods are used for that. 
+	*/
 	bool CanAffectHeightmap() const { return AffectHeightmap; }
 	bool CanAffectWeightmap() const { return AffectWeightmap; }
 	bool CanAffectVisibilityLayer() const { return AffectVisibilityLayer; }
-
-	virtual bool AffectsHeightmap() const { return CanAffectHeightmap(); }
-	virtual bool AffectsWeightmap() const { return CanAffectWeightmap(); }
+	LANDSCAPE_API virtual bool CanAffectWeightmapLayer(const FName& InLayerName) const;
+	LANDSCAPE_API virtual bool AffectsHeightmap() const { return CanAffectHeightmap(); }
+	LANDSCAPE_API virtual bool AffectsWeightmap() const { return CanAffectWeightmap(); }
 	LANDSCAPE_API virtual bool AffectsWeightmapLayer(const FName& InLayerName) const;
-	virtual bool AffectsVisibilityLayer() const { return CanAffectVisibilityLayer(); }
+	LANDSCAPE_API virtual bool AffectsVisibilityLayer() const { return CanAffectVisibilityLayer(); }
 
 	bool IsVisible() const { return bIsVisible; }
 	LANDSCAPE_API bool IsLayerUpdatePending() const;
diff --git a/Engine/Source/Runtime/Landscape/Public/LandscapeEditLayerRenderer.h b/Engine/Source/Runtime/Landscape/Public/LandscapeEditLayerRenderer.h
new file mode 100644
--- /dev/null
+++ b/Engine/Source/Runtime/Landscape/Public/LandscapeEditLayerRenderer.h
@@ -0,0 +1,845 @@
+// Copyright Epic Games, Inc. All Rights Reserved.
+
+#pragma once
+
+#include "CoreMinimal.h"
+#include "LandscapeEditTypes.h"
+#include "Engine/TextureRenderTarget2D.h"
+#include "RHIAccess.h"
+#include "EngineDefines.h"
+
+#include "LandscapeEditLayerRenderer.generated.h"
+
+class ALandscape;
+class ULandscapeComponent;
+class ULandscapeInfo;
+class ULandscapeLayerInfoObject;
+class ULandscapeScratchRenderTarget;
+class ULandscapeEditLayerRenderer;
+class ILandscapeEditLayerRenderer;
+
+namespace UE::Landscape::EditLayers
+{
+
+// ----------------------------------------------------------------------------------
+
+FString ConvertTargetLayerNamesToString(const TArrayView<const FName>& InTargetLayerNames);
+
+
+// ----------------------------------------------------------------------------------
+
+// Must match EEditLayerWeightmapBlendMode in LandscapeLayersWeightmapsPS.usf / LandscapeEditLayersWeightmapsPS.usf
+enum class EWeightmapBlendMode : uint32
+{
+	Additive = 0,
+	Subtractive,
+
+	Num,
+};
+
+// Must match EWeightmapPaintLayerFlags in LandscapeLayersWeightmapsPS.usf / LandscapeEditLayersWeightmapsPS.usf
+enum class EWeightmapPaintLayerFlags : uint32
+{
+	IsVisibilityLayer = (1 << 0), // This paint layer is the visibility layer
+	IsWeightBlended = (1 << 1), // Blend the paint layer's value with all the other paint layers weights
+
+	None = 0
+};
+ENUM_CLASS_FLAGS(EWeightmapPaintLayerFlags);
+
+// ----------------------------------------------------------------------------------
+
+// Must match FWeightmapPaintLayerInfo in LandscapeLayersWeightmapsPS.usf / LandscapeEditLayersWeightmaps.usf
+struct FWeightmapPaintLayerInfo
+{
+	EWeightmapPaintLayerFlags Flags = EWeightmapPaintLayerFlags::None; // Additional info about this paint layer
+};
+
+
+// ----------------------------------------------------------------------------------
+
+#if WITH_EDITOR
+/** 
+ * FEditLayerTargetTypeState fully describes the state of an edit layer renderer wrt its target types. It's named after the enum "ELandscapeToolTargetType" in order to tell
+ *  whether the renderer's heightmaps and/or visibility and/or weightmaps are enabled (and if so, which weightmap is enabled exactly)  
+ *  It is meant to be provided by the edit layer renderer's GetRendererStateInfo implementation.
+ */
+struct FEditLayerTargetTypeState
+{
+	friend class ILandscapeEditLayerRenderer;
+
+	LANDSCAPE_API FEditLayerTargetTypeState() = default;
+	LANDSCAPE_API FEditLayerTargetTypeState(ELandscapeToolTargetTypeFlags InTargetTypeMask, const TArrayView<const FName>& InSupportedWeightmaps = TArrayView<const FName>());
+	
+	/**
+	 * Indicates whether a given target type is currently active in this state
+	 * @param InTargetType the requested target type (heightmap/weightmap/visibility)
+	 * @param InWeightmapLayerName (optional) is the requested weightmap, only relevant for the ELandscapeToolTargetType::Weightmap case
+	 */
+	LANDSCAPE_API bool IsActive(ELandscapeToolTargetType InTargetType, FName InWeightmapLayerName = NAME_None) const;
+
+	/** Returns the currently active weightmaps, if Weightmap is amongst the supported target types */
+	LANDSCAPE_API const TArray<FName>& GetActiveWeightmaps() const;
+	
+	/** Returns the target type mask (i.e. same as ELandscapeToolTargetType, but as bit flags) */
+	LANDSCAPE_API ELandscapeToolTargetTypeFlags GetTargetTypeMask() const { return TargetTypeMask; }
+
+	/** Sets the target type mask (i.e. same as ELandscapeToolTargetType, but as bit flags) */
+	LANDSCAPE_API void SetTargetTypeMask(ELandscapeToolTargetTypeFlags InTargetTypeMask);
+
+	/** Adds the target type in parameter to the mask of active target types */
+	LANDSCAPE_API void AddTargetType(ELandscapeToolTargetType InTargetType);
+
+	/** Appends the target type mask in parameter to the mask of active target types */
+	LANDSCAPE_API void AddTargetTypeMask(ELandscapeToolTargetTypeFlags InTargetTypeMask);
+
+	/** Removes a single target type from the mask of active target types */
+	LANDSCAPE_API void RemoveTargetType(ELandscapeToolTargetType InTargetType);
+
+	/** Removes the target type mask in parameter from the mask of active target types */
+	LANDSCAPE_API void RemoveTargetTypeMask(ELandscapeToolTargetTypeFlags InTargetTypeMask);
+
+	/** Adds a weightmap to the list of active weightmaps (make sure ELandscapeToolTargetType::Weightmap is amongst the supported target types) */
+	LANDSCAPE_API void AddWeightmap(FName InWeightmapLayerName);
+
+	/** Removes a weightmap from the list of active weightmaps */
+	LANDSCAPE_API void RemoveWeightmap(FName InWeightmapLayerName);
+
+	/**
+	 * Returns the "intersection" (AND operation) between the target type state and the one in parameter. e.g. if state S0 is:
+	 *  - (---------|Visibility|Weightmap) with active weightmaps (A|B|C|-), and state S1 is:
+	 *  - (Heightmap|----------|Weightmap) with active weightmaps (-|-|C|D), then state S0.Intersect(S1) is:
+	 *  - (---------|----------|Weightmap) with active weightmaps (-|-|C|-)
+	 */
+	FEditLayerTargetTypeState Intersect(const FEditLayerTargetTypeState& InOther) const;
+
+	bool operator == (const FEditLayerTargetTypeState& InOther) const;
+
+	FString ToString() const;
+
+private:
+	/** Bitmask of the target types that are supported  */
+	ELandscapeToolTargetTypeFlags TargetTypeMask = ELandscapeToolTargetTypeFlags::None;
+	
+	/** List of weightmaps that are supported for the ELandscapeToolTargetType::Weightmap/ELandscapeToolTargetType::Visibility type */
+	TArray<FName> Weightmaps;
+};
+
+
+// ----------------------------------------------------------------------------------
+
+/**
+ * FEditLayerRendererState describes the entire state of an edit layer renderer : what it is capable of doing (SupportedTargetTypeState, immutable) and what it is currently doing (EnabledTargetTypeState, mutable)
+ *  These states are provided by IEditLayerRendererProvider in order to describe both what the renderer can do and what it currently does by default. e.g. a disabled edit layer supports rendering heightmaps but
+ *  its enabled state for the heightmap target type is false. This way, the user can selectively enable it at merge time without altering the entire landscape's state (i.e. just for the purpose of a specific merge render)
+ *  A target type must be both supported and enabled on a given edit layer renderer in order for this renderer to render anything. 
+ *  It also describes the render groups this renderer needs when rendering its weightmap (i.e. which weightmap needs to be rendered with which weightmaps : e.g. for weight blending)
+ */
+struct FEditLayerRendererState
+{
+	LANDSCAPE_API FEditLayerRendererState() = default;
+	LANDSCAPE_API explicit FEditLayerRendererState(ILandscapeEditLayerRenderer* InRenderer, const ULandscapeInfo* InLandscapeInfo);
+
+	/** Returns the edit layer renderer which this state relates to */
+	LANDSCAPE_API ILandscapeEditLayerRenderer* GetRenderer() const { return Renderer; }
+
+	/**
+	 * Indicates whether a given target type is currently supported by this renderer
+	 * @param InTargetType the requested target type (heightmap/weightmap/visibility)
+	 * @param InWeightmapLayerName (optional) is the requested weightmap, only relevant for the ELandscapeToolTargetType::Weightmap case
+	 */
+	LANDSCAPE_API bool IsTargetSupported(ELandscapeToolTargetType InTargetType, FName InWeightmapLayerName = NAME_None) const;
+
+	/** Returns a mask of all target types supported by this renderer */
+	LANDSCAPE_API const FEditLayerTargetTypeState& GetSupportedTargetTypeState() const { return SupportedTargetTypeState; }
+
+	/** Returns a mask of all target types supported by this renderer */
+	LANDSCAPE_API const FEditLayerTargetTypeState& GetEnabledTargetTypeState() const { return EnabledTargetTypeState; }
+
+	/** Returns a list of all weightmaps supported by this renderer (only relevant for ELandscapeToolTargetType::Weightmap (and ELandscapeToolTargetType::Visibility))*/
+	LANDSCAPE_API const TArray<FName>& GetSupportedTargetWeightmaps() const;
+
+	/** Mutates the EnabledTargetTypeState by adding the target type in parameter to the mask of active target types */
+	LANDSCAPE_API void EnableTargetType(ELandscapeToolTargetType InTargetType);
+	
+	/** Mutates the EnabledTargetTypeState by adding the target type mask in parameter to the mask of active target types */
+	LANDSCAPE_API void EnableTargetTypeMask(ELandscapeToolTargetTypeFlags InTargetTypeMask);
+
+	/** Mutates the EnabledTargetTypeState by removing the target type in parameter from the mask of active target types */
+	LANDSCAPE_API void DisableTargetType(ELandscapeToolTargetType InTargetType);
+
+	/** Mutates the EnabledTargetTypeState by removing the target type mask in parameter from the mask of active target types */
+	LANDSCAPE_API void DisableTargetTypeMask(ELandscapeToolTargetTypeFlags InTargetTypeMask);
+
+	/**
+	 * Indicates whether a given target type is currently enabled by this renderer
+	 * @param InTargetType the requested target type (heightmap/weightmap/visibility)
+	 * @param InWeightmapLayerName (optional) is the requested weightmap, only relevant for the ELandscapeToolTargetType::Weightmap case
+	 */
+	LANDSCAPE_API bool IsTargetEnabled(ELandscapeToolTargetType InTargetType, FName InWeightmapLayerName = NAME_None) const;
+
+	/** Mutates the EnabledTargetTypeState by adding the weightmap in parameter to the list of enabled weightmaps */
+	LANDSCAPE_API bool EnableTarget(ELandscapeToolTargetType InTargetType, FName InWeightmapLayerName = NAME_None);
+
+	/** Mutates the EnabledTargetTypeState by removing the weightmap in parameter from the list of enabled weightmaps */
+	LANDSCAPE_API void DisableTarget(FName InWeightmapLayerName);
+
+	/** Returns the list of all weightmaps currently enabled */
+	LANDSCAPE_API TArray<FName> GetEnabledTargetWeightmaps() const;
+
+	/** Returns the render groups associated with this renderer. A render group is a set of target layers (weightmaps) that depend on one another in order to produce the output target layers.
+	 This allows to implement "horizontal blending", where weightmaps can be blended with one another at each step of the landscape edit layers merge algorithm */
+	LANDSCAPE_API TArray<TSet<FName>> GetWeightmapRenderGroups() const { return RenderGroups; }
+
+private:
+	ILandscapeEditLayerRenderer* Renderer = nullptr;
+	
+	/** Struct that contains the supported target types and the which weightmap are supported by this renderer. Immutable. */
+	FEditLayerTargetTypeState SupportedTargetTypeState;
+	
+	/** Struct that contains the enabled target types and the which weightmap are currently enabled by this renderer. Can be set by the user using EnableTargetType. */
+	FEditLayerTargetTypeState EnabledTargetTypeState;
+
+	/* List of groups of target layers that this renderer requires to be rendered together. All target layers listed in SupportedTargetTypeState must belong to 1 (and 1 only) render group. */
+	TArray<TSet<FName>> RenderGroups;
+};
+
+
+// ----------------------------------------------------------------------------------
+
+/** 
+ * Params struct passed to the merge function. It contains everything needed for requesting a given set of target layers (for weightmaps) on a given number of components and for a certain configuration of edit layers
+ *  Note that this is what is requested by the caller, but in practice, there might be more renderers (e.g. some might get added e.g. legacy weight-blending, some removed because they turn out to be disabled...)
+ *  and more weightmaps being rendered (e.g. a requested weightmap might depend on another one that has not been requested), or less (e.g. a requested weightmap is actually invalid) :
+ */
+struct FMergeRenderParams
+{
+	FMergeRenderParams(bool bInIsHeightmapMerge, ALandscape* InLandscape, TArray<ULandscapeComponent*> InComponentsToMerge, const TArrayView<FEditLayerRendererState>& InEditLayerRendererStates, const TSet<FName>& InWeightmapLayerNames = {});
+
+	// COMMENT [jonathan.bard] : purposefully don't use ELandscapeToolTargetType here as ELandscapeToolTargetType::Weightmap and ELandscapeToolTargetType::Visibility have to be processed together (because of weightmap packing, 
+	//  which means a visibility weightmap could be another channel of a texture which contains weightmap up to 3 other weightmaps, so we have to resolve the 4 channels altogether). 
+	//  Note: this could change when SUPPORTS_LANDSCAPE_EDITORONLY_UBER_MATERIAL is done...
+	/** Type of merge being requested */
+	bool bIsHeightmapMerge = true;
+	
+	/** Parent landscape actor to which all components to merge belong */
+	ALandscape* Landscape = nullptr;
+
+	/** List of components that need merging */
+	TArray<ULandscapeComponent*> ComponentsToMerge;
+
+	/** Requested states for every edit layer renderer participating to the merge */
+	TArray<FEditLayerRendererState> EditLayerRendererStates;
+	
+	/** List of weightmap layers being requested. */
+	TSet<FName> WeightmapLayerNames;
+};
+
+
+/**
+ * Defines an individual render step of the batch merge
+ */
+struct FMergeRenderStep
+{
+	enum class EType
+	{
+		RenderLayer, // Perform the rendering of a render group on an edit layer on a given world region (i.e. in a batch)
+		SignalBatchMergeGroupDone, // Final step when rendering a render group on a given world region (i.e. in a batch) : allows to retrieve the result of the merge and do something with it (e.g. resolve the final textures)
+	};
+
+	FMergeRenderStep(EType InType, const FEditLayerRendererState& InRendererState, const TBitArray<>& InRenderGroupBitIndices, const TArrayView<ULandscapeComponent*>& InComponentsToRender)
+		: Type(InType)
+		, RendererState(InRendererState)
+		, RenderGroupBitIndices(InRenderGroupBitIndices)
+		, ComponentsToRender(InComponentsToRender)
+	{}
+
+	/** Type of operation for this step */
+	EType Type = EType::RenderLayer;
+	
+	/** Renderer state to be used this step (only when Type == EType::RenderLayer). This includes the renderer as well as its precise step (e.g. which weightmap are supported? which are enabled?) */
+	FEditLayerRendererState RendererState;
+	
+	/** List of target layers being involved in this step. Each bit in that bit array corresponds to an entry in FMergeRenderContext's AllTargetLayerNames */
+	TBitArray<> RenderGroupBitIndices;
+
+	/** List of components involved in this step */
+	TArray<ULandscapeComponent*> ComponentsToRender;
+};
+
+
+// ----------------------------------------------------------------------------------
+
+/** 
+ * Defines an individual render batch when merging the landscape. A batch corresponds to a render group on a world's region, i.e. a set of weightmaps (or the heightmap) to render on a 
+ *  portion of the world. Each batch is composed of a series of render steps. 
+ */
+struct FMergeRenderBatch
+{
+	bool operator < (const FMergeRenderBatch& InOther) const;
+
+	ALandscape* Landscape = nullptr;
+
+	/** Section of the landscape being covered by this batch (in landscape vertex coordinates, inclusive bounds) */
+	FIntRect SectionRect;
+
+	/** Resolution of the render target needed for this batch (including duplicate borders) */
+	// TODO [jonathan.bard] : rename EffectiveResolution and make private maybe?
+	FIntPoint Resolution = FIntPoint(ForceInitToZero);
+	
+	FIntPoint MinComponentKey = FIntPoint(MAX_int32, MAX_int32);
+	FIntPoint MaxComponentKey = FIntPoint(MIN_int32, MIN_int32);
+	
+	/** Sequential list of rendering operations that need to be performed to fully render this batch*/
+	TArray<FMergeRenderStep> RenderSteps;
+	
+	/** List of all components involved in this batch */
+	TSet<ULandscapeComponent*> ComponentsToRender;
+
+	/** List of all target layers being rendered in this batch (i.e. bitwise OR of all of the render steps' RenderGroupBitIndices). Each bit in that bit array corresponds to an entry in FMergeRenderContext's AllTargetLayerNames  */
+	TBitArray<> TargetLayerNameBitIndices;
+
+	/** List of components involved in this batch and the target layers they're writing to (redundant with ComponentsToRender but we keep the latter for convenience)
+	 (each bit corresponds to a target layer name in FMergeRenderContext's AllTargetLayerNames) */
+	TMap<ULandscapeComponent*, TBitArray<>> ComponentToTargetLayerBitIndices;
+
+	/** Reverse lookup of ComponentToTargetLayerBitIndices : one entry per element in ComponentToTargetLayerBitIndices, each entry containing all of the components involved in this merge for this target layer */
+	TArray<TSet<ULandscapeComponent*>> TargetLayersToComponents;
+
+public:
+	FIntPoint GetRenderTargetResolution(bool bInWithDuplicateBorders) const;
+	
+	/**
+	 * Find the area in the render batch render target corresponding to each of the subsections of this component 
+	 * @param InComponent Which component to compute the subsection rects for 
+	 * @param OutSubsectionRects List of (up to 4) subsection rects when *not* taking into account duplicate borders (inclusive bounds)
+	 * @param OutSubsectionRectsWithDuplicateBorders List of (up to 4) subsection rects when taking into account duplicate borders (inclusive bounds)
+	 * @param bInWithDuplicateBorders indicates whether the rect coordinates should include the duplicated column/row at the end of the subsection or not
+	 * 
+	 * @return number of subsections
+	 */
+	int32 ComputeSubsectionRects(ULandscapeComponent* InComponent, TArray<FIntRect, TInlineAllocator<4>>& OutSubsectionRects, TArray<FIntRect, TInlineAllocator<4>>& OutSubsectionRectsWithDuplicateBorders) const;
+
+	/**
+	 * Find the area in the render batch render target corresponding to this component 
+	 * @param InComponent Which component to compute the section rect for 
+	 * @param bInWithDuplicateBorders indicates whether the rect coordinates should include the duplicated columns/rows at the end of each subsection or not
+	 * 
+	 * @Return the section rect (inclusive bounds)
+	 */
+	FIntRect ComputeSectionRect(ULandscapeComponent* InComponent, bool bInWithDuplicateBorders) const;
+
+	/**
+	 * Compute the rects corresponding to the sub-sections that need to be read from and written to when expanding the render target (inclusive bounds)
+	 * @param OutSubsectionRects List of all subsection rects when *not* taking into account duplicate borders (inclusive bounds)
+	 * @param OutSubsectionRectsWithDuplicateBorders List of all subsection rects when taking into account duplicate borders (inclusive bounds)
+	 */
+	void ComputeAllSubsectionRects(TArray<FIntRect>& OutSubsectionRects, TArray<FIntRect>& OutSubsectionRectsWithDuplicateBorders) const;
+};
+
+
+// ----------------------------------------------------------------------------------
+
+/**
+ * Utility struct for attaching some information that pertains to a given landscape component in the context of a batch render
+ */
+struct FComponentMergeRenderInfo
+{
+	/** Component to render */
+	ULandscapeComponent* Component = nullptr;
+
+	/** Texture region that corresponds to this component in the render area's render target */
+	FIntRect ComponentRegionInRenderArea;
+
+	/** Index of the component in the render area's render target */
+	FIntPoint ComponentKeyInRenderArea = FIntPoint(ForceInitToZero);
+
+	bool operator < (const FComponentMergeRenderInfo& InOther) const;
+};
+
+
+// ----------------------------------------------------------------------------------
+
+/** 
+ * Utility class that contains everything necessary to perform the batched merge : scratch render targets, list of batches, etc.
+ */
+class FMergeRenderContext
+{
+public:
+	friend class ::ALandscape;
+
+	FMergeRenderContext(ALandscape* InLandscape, bool bInIsHeightmapMerge);
+	virtual ~FMergeRenderContext();
+	FMergeRenderContext(const FMergeRenderContext& Other) = default;
+	FMergeRenderContext(FMergeRenderContext&& Other) = default;
+	FMergeRenderContext& operator=(const FMergeRenderContext& Other) = default;
+	FMergeRenderContext& operator=(FMergeRenderContext&& Other) = default;
+
+	bool IsValid() const
+	{ 
+		return !RenderBatches.IsEmpty(); 
+	}
+
+	/**
+	 * Cycle between the 3 render targets used for blending:
+	 *  Write becomes Read -> Read becomes ReadPrevious -> ReadPrevious becomes Write
+	 *	The new Write RT will be transitioned to state == InWriteAccess 
+	 *  The new Read RT will be transitioned to state ERHIAccess::SRVMask
+	 *  The new ReadPrevious RT will stay in state ERHIAccess::SRVMask
+	 * 
+	 * @param InDesiredWriteAccess RHI state of the write RT after the operation, if specified
+	 */
+	LANDSCAPE_API void CycleBlendRenderTargets(ERHIAccess InDesiredWriteAccess = ERHIAccess::None);
+
+	LANDSCAPE_API ULandscapeScratchRenderTarget* GetBlendRenderTargetWrite() const;
+	LANDSCAPE_API ULandscapeScratchRenderTarget* GetBlendRenderTargetRead() const;
+	LANDSCAPE_API ULandscapeScratchRenderTarget* GetBlendRenderTargetReadPrevious() const;
+	// TODO [jonathan.bard] : pass a target layer index instead of a name here?
+	LANDSCAPE_API ULandscapeScratchRenderTarget* GetValidityRenderTarget(const FName& InTargetLayerName) const;
+	
+	struct FOnRenderBatchTargetGroupDoneParams
+	{
+		FOnRenderBatchTargetGroupDoneParams(FMergeRenderContext* InMergeRenderContext, const FMergeRenderBatch& InRenderBatch, const TBitArray<>& InRenderGroupBitIndices, const TArrayView<FComponentMergeRenderInfo>& InSortedComponentMergeRenderInfos)
+			: MergeRenderContext(InMergeRenderContext)
+			, RenderBatch(&InRenderBatch)
+			, RenderGroupBitIndices(InRenderGroupBitIndices)
+			, SortedComponentMergeRenderInfos(InSortedComponentMergeRenderInfos)
+		{}
+
+		/** Render context : this is still active in this step and can be used for doing additional renders in the blend render targets, etc. */
+		FMergeRenderContext* MergeRenderContext = nullptr;
+
+		/** Batch to that was just rendered for this render group */
+		const FMergeRenderBatch* RenderBatch = nullptr;
+
+		/** Render group that was just rendered for this batch. Each bit in that bit array corresponds to an entry in FMergeRenderContext's AllTargetLayerNames  */
+		const TBitArray<> RenderGroupBitIndices;
+
+		/** Additional info about the components that have been processed in this batch render */
+		const TArray<FComponentMergeRenderInfo> SortedComponentMergeRenderInfos;
+	};
+	void Render(TFunction<void(const FOnRenderBatchTargetGroupDoneParams&)> OnRenderBatchTargetGroupDone);
+
+	inline ALandscape* GetLandscape() const { return Landscape; }
+	inline bool IsHeightmapMerge() const { return bIsHeightmapMerge; }
+
+	inline FIntPoint GetMaxNeededResolution() const { return MaxNeededResolution; }
+	inline const TArray<FMergeRenderBatch>& GetRenderBatches() const { return RenderBatches; }
+	LANDSCAPE_API const FMergeRenderBatch* GetCurrentRenderBatch() const;
+
+	LANDSCAPE_API FTransform ComputeVisualLogTransform(const FTransform& InTransform) const;
+	LANDSCAPE_API void IncrementVisualLogOffset();
+	LANDSCAPE_API void ResetVisualLogOffset();
+
+#if ENABLE_VISUAL_LOG
+	LANDSCAPE_API static int32 GetVisualLogAlpha();
+	LANDSCAPE_API bool IsVisualLogEnabled() const;
+#endif // ENABLE_VISUAL_LOG
+
+	/** Render the stencil render targets for each target layer in this merge for this batch */
+	void RenderValidityRenderTargets(const FMergeRenderBatch& InRenderBatch);
+
+	/** Duplicates the vertex data from the (sub-)sections of the batch, assuming GetBlendRenderTargetRead() is the RT that is read from and GetBlendRenderTargetWrite() the one that is written to */
+	void RenderExpandedRenderTarget(const FMergeRenderBatch& InRenderBatch);
+
+	const TArray<FName>& GetAllTargetLayerNames() const { return AllTargetLayerNames; }
+	inline int32 GetTargetLayerIndexForName(const FName& InName) const { return AllTargetLayerNames.Find(InName); }
+	inline int32 GetTargetLayerIndexForNameChecked(const FName& InName) const { int32 Index = AllTargetLayerNames.Find(InName); check(Index); return Index; }
+	inline FName GetTargetLayerNameForIndex(int32 InIndex) const { return AllTargetLayerNames.IsValidIndex(InIndex) ? AllTargetLayerNames[InIndex] : NAME_None; }
+	inline FName GetTargetLayerNameForIndexChecked(int32 InIndex) const { check(AllTargetLayerNames.IsValidIndex(InIndex)); return AllTargetLayerNames[InIndex]; }
+	TBitArray<> ConvertTargetLayerNamesToBitIndices(TConstArrayView<FName> InTargetLayerNames) const;
+	TBitArray<> ConvertTargetLayerNamesToBitIndicesChecked(TConstArrayView<FName> InTargetLayerNames) const;
+	TArray<FName> ConvertTargetLayerBitIndicesToNames(const TBitArray<>& InTargetLayerBitIndices) const;
+	TArray<FName> ConvertTargetLayerBitIndicesToNamesChecked(const TBitArray<>& InTargetLayerBitIndices) const;
+	TArray<ULandscapeLayerInfoObject*> ConvertTargetLayerBitIndicesToLayerInfos(const TBitArray<>& InTargetLayerBitIndices) const;
+	TArray<ULandscapeLayerInfoObject*> ConvertTargetLayerBitIndicesToLayerInfosChecked(const TBitArray<>& InTargetLayerBitIndices) const;
+
+	/**
+	 * Runs the given function for each all valid target layer in the bit indices in parameters, with the possibility of early exit
+	 * Most easily used with a lambda as follows:
+	 * ForEachTargetLayer([](int32 InTargetLayerIndex, FName InTargetLayerName) -> bool
+	 * {
+	 *     return continueLoop ? true : false;
+	 * });
+	 */
+	LANDSCAPE_API void ForEachTargetLayer(const TBitArray<>& InTargetLayerBitIndices, TFunctionRef<bool(int32 /*InTargetLayerIndex*/, FName /*InTargetLayerName*/)> Fn) const;
+	LANDSCAPE_API void ForEachTargetLayerChecked(const TBitArray<>& InTargetLayerBitIndices, TFunctionRef<bool(int32 /*InTargetLayerIndex*/, FName /*InTargetLayerName*/)> Fn) const;
+
+	const TBitArray<>& GetFinalTargetLayerBitIndices() const { return FinalTargetLayerBitIndices; }
+
+private:
+	/** Allocate all needed render targets for this merge */
+	void AllocateResources();
+	/** Free all render targets used in this merge */
+	void FreeResources();
+	/** Allocate all needed render targets for this batch */
+	void AllocateBatchResources(const FMergeRenderBatch& InRenderBatch);
+	/** Free all render targets used in this batch */
+	void FreeBatchResources(const FMergeRenderBatch& InRenderBatch);
+
+private:
+	// Blending is pretty much what we only do during the merge. It requires 3 render targets : 1 that we write to and therefore use as RTV (Write) and 2 that we read from and therefore use 
+	//  as SRV (one that contains the layer to merge, the other the accumulated result so far) : Previous(SRV) + Current(SRV) --> Write(RTV)
+	static constexpr int32 NumBlendRenderTargets = 3; 
+
+	/** Render targets that are used throughout the blending operations (they could be texture arrays in the case of multiple weightmaps) */
+	TStaticArray<ULandscapeScratchRenderTarget*, NumBlendRenderTargets> BlendRenderTargets;
+	int32 CurrentBlendRenderTargetWriteIndex = -1; 
+		
+	/** The landscape on which the merge occurs*/
+	ALandscape* Landscape = nullptr;
+
+	/** Type of merge being requested */
+	bool bIsHeightmapMerge = false;
+
+	/** List of all target layer names that can possibly be rendered on this landscape (even invalid ones). The index of the target layer in that list is important : 
+	  it allows to use a TBitArray instead of a TSet to designate a list of target layers all throughout the merge, which greatly accelerates all the intersection tests being made.
+	  If merging heightmaps, it contains a single name (that is only meaningful for debug display). */
+	TArray<FName> AllTargetLayerNames;
+
+	/** Per-target layer info. Same size as AllTargetLayerNames (only meaningful when rendering weightmaps) */
+	TArray<ULandscapeLayerInfoObject*> WeightmapLayerInfos;
+
+	/** List of valid target layers. If a target layer name is present here, it's because it has a valid landscape layer info object.Each bit in that bit array corresponds to an entry in AllTargetLayerNames */
+	TBitArray<> ValidTargetLayerBitIndices;
+
+	/** Final list of target layer names being involved in this merge context. If a target layer name is present here, it's because it's a valid target layer and it needs to be rendered
+	  because it has been requested or one of the target layers that have been requested needs it to be present (e.g. weight-blending). Each bit in that bit array corresponds to an entry in AllTargetLayerNames */
+	TBitArray<> FinalTargetLayerBitIndices;
+
+	/** Render targets storing the validity of each pixel wrt the target layer (i.e. stencil-like buffer, but stored as a RT to let users access it as a standard texture) : 
+	 Useful when sampling neighbors to know whether the data there corresponds to a valid neighbor.*/
+	TMap<FName, ULandscapeScratchRenderTarget*> PerTargetLayerValidityRenderTargets;
+
+	/** Maximum resolution needed by a given batch in this context (means we won't ever need more than this size of a render target during the whole merge */
+	FIntPoint MaxNeededResolution = FIntPoint(ForceInit);
+
+	/** Maximum number of slices needed by a given batch / render group in this context (means we won't ever need more than this number of slices in a given render target texture array during the whole merge */
+	int32 MaxNeededNumSlices = 0;
+
+	/** Successive batches of components being processed by this context. Each batch should be self-contained so that we won't ever need to keep more than one in memory (VRAM)
+	  in order to compute the info we need (heightmap/weightmaps), i.e. nothing should ever depend on 2 different batches. */
+	TArray<FMergeRenderBatch> RenderBatches;
+
+	/** Current batch being rendered */
+	int32 CurrentRenderBatchIndex = INDEX_NONE;
+
+	/** Offset for visual debugging */
+	FVector CurrentVisualLogOffset = FVector(ForceInitToZero);
+
+	/** Maximum height of all components to render in local space. */
+	double MaxLocalHeight = DBL_MIN;
+
+	/** List of components involved in this merge and the target layers they're writing to (each bit corresponds to the target layer name in AllTargetLayerNames) */
+	TMap<ULandscapeComponent*, TBitArray<>> ComponentToTargetLayerBitIndices;
+
+	/** Reverse lookup of ComponentToTargetLayerBitIndices : one entry per element in ComponentToTargetLayerBitIndices, each entry containing all of the components involved in this merge for this target layer */
+	TArray<TSet<ULandscapeComponent*>> TargetLayersToComponents;
+};
+
+
+// ----------------------------------------------------------------------------------
+
+/** A simple world space Object-Oriented Bounding Box */
+// TODO [jonathan.bard] : use FOrientedBox2d instead?
+struct FOOBox2D
+{
+	FOOBox2D() = default;
+	FOOBox2D(const FTransform& InTransform, const FVector2D& InExtents)
+		: Transform(InTransform)
+		, Extents(InExtents)
+	{}
+
+	FTransform Transform;
+	FVector2D Extents = FVector2D(ForceInit);
+};
+
+
+// ----------------------------------------------------------------------------------
+
+/** 
+ * Describes the input area needed for a given edit layer renderer's render item: this allows to infer the dependency between each component being rendered and the components it depends on
+ */
+class FInputWorldArea
+{
+public:
+	enum class EType
+	{
+		LocalComponent, // Designates any landscape component (i.e. the input area corresponds to the component being requested), with an optional number of neighboring components around it
+		SpecificComponent, // Designates a specific landscape component (based on its component key), with an optional number of neighboring components around it
+		OOBox, // Designates a fixed world area (an object-oriented box)
+		Infinite, // Designates the entire loaded landscape area
+	};
+
+	static FInputWorldArea CreateInfinite() { return FInputWorldArea(EType::Infinite); }
+	static FInputWorldArea CreateLocalComponent(const FIntRect& InLocalArea = FIntRect()) { return FInputWorldArea(EType::LocalComponent, FIntPoint(ForceInit), InLocalArea); }
+	static FInputWorldArea CreateSpecificComponent(const FIntPoint& InComponentKey, const FIntRect& InLocalArea = FIntRect()) { return FInputWorldArea(EType::SpecificComponent, InComponentKey, InLocalArea); }
+	static FInputWorldArea CreateOOBox(const FOOBox2D& InOOBox) { return FInputWorldArea(EType::OOBox, FIntPoint(ForceInit), FIntRect(), InOOBox); }
+
+	EType GetType() const { return Type; }
+	/** In the EType::LocalComponent case, returns the component's coordinates and the local area around it (inclusive bounds) */
+	FIntRect GetLocalComponentKeys(const FIntPoint& InComponentKey) const;
+	/** In the EType::SpecificComponent case, returns the component's coordinates and the local area around it (inclusive bounds) */
+	FIntRect GetSpecificComponentKeys() const;
+	/** In the EType::OOBox case, returns the OOBox */
+	const FOOBox2D& GetOOBox() const { check(Type == EType::OOBox); return OOBox2D; }
+
+	FBox ComputeWorldAreaAABB(const FTransform& InLandscapeTransform, const FBox& InLandscapeLocalBounds, const FTransform& InComponentTransform, const FBox& InComponentLocalBounds) const;
+	FOOBox2D ComputeWorldAreaOOBB(const FTransform& InLandscapeTransform, const FBox& InLandscapeLocalBounds, const FTransform& InComponentTransform, const FBox& InComponentLocalBounds) const;
+
+private: 
+	FInputWorldArea(EType InType, const FIntPoint& InComponentKey = FIntPoint(ForceInit), const FIntRect& InLocalArea = FIntRect(), const FOOBox2D& InOOBox2D = FOOBox2D())
+		: Type(InType)
+		, SpecificComponentKey(InComponentKey)
+		, LocalArea(InLocalArea)
+		, OOBox2D(InOOBox2D)
+	{}
+
+private:
+	EType Type = EType::LocalComponent;
+	/** Coordinates of the component (see ULandscapeComponent::GetComponentKey()) in the EType::SpecificComponent case */
+	FIntPoint SpecificComponentKey = FIntPoint(ForceInit);
+	/** Area around the component that is needed in the EType::LocalComponent / EType::SpecificComponent case (in component coordinates (see ULandscapeComponent::GetComponentKey()), 
+	 e.g. use (-1, -1, 1, 1) for the component and its immediate neighbors all around) */
+	FIntRect LocalArea; 
+	/** World space object-oriented box in the EType::OOBBox case */
+	FOOBox2D OOBox2D;
+};
+
+
+// ----------------------------------------------------------------------------------
+
+/**
+ * Describes the output area needed where a given edit layer renderer's render item writes: this allows to define the components of landscape that need to be processed and allows to divide
+ *  the work into batches
+ */
+class FOutputWorldArea
+{
+public:
+	enum class EType
+	{
+		LocalComponent, // Designates any landscape component (i.e. the input area corresponds to the component being requested)
+		SpecificComponent, // Designates a specific landscape component (based on its component key)
+		OOBox, // Designates a fixed world area (an Object-oriented box)
+	};
+
+	static FOutputWorldArea CreateLocalComponent() { return FOutputWorldArea(EType::LocalComponent); }
+	static FOutputWorldArea CreateSpecificComponent(const FIntPoint& InComponentKey) { return FOutputWorldArea(EType::SpecificComponent, InComponentKey); }
+	static FOutputWorldArea CreateOOBox(const FOOBox2D& InOOBox) { return FOutputWorldArea(EType::OOBox, FIntPoint(ForceInit), InOOBox); }
+
+	EType GetType() const { return Type; }
+	/** In the EType::SpecificComponent case, returns the component's coordinates */
+	const FIntPoint& GetSpecificComponentKey() const { check(Type == EType::SpecificComponent); return SpecificComponentKey; }
+	/** In the EType::OOBox case, returns the OOBox */
+	const FOOBox2D& GetOOBox() const { check(Type == EType::OOBox); return OOBox2D; }
+
+	FBox ComputeWorldAreaAABB(const FTransform& InComponentTransform, const FBox& InComponentLocalBounds) const;
+	FOOBox2D ComputeWorldAreaOOBB(const FTransform& InComponentTransform, const FBox& InComponentLocalBounds) const;
+
+private:
+	FOutputWorldArea(EType InType, const FIntPoint& InComponentKey = FIntPoint(ForceInit), const FOOBox2D& InOOBox = FOOBox2D())
+		: Type(InType)
+		, SpecificComponentKey(InComponentKey)
+		, OOBox2D(InOOBox)
+	{}
+
+private:
+	EType Type = EType::LocalComponent;
+	/** Coordinates of the component (see ULandscapeComponent::GetComponentKey()) in the EType::SpecificComponent case */
+	FIntPoint SpecificComponentKey = FIntPoint(ForceInit);
+	/** World space Object-oriented box in the EType::OOBBox case */
+	FOOBox2D OOBox2D;
+};
+
+
+// ----------------------------------------------------------------------------------
+
+/** 
+ * Each edit layer render item represents the capabilities of what a given edit layer can render in terms of landscape data : a renderer can provide one or many render items, which contain
+ *  the "locality" (what area do I affect?) as well as the "capability" (what target tool type do I affect? what weightmap(s)?) information related to what this layer item can do. 
+ *  See ILandscapeEditLayerRenderer::GetRenderItems
+ */
+class FEditLayerRenderItem
+{
+public:
+	FEditLayerRenderItem() = delete;
+	FEditLayerRenderItem(const FEditLayerTargetTypeState& InTargetTypeState, const FInputWorldArea& InInputWorldArea, const FOutputWorldArea& InOutputWorldArea, bool bInModifyExistingWeightmapsOnly)
+		: TargetTypeState(InTargetTypeState)
+		, InputWorldArea(InInputWorldArea)
+		, OutputWorldArea(InOutputWorldArea)
+		, bModifyExistingWeightmapsOnly(bInModifyExistingWeightmapsOnly)
+	{}
+
+	const FEditLayerTargetTypeState& GetTargetTypeState() const { return TargetTypeState; }
+
+	const FInputWorldArea& GetInputWorldArea() const { return InputWorldArea; }
+	void SetInputWorldArea(const FInputWorldArea& InInputWorldArea) { InputWorldArea = InInputWorldArea; }
+
+	const FOutputWorldArea& GetOutputWorldArea() const { return OutputWorldArea; }
+	void SetOutputWorldArea(const FOutputWorldArea& InOutputWorldArea) { OutputWorldArea = InOutputWorldArea; }
+
+	bool GetModifyExistingWeightmapsOnly() const { return bModifyExistingWeightmapsOnly; }
+
+private:
+	
+	/** Target types / weightmaps that this render item writes to */
+	FEditLayerTargetTypeState TargetTypeState;
+
+	/**
+	 * Area that this render item needs in order to render properly. 
+	 *  - If Infinite, it is assumed the render item needs the entire loaded landscape to render properly (i.e. it's dependent on all loaded landscape components)
+	 *  - If Local, it requires a particular component and optionally its immediate neighbors
+	 *  - If OOBox, then only the landscape components covered by this area will be considered as inputs
+	 */ 
+	FInputWorldArea InputWorldArea;
+
+	/**
+	 * Area that this render item writes to.
+	 *  - If Infinite, the render item writes everywhere (use only if necessary, as it will
+	 *  - If Local, it requires a particular component and optionally its immediate neighbors
+	 *  - If OOBox, then only the landscape components covered by this area will be considered as inputs
+	 */
+	FOutputWorldArea OutputWorldArea;
+
+	/** Indicates whether this render item actually outputs weightmaps (if false) or only modifies existing ones underneath (i.e. blending-only)*/
+	bool bModifyExistingWeightmapsOnly = false;
+};
+
+
+// ----------------------------------------------------------------------------------
+
+/** 
+ * Interface to implement to be able to provide an ordered list of renderers to the landscape. 
+ */
+class IEditLayerRendererProvider
+{
+public:
+	virtual ~IEditLayerRendererProvider() {}
+
+	/**
+	 * GetEditLayerRendererStates returns a list of renderers that this provider can provide, along with their current state. The state can then be manipulated at will by the user to enable/disable renderers selectively
+	 *
+	 * @param InLandscapeInfo ULandscapeInfo containing information about the landscape affected by this renderer
+	 *
+	 * @return a list of renderer states (i.e. a ILandscapeEditLayerRenderer and its current state) to be processed in that order by the merge operation
+	 */
+	LANDSCAPE_API virtual TArray<FEditLayerRendererState> GetEditLayerRendererStates(const ULandscapeInfo* InLandscapeInfo, bool bInSkipBrush) PURE_VIRTUAL(IEditLayerRendererProvider::GetEditLayerRendererStates, return {}; );
+};
+#endif // WITH_EDITOR
+
+} //namespace UE::Landscape::EditLayers
+
+// ----------------------------------------------------------------------------------
+
+/** 
+ * UInterface for a landscape edit layer renderer 
+ */
+UINTERFACE()
+class ULandscapeEditLayerRenderer :
+	public UInterface
+{
+	GENERATED_BODY()
+};
+
+/** 
+ * Interface that needs to be implemented for anything that can render heightmap/weightmap/visibility when merging landscape edit layers. 
+ *  Ideally it would have been defined in the UE::Landscape::EditLayers namespace but UHT prevents us from doing so.
+ *  The renderers are provided to the landscape by a IEditLayerRendererProvider.
+ */
+class ILandscapeEditLayerRenderer
+{
+	GENERATED_BODY()
+
+#if WITH_EDITOR
+public:
+	struct FRenderParams
+	{
+		FRenderParams(UE::Landscape::EditLayers::FMergeRenderContext* InMergeRenderContext, 
+			const TBitArray<>& InRenderGroupBitIndices,
+			const UE::Landscape::EditLayers::FEditLayerRendererState& InRendererState, 
+			const TArrayView<UE::Landscape::EditLayers::FComponentMergeRenderInfo>& InSortedComponentMergeRenderInfos, 
+			const FTransform& InRenderAreaWorldTransform, const FIntRect& InRenderAreaSectionRect)
+			: MergeRenderContext(InMergeRenderContext)
+			, RenderGroupBitIndices(InRenderGroupBitIndices)
+			, RendererState(InRendererState)
+			, SortedComponentMergeRenderInfos(InSortedComponentMergeRenderInfos)
+			, RenderAreaWorldTransform(InRenderAreaWorldTransform)
+			, RenderAreaSectionRect(InRenderAreaSectionRect)
+		{}
+
+		/** Merge context */
+		UE::Landscape::EditLayers::FMergeRenderContext* MergeRenderContext = nullptr;
+
+		/** List of target layers being involved in this step. Each bit in that bit array corresponds to an entry in FMergeRenderContext's AllTargetLayerNames  */
+		TBitArray<> RenderGroupBitIndices;
+
+		/** Full state for the renderer involved in this step. This allows to retrieve the exact state of this renderer (e.g. enabled weightmaps, which can be different than the render group, in that 
+		 target layers A, B and C might belong to the same group but this renderer actually only has A enabled). This is therefore the renderer's responsibility to check that a given target layer from the 
+		 render group is effectively enabled. */
+		UE::Landscape::EditLayers::FEditLayerRendererState RendererState;
+
+		/** List of components (with additional info) to render */
+		TArray<UE::Landscape::EditLayers::FComponentMergeRenderInfo> SortedComponentMergeRenderInfos;
+
+		// TODO [jonathan.bard] Verify that scale is correct
+		/** World transform that corresponds to the origin (bottom left corner) of the render area. The scale corresponds to the size of each quad in the landscape. */
+		FTransform RenderAreaWorldTransform;
+
+		/** SectionRect (i.e. landscape vertex coordinates, in landscape space) that corresponds to this render area */
+		FIntRect RenderAreaSectionRect;
+	};
+
+	/**
+	 * GetRendererStateInfo retrieves the current state of this renderer (what it can and does render, as well as how to group target layers together), and part of this will then be mutable for the duration of the merge. 
+	 *  The idea is that FEditLayerRendererState's SupportedTargetTypeState tells the capabilities of this renderer, while EnabledTargetTypeState tells what it currently does render.
+	 *  A target type must be both supported and enabled in order to have this renderer affect it and the "enabled" state can be changed at will by the user (e.g. to temporarily disable 
+	 *  a given edit layer just for the duration of the merge) : see FEditLayerRendererState
+	 * 
+	 * @param InLandscapeInfo ULandscapeInfo containing information about the landscape affected by this renderer
+	 * 
+	 * @param OutSupportedTargetTypeState List of all target types / weightmaps that this renderer supports (e.g. if, say ELandscapeToolTargetType::Weightmap, is not included, the renderer
+	 *  will *not* be used at all when rendering any kind of weightmap)
+	 * 
+	 * @param OutEnabledTargetTypeMask List of all target types / weightmaps that this renderer is currently enabled for (i.e. default state of this renderer wrt this target type).
+	 *  A target type must be both supported and enabled in order to have this renderer affect it.
+	 *  The "enabled" state can be changed at will by the user (e.g. to temporarily disable a given edit layer) : see FEditLayerRendererState
+	 * 
+	 * @param OutRenderGroups List of groups of target layers that this renderer requires to be rendered together.
+	 *  This allows it to perform horizontal blending (i.e.adjust the weights of the targeted weightmaps wrt one another).
+	 *  Depending on the other renderer's needs, the final render groups might contain more layers than was requested by a given renderer. This only means that more layers will be processed together
+	 *  and if this renderer doesn't act on one of these layers, it will simply do nothing with it in its RenderLayer function
+	 */
+	LANDSCAPE_API virtual void GetRendererStateInfo(const ULandscapeInfo* InLandscapeInfo,
+		UE::Landscape::EditLayers::FEditLayerTargetTypeState& OutSupportedTargetTypeState, UE::Landscape::EditLayers::FEditLayerTargetTypeState& OutEnabledTargetTypeState, TArray<TSet<FName>>& OutRenderGroups) const
+		PURE_VIRTUAL(ILandscapeEditLayerRenderer::GetRendererStateInfo, );
+
+	/**
+	* @return the a debug name for this renderer
+	*/
+	LANDSCAPE_API virtual FString GetEditLayerRendererDebugName() const 
+		PURE_VIRTUAL(ILandscapeEditLayerRenderer::GetEditLayerRendererDebugName, return TEXT(""); );
+
+	/**
+	 * GetRenderItems retrieves information about the areas this renderer renders to and specifically what respective input area they require to render properly
+	 *
+	 * @param InLandscapeInfo ULandscapeInfo containing information about the landscape affected by this renderer
+	 *
+	 * @return list of all render items that affect this renderer
+	*/
+	LANDSCAPE_API virtual TArray<UE::Landscape::EditLayers::FEditLayerRenderItem> GetRenderItems(const ULandscapeInfo* InLandscapeInfo) const 
+		PURE_VIRTUAL(ILandscapeEditLayerRenderer::GetRenderItems, return { }; );
+
+	/**
+	 * Indicates whether the renderer actually does anything in the render phase
+	*/
+	LANDSCAPE_API virtual bool CanRender() const { return true; }
+
+	/**
+	 * RenderLayer is where the renderer has a chance to render its content and eventually blend it with the merged result of all preceding layers
+	 * It operates on a limited set of components (depending on the size of the render batches) and on a set of target layers (e.g. multiple weightmaps).
+	 * It guarantees access to merged result from preceding layers of each target layer
+	*/
+	LANDSCAPE_API virtual void RenderLayer(ILandscapeEditLayerRenderer::FRenderParams& InRenderParams) {}
+
+#endif // WITH_EDITOR
+};
diff --git a/Engine/Source/Runtime/Landscape/Public/LandscapeEditTypes.h b/Engine/Source/Runtime/Landscape/Public/LandscapeEditTypes.h
--- a/Engine/Source/Runtime/Landscape/Public/LandscapeEditTypes.h
+++ b/Engine/Source/Runtime/Landscape/Public/LandscapeEditTypes.h
@@ -3,6 +3,7 @@
 #pragma once
 
 #include "CoreMinimal.h"
+#include "Misc/EnumRange.h"
 
 #include "LandscapeEditTypes.generated.h"
 
@@ -13,7 +14,20 @@ enum class ELandscapeToolTargetType : uint8
 	Weightmap = 1,
 	Visibility = 2,
 	Invalid = 3 UMETA(Hidden), // only valid for LandscapeEdMode->CurrentToolTarget.TargetType
+	Count = Invalid UMETA(Hidden), // Only the elements above Invalid actually count as proper target types
 };
+ENUM_RANGE_BY_COUNT(ELandscapeToolTargetType, ELandscapeToolTargetType::Count);
+
+enum class ELandscapeToolTargetTypeFlags : uint8
+{
+	None = 0,
+	Heightmap = (1 << static_cast<uint8>(ELandscapeToolTargetType::Heightmap)),
+	Weightmap = (1 << static_cast<uint8>(ELandscapeToolTargetType::Weightmap)),
+	Visibility = (1 << static_cast<uint8>(ELandscapeToolTargetType::Visibility)),
+	All = Heightmap | Weightmap | Visibility,
+};
+ENUM_CLASS_FLAGS(ELandscapeToolTargetTypeFlags);
+
 
 namespace UE::Landscape
 {
diff --git a/Engine/Source/Runtime/Landscape/Public/LandscapeSubsystem.h b/Engine/Source/Runtime/Landscape/Public/LandscapeSubsystem.h
--- a/Engine/Source/Runtime/Landscape/Public/LandscapeSubsystem.h
+++ b/Engine/Source/Runtime/Landscape/Public/LandscapeSubsystem.h
@@ -85,7 +85,7 @@ public:
 	 * @param InOptionalCameraLocations : (optional) camera locations that should be used when updating the grass. If not specified, the usual (streaming manager-based) view locations will be used
 	 */
 	LANDSCAPE_API void RegenerateGrass(bool bInFlushGrass, bool bInForceSync, TOptional<TArrayView<FVector>> InOptionalCameraLocations = TOptional<TArrayView<FVector>>());
-	
+
 	// Remove all grass instances from the specified components.  If passed null, removes all grass instances from all proxies.
 	void RemoveGrassInstances(const TSet<ULandscapeComponent*>* ComponentsToRemoveGrassInstances = nullptr);
 
@@ -171,7 +171,7 @@ private:
 	FLandscapeTextureStreamingManager* TextureStreamingManager = nullptr;
 	FLandscapeGrassMapsBuilder* GrassMapsBuilder = nullptr;
 
-#if WITH_EDITOR
+#if WITH_EDITORONLY_DATA
 	class FLandscapePhysicalMaterialBuilder* PhysicalMaterialBuilder = nullptr;
 	
 	FLandscapeNotificationManager* NotificationManager = nullptr;
@@ -186,7 +186,7 @@ private:
 	std::atomic<int32> NaniteBuildsInFlight;
 	std::atomic<int32> NaniteStaticMeshesInFlight;
 
-#endif // WITH_EDITOR
+#endif // WITH_EDITORONLY_DATA
 	
 	FDelegateHandle OnScalabilityChangedHandle;
 };
diff --git a/Engine/Source/Runtime/Landscape/Public/LandscapeUtils.h b/Engine/Source/Runtime/Landscape/Public/LandscapeUtils.h
--- a/Engine/Source/Runtime/Landscape/Public/LandscapeUtils.h
+++ b/Engine/Source/Runtime/Landscape/Public/LandscapeUtils.h
@@ -4,9 +4,6 @@
 
 #include "Containers/UnrealString.h"
 #include "CoreMinimal.h"
-#if UE_ENABLE_INCLUDE_ORDER_DEPRECATED_IN_5_2
-#include "RHIDefinitions.h"
-#endif
 #include "UObject/NameTypes.h"
 #include "AssetRegistry/AssetData.h"
 
@@ -16,6 +13,8 @@ class ULandscapeComponent;
 class ULandscapeLayerInfoObject;
 class UTexture2D;
 class UTexture;
+enum class ELandscapeToolTargetType : uint8;
+enum class ELandscapeToolTargetTypeFlags : uint8;
 
 namespace UE::Landscape
 {
@@ -26,10 +25,9 @@ namespace UE::Landscape
 */
 LANDSCAPE_API bool DoesPlatformSupportEditLayers(EShaderPlatform InShaderPlatform);
 
-int32 ComputeMaxDeltasOffsetForMip(int32 InMipIndex, int32 InNumRelevantMips);
-int32 ComputeMaxDeltasCountForMip(int32 InMipIndex, int32 InNumRelevantMips);
-int32 ComputeMipToMipMaxDeltasIndex(int32 InSourceMipIndex, int32 InDestinationMipIndex, int32 InNumRelevantMips);
-int32 ComputeMipToMipMaxDeltasCount(int32 InNumRelevantMips);
+LANDSCAPE_API ELandscapeToolTargetTypeFlags GetLandscapeToolTargetTypeAsFlags(ELandscapeToolTargetType InTargetType);
+LANDSCAPE_API ELandscapeToolTargetType GetLandscapeToolTargetTypeSingleFlagAsType(ELandscapeToolTargetTypeFlags InSingleFlag);
+LANDSCAPE_API FString GetLandscapeToolTargetTypeFlagsAsString(ELandscapeToolTargetTypeFlags InTargetTypeFlags);
 
 #if WITH_EDITOR
 
@@ -107,12 +105,6 @@ LANDSCAPE_API FString GetLayerInfoObjectPackageName(const ULevel* InLevel, const
 /** Returns true if the provided layer info object is the current visibility layer. */
 LANDSCAPE_API bool IsVisibilityLayer(const ULandscapeLayerInfoObject* InLayerInfoObject);
 
-/** Returns true if InPlatform is a mobile platform and the Landscape.MobileWeightTextureArray CVar set */
-bool UseWeightmapTextureArray(EShaderPlatform InPlatform);	
-
-/** Check if Landscape.MobileWeightTextureArray CVar set and we should attempt to use texture arrays for weight maps on mobile platforms  */
-bool IsMobileWeightmapTextureArrayEnabled();
-
 struct FLayerInfoFinder
 {
 	LANDSCAPE_API FLayerInfoFinder();
