Make landscape texture streaming manager consider all other worlds when turning off streaming requests, to fix an assert when running PIE.

Also fixes an issue where the streaming manager was not being deleted when the landscape subsystem is deinitialized, potentially leading to a memory leak.

diff --git a/Engine/Source/Runtime/Landscape/Private/LandscapeSubsystem.cpp b/Engine/Source/Runtime/Landscape/Private/LandscapeSubsystem.cpp
--- a/Engine/Source/Runtime/Landscape/Private/LandscapeSubsystem.cpp
+++ b/Engine/Source/Runtime/Landscape/Private/LandscapeSubsystem.cpp
@@ -258,6 +258,8 @@ void ULandscapeSubsystem::Deinitialize()
 	delete GrassMapsBuilder;
 	delete PhysicalMaterialBuilder;
 	delete NotificationManager;
+	delete TextureStreamingManager;
+
 #endif
 	Proxies.Empty();
 	LandscapeActors.Empty();
diff --git a/Engine/Source/Runtime/Landscape/Private/LandscapeTextureStreamingManager.cpp b/Engine/Source/Runtime/Landscape/Private/LandscapeTextureStreamingManager.cpp
--- a/Engine/Source/Runtime/Landscape/Private/LandscapeTextureStreamingManager.cpp
+++ b/Engine/Source/Runtime/Landscape/Private/LandscapeTextureStreamingManager.cpp
@@ -3,6 +3,7 @@
 #include "LandscapeTextureStreamingManager.h"
 #include "Engine/Texture.h"
 #include "TextureCompiler.h"
+#include "LandscapePrivate.h"
 
 namespace UE::Landscape
 {
@@ -19,6 +20,8 @@ namespace UE::Landscape
 	}
 }
 
+TArray<FLandscapeTextureStreamingManager*> FLandscapeTextureStreamingManager::AllStreamingManagers;
+
 bool FLandscapeTextureStreamingManager::RequestTextureFullyStreamedIn(UTexture* Texture, bool bWaitForStreaming)
 {
 	TWeakObjectPtr<UTexture> TexturePtr = Texture;
@@ -79,11 +82,19 @@ void FLandscapeTextureStreamingManager::UnrequestTextureFullyStreamedIn(UTexture
 		if (State->RequestCount > 0)
 		{
 			State->RequestCount--;
-			if (!State->bForever && State->RequestCount <= 0)
+			if (!State->WantsTextureStreamedIn())
 			{
-				// allow stream out, remove tracking
-				Texture->bForceMiplevelsToBeResident = false;
+				// remove state tracking for this texture
 				TextureStates.Remove(TexturePtr);
+				if ((AllStreamingManagers.Num() == 1) || !AnyStreamingManagerWantsTextureStreamedIn(TexturePtr))
+				{
+					// allow stream out
+					Texture->bForceMiplevelsToBeResident = false;
+				}
+				else
+				{
+					UE::Landscape::EnsureTextureForcedResident(Texture);
+				}
 			}
 			else
 			{
@@ -92,9 +103,7 @@ void FLandscapeTextureStreamingManager::UnrequestTextureFullyStreamedIn(UTexture
 		}
 		else
 		{
-			// only way the request count should get to zero is if the texture is flagged as forever streamed.
-			ensure(State->bForever);
-			UE::Landscape::EnsureTextureForcedResident(Texture);
+			UE_LOG(LogLandscape, Warning, TEXT("Texture Streaming Manager received more Unrequests than Requests to stream texture %s"), *Texture->GetName());
 		}
 	}
 }
@@ -141,7 +150,7 @@ void FLandscapeTextureStreamingManager::CleanupPostGarbageCollect()
 		{
 			// reset the texture force resident after garbage collection (which clears it sometimes)
 			FTextureState& State = It.Value();
-			if (State.bForever || State.RequestCount > 0)
+			if (State.WantsTextureStreamedIn())
 			{
 				Texture->bForceMiplevelsToBeResident = true;
 			}
@@ -161,7 +170,7 @@ void FLandscapeTextureStreamingManager::CheckRequestedTextures()
 			if (UTexture* Texture = It.Key().Get())
 			{
 				FTextureState& State = It.Value();
-				if (State.bForever || State.RequestCount > 0)
+				if (State.WantsTextureStreamedIn())
 				{
 					if (!Texture->bForceMiplevelsToBeResident)
 					{
@@ -184,18 +193,53 @@ bool FLandscapeTextureStreamingManager::IsTextureFullyStreamedIn(UTexture* InTex
 		!InTexture->HasPendingInitOrStreaming() && InTexture->IsFullyStreamedIn();
 }
 
+bool FLandscapeTextureStreamingManager::AnyStreamingManagerWantsTextureStreamedIn(TWeakObjectPtr<UTexture> TexturePtr)
+{
+	for (FLandscapeTextureStreamingManager* Manager : AllStreamingManagers)
+	{
+		if (FTextureState* State = Manager->TextureStates.Find(TexturePtr))
+		{
+			if (State->WantsTextureStreamedIn())
+			{
+				return true;
+			}
+		}
+	}
+	return false;
+}
+
+FLandscapeTextureStreamingManager::FLandscapeTextureStreamingManager()
+{
+	AllStreamingManagers.Add(this);
+}
+
 FLandscapeTextureStreamingManager::~FLandscapeTextureStreamingManager()
 {
-	if (!ensure(TextureStates.IsEmpty()))
+	AllStreamingManagers.RemoveSwap(this, EAllowShrinking::No);
+
+	// there could be some textures still requested, if they were requested "forever".
+	// since the world is going away, we can re-evaluate whether they should remain streamed or not.
+	int32 RemainingRequests = 0;
+	for (auto It = TextureStates.CreateIterator(); It; ++It)
 	{
-		// clear force stream flag on all textures, just in case
-		for (auto It = TextureStates.CreateIterator(); It; ++It)
+		FTextureState& State = It.Value();
+		if (State.RequestCount > 0)
+		{
+			RemainingRequests++;
+		}
+
+		if (UTexture* Texture = It.Key().Get())
 		{
-			UTexture* Texture = It.Key().Get();
-			if (Texture)
+			if (!AnyStreamingManagerWantsTextureStreamedIn(It.Key()))
 			{
+				// none of the remaining streaming managers request this texture, we can disable the mip requests
 				Texture->bForceMiplevelsToBeResident = false;
 			}
 		}
 	}
+
+	if (RemainingRequests > 0)
+	{
+		UE_LOG(LogLandscape, Display, TEXT("At destruction, the Landscape Texture Streaming Manager still has streaming requests for %d Textures, this may indicate failure to clean them up."), RemainingRequests);
+	}
 }
diff --git a/Engine/Source/Runtime/Landscape/Private/LandscapeTextureStreamingManager.h b/Engine/Source/Runtime/Landscape/Private/LandscapeTextureStreamingManager.h
--- a/Engine/Source/Runtime/Landscape/Private/LandscapeTextureStreamingManager.h
+++ b/Engine/Source/Runtime/Landscape/Private/LandscapeTextureStreamingManager.h
@@ -65,6 +65,7 @@ public:
 
 	static bool IsTextureFullyStreamedIn(UTexture* Texture);
 
+	FLandscapeTextureStreamingManager();
 	~FLandscapeTextureStreamingManager();
 
 private:
@@ -72,10 +73,18 @@ private:
 	{
 		int32 RequestCount = 0;
 		bool bForever = false;
+
+		bool WantsTextureStreamedIn() const
+		{
+			return bForever || (RequestCount > 0);
+		}
 	};
 
 	TMap<TWeakObjectPtr<UTexture>, FTextureState, FDefaultSetAllocator, TWeakObjectPtrMapKeyFuncs<TWeakObjectPtr<UTexture>, FTextureState>> TextureStates;
 
+	static TArray<FLandscapeTextureStreamingManager*> AllStreamingManagers;
+	static bool AnyStreamingManagerWantsTextureStreamedIn(TWeakObjectPtr<UTexture> TexturePtr);
+
 #if WITH_EDITOR
 	FLandscapeTextureStreamingManagerUndoDetector UndoDetector;
 #endif // WITH_EDITOR
