Fixed landscape BP brushes not working with batched merge

diff --git a/Engine/Source/Runtime/Landscape/Private/LandscapeBlueprintBrushBase.cpp b/Engine/Source/Runtime/Landscape/Private/LandscapeBlueprintBrushBase.cpp
--- a/Engine/Source/Runtime/Landscape/Private/LandscapeBlueprintBrushBase.cpp
+++ b/Engine/Source/Runtime/Landscape/Private/LandscapeBlueprintBrushBase.cpp
@@ -314,7 +314,7 @@ void ALandscapeBlueprintBrushBase::RenderLayer(ILandscapeEditLayerRenderer::FRen
 
 		// Only render the target layer if it's effectively enabled for this merge : it's possible there are target layers in the render group that we don't support or are not enabled so we have to 
 		//  do the validation here first :
-		if (!bIsHeightmapMerge && EnabledWeightmaps.Contains(TargetLayerName))
+		if (bIsHeightmapMerge || EnabledWeightmaps.Contains(TargetLayerName))
 		{
 			// Execute (i.e. (Initialize/)Render the BP brush) : 
 			FLandscapeBrushParameters BrushParameters(bIsHeightmapMerge, InRenderParams.RenderAreaWorldTransform, InRenderParams.RenderAreaSectionRect.Size(), ReadRT2D, TargetLayerName);
