Unshelved from pending changelist '35029374':

Fixed inverted parameters in header

#rnx
#skipci

diff --git a/Engine/Source/Runtime/Landscape/Public/LandscapeUtils.h b/Engine/Source/Runtime/Landscape/Public/LandscapeUtils.h
--- a/Engine/Source/Runtime/Landscape/Public/LandscapeUtils.h
+++ b/Engine/Source/Runtime/Landscape/Public/LandscapeUtils.h
@@ -62,13 +62,13 @@ public:
 	/**
 	* Uses the provided arguments to add proper source/destination entries to internal copy requests.
 	* @param	InDestination	The texture used as a destination for the copy.
-	* @param	InDestinationChannel	The channel used as a destination for the copy.
 	* @param	InDestinationSlice	The Texture array slice to write to (use 0 for a Texture2D)
+	* @param	InDestinationChannel	The channel used as a destination for the copy.
 	* @param	InComponent		The component containing the wanted source weightmap.
 	* @param	InLayerInfo		The layer info used to retrieve the proper source weightmap and channel.
 	* @return True if the copy has been successfully added.
 	*/
-	LANDSCAPE_API bool AddWeightmapCopy(UTexture* InDestination, int8 InDestinationChannel, int8 InDestinationSlice, const ULandscapeComponent* InComponent, ULandscapeLayerInfoObject* InLayerInfo);
+	LANDSCAPE_API bool AddWeightmapCopy(UTexture* InDestination, int8 InDestinationSlice, int8 InDestinationChannel, const ULandscapeComponent* InComponent, ULandscapeLayerInfoObject* InLayerInfo);
 
 	/** Process pending internal copy requests. */
 	LANDSCAPE_API bool ProcessTextureCopies();
