Fix a crash in landscape when layerinfos are not loaded/available.

diff --git a/Engine/Source/Runtime/Landscape/Private/LandscapeEdit.cpp b/Engine/Source/Runtime/Landscape/Private/LandscapeEdit.cpp
--- a/Engine/Source/Runtime/Landscape/Private/LandscapeEdit.cpp
+++ b/Engine/Source/Runtime/Landscape/Private/LandscapeEdit.cpp
@@ -1158,7 +1158,7 @@ void ULandscapeComponent::FixupWeightmaps(const FGuid& InEditLayerGuid)
 					{
 						FFormatNamedArguments Arguments;
 						Arguments.Add(TEXT("LandscapeName"), FText::FromString(GetPathName()));
-						Arguments.Add(TEXT("TargetLayerName"), FText::FromString(Allocation.LayerInfo->LayerName.ToString()));
+						Arguments.Add(TEXT("TargetLayerName"), FText::FromName(Allocation.GetLayerName()));
 						Arguments.Add(TEXT("EditLayerName"), FText::FromString(InEditLayerGuid.IsValid() ? LayersData[InEditLayerGuid].DebugName.ToString() : FString(TEXT("<Runtime>"))));
 						FMessageLog("MapCheck").Info()
 							->AddToken(FTextToken::Create(FText::Format(LOCTEXT("MapCheck_Message_FixedUpDeletedLayerWeightmap", "{LandscapeName} : Removed unused layer weightmap Edit Layer: '{EditLayerName}' Target Layer: '{TargetLayerName}'"), Arguments)))
