Fix assert when deleting a newly created and saved WP landscape.

diff --git a/Engine/Source/Runtime/Landscape/Private/LandscapeEdit.cpp b/Engine/Source/Runtime/Landscape/Private/LandscapeEdit.cpp
--- a/Engine/Source/Runtime/Landscape/Private/LandscapeEdit.cpp
+++ b/Engine/Source/Runtime/Landscape/Private/LandscapeEdit.cpp
@@ -6524,16 +6524,33 @@ bool ULandscapeInfo::CanDeleteLandscape(FText& OutReason) const
 					ALandscapeProxy* LandscapeProxy = Cast<ALandscapeProxy>(ActorDescInstance->GetActor());
 					if (LandscapeProxy != LandscapeActor)
 					{
-						// If LandscapeProxy is null then it is not loaded so not deleted.
-						if (!LandscapeProxy)
+						const bool bIsLoaded = ActorDescInstance->IsLoaded();
+						const bool bIsAllocated = (LandscapeProxy != nullptr);
+						const bool bIsGarbage = bIsAllocated && !IsValidChecked(LandscapeProxy);
+						const bool bIsRegistered = bIsAllocated && StreamingProxies.Contains(CastChecked<ALandscapeStreamingProxy>(LandscapeProxy));
+						const bool bRegisteredComps = bIsAllocated && LandscapeProxy->HasActorRegisteredAllComponents();
+						check(bIsRegistered == bRegisteredComps);
+
+						if (bIsGarbage)
 						{
-							++UndeletedProxyCount;
+							// proxy has been deleted, WP just hasn't cleaned up the actor desc yet.
+							// so don't count it as an undeleted proxy
 						}
-						else
+						else if (bIsLoaded && !bIsRegistered)
+						{
+							// this is a bit of a messed up state that we can encounter after an initial save
+							// where the proxy is loaded but not registered with the world
+							// TODO [chris.tchou] : figure out how to fix this for real so we don't get into this state
+							UE_LOG(LogLandscape, Display, TEXT("The Landscape Proxy '%p' is loaded and valid, but not registered. This may cause incorrect display in the outliner view. You can reload the level to correct this."),
+								LandscapeProxy);
+
+							// since this proxy was not counted in the registered proxy loop above, count it here
+							UndeletedProxyCount++;
+						}
+						else if (!bIsLoaded)
 						{
-							// If Actor is loaded it should be Registered and not pending kill (already accounted for) or pending kill (deleted)
-							TWeakObjectPtr<ALandscapeStreamingProxy> StreamingProxyPtr = CastChecked<ALandscapeStreamingProxy>(LandscapeProxy);
-							check(StreamingProxies.Contains(StreamingProxyPtr) == IsValidChecked(LandscapeProxy));
+							// unloaded proxies still exist (just not in memory at the moment)
+							UndeletedProxyCount++;
 						}
 					}
 				}
