Prevents GetActionableMessage from ticking in PIE or simulating world.
Parallelize dirty landscape proxies check.

diff --git a/Engine/Source/Runtime/Landscape/Private/LandscapePhysicalMaterial.cpp b/Engine/Source/Runtime/Landscape/Private/LandscapePhysicalMaterial.cpp
--- a/Engine/Source/Runtime/Landscape/Private/LandscapePhysicalMaterial.cpp
+++ b/Engine/Source/Runtime/Landscape/Private/LandscapePhysicalMaterial.cpp
@@ -578,7 +578,7 @@ void FLandscapePhysicalMaterialRenderTask::Release()
 
 bool FLandscapePhysicalMaterialRenderTask::IsValid() const
 {
-	check(IsInGameThread());
+	check(IsInParallelGameThread() || IsInGameThread());
 	return PoolHandle != -1;
 }
 
diff --git a/Engine/Source/Runtime/Landscape/Private/LandscapeSubsystem.cpp b/Engine/Source/Runtime/Landscape/Private/LandscapeSubsystem.cpp
--- a/Engine/Source/Runtime/Landscape/Private/LandscapeSubsystem.cpp
+++ b/Engine/Source/Runtime/Landscape/Private/LandscapeSubsystem.cpp
@@ -41,6 +41,7 @@
 
 #if WITH_EDITOR
 #include "ActionableMessageSubsystem.h"
+#include "Async/ParallelFor.h"
 #include "FileHelpers.h"
 #include "Editor.h"
 #endif
@@ -588,13 +589,15 @@ void ULandscapeSubsystem::Tick(float DeltaTime)
 	NaniteMeshBuildEvents.RemoveAllSwap([](const FGraphEventRef& Ref) -> bool { return Ref->IsComplete(); });
 
 	LastTickFrameNumber = FrameNumber;
+
+	UActionableMessageSubsystem* ActionableMessageSubsystem = World->GetSubsystem<UActionableMessageSubsystem>();
 	
-	if (UActionableMessageSubsystem* ActionableMessageSubsystem = World->GetSubsystem<UActionableMessageSubsystem>())
+	if ((ActionableMessageSubsystem != nullptr) && GIsEditor)
 	{
 		FActionableMessage ActionableMessage;
 		const FName LandscapeMessageProvider = TEXT("Landscape");
 
-		if (GetActionableMessage(ActionableMessage))
+		if (!(GEditor->bIsSimulatingInEditor || (GEditor->PlayWorld != nullptr)) && GetActionableMessage(ActionableMessage))
 		{
 			ActionableMessageSubsystem->SetActionableMessage(LandscapeMessageProvider, ActionableMessage);
 		}
@@ -690,22 +693,24 @@ TArray<TTuple<ALandscapeProxy*, UE::Landscape::EOutdatedDataFlags>> ULandscapeSu
 	if (!World || World->IsGameWorld())
 	{
 		return {};
-}
+	}
 
 	TArray<TTuple<ALandscapeProxy*, UE::Landscape::EOutdatedDataFlags>> OutdatedProxies;
-	OutdatedProxies.Reserve(Proxies.Num());
-	for (TObjectPtr<ALandscapeProxy> Proxy : Proxies)
+	
+	ParallelForWithTaskContext(OutdatedProxies, Proxies.Num(), [this, InMatchingOutdatedDataFlags, bInMustMatchAllFlags](TTuple<ALandscapeProxy*, UE::Landscape::EOutdatedDataFlags>& Context, int32 Index)
 	{
-		ALandscapeProxy* ValidProxy = Proxy.Get();
-{
-			const UE::Landscape::EOutdatedDataFlags ProxyOutdatedDataFlags = ValidProxy->GetOutdatedDataFlags();
-			if ((bInMustMatchAllFlags && EnumHasAllFlags(ProxyOutdatedDataFlags, InMatchingOutdatedDataFlags))
-				|| (!bInMustMatchAllFlags && EnumHasAnyFlags(ProxyOutdatedDataFlags, InMatchingOutdatedDataFlags)))
-			{
-				OutdatedProxies.Add({ValidProxy, ProxyOutdatedDataFlags});
-			}
+		FTaskTagScope Scope(ETaskTag::EParallelGameThread);
+		ALandscapeProxy* ValidProxy = Proxies[Index].Get();
+		const UE::Landscape::EOutdatedDataFlags ProxyOutdatedDataFlags = ValidProxy->GetOutdatedDataFlags();
+
+		if ((bInMustMatchAllFlags && EnumHasAllFlags(ProxyOutdatedDataFlags, InMatchingOutdatedDataFlags))
+			|| (!bInMustMatchAllFlags && EnumHasAnyFlags(ProxyOutdatedDataFlags, InMatchingOutdatedDataFlags)))
+		{
+			Context = {ValidProxy, ProxyOutdatedDataFlags};
 		}
-	}
+	});
+
+	OutdatedProxies.RemoveAllSwap([](const TTuple<ALandscapeProxy*, UE::Landscape::EOutdatedDataFlags>& OutdatedProxy){ return OutdatedProxy.Key == nullptr; });
 	return OutdatedProxies;
 }
 
