Fixed crash when resizing landscape.
Also, deprecated CreateLayerFrom, since it's a duplicate from DuplicateLayerAndMoveBrushes, but missed the brush reassignment, and was unused anyway.

A call to OnLayerCreatedInternal was missing from DuplicateLayerAndMoveBrushes, which prevented the duplicated edit layer from having a proper pointer to its landscape

diff --git a/Engine/Source/Runtime/Landscape/Classes/Landscape.h b/Engine/Source/Runtime/Landscape/Classes/Landscape.h
--- a/Engine/Source/Runtime/Landscape/Classes/Landscape.h
+++ b/Engine/Source/Runtime/Landscape/Classes/Landscape.h
@@ -176,7 +176,7 @@ struct FLandscapeLayer
 	//~ End UE::Landscape::EditLayers::IEditLayerRendererProvider implementation
 #endif // WITH_EDITOR
 
-	// TODO [jonathan.bard] : make those private, FLandscapeLayer is not LANDSCAPE_API so it can be done without deprecation
+	// TODO [jonathan.bard] : make those private
 public:
 	UPROPERTY(meta = (IgnoreForMemberInitializationTest))
 	FGuid Guid;
@@ -381,6 +381,7 @@ public:
 	* @Param InLayer edit layer to copy from
 	* @return the index of the newly-created layer
 	*/
+	UE_DEPRECATED(5.5, "Use DuplicateLayerAndMoveBrushes")
 	LANDSCAPE_API int32 CreateLayerFrom(const FLandscapeLayer& InLayer);
 	LANDSCAPE_API void CreateDefaultLayer();
 
diff --git a/Engine/Source/Runtime/Landscape/Private/LandscapeConfigHelper.cpp b/Engine/Source/Runtime/Landscape/Private/LandscapeConfigHelper.cpp
--- a/Engine/Source/Runtime/Landscape/Private/LandscapeConfigHelper.cpp
+++ b/Engine/Source/Runtime/Landscape/Private/LandscapeConfigHelper.cpp
@@ -575,8 +575,8 @@ ULandscapeInfo* FLandscapeConfigHelper::ChangeConfiguration(ULandscapeInfo* InLa
 		NewLandscape->bCanHaveLayersContent = true;
 		for (const FLandscapeLayer& OldLayer : OldLandscape->GetLayers())
 		{
-			FLandscapeLayer* NewLayer = NewLandscape->DuplicateLayerAndMoveBrushes(OldLayer);
-			check((NewLayer != nullptr) && (NewLayer->EditLayer != nullptr));
+			const FLandscapeLayer* NewLayer = NewLandscape->DuplicateLayerAndMoveBrushes(OldLayer);
+			check((NewLayer == nullptr) || (NewLayer->EditLayer != nullptr)); // it's possible DuplicateLayerAndMoveBrushes fails (e.g. max number of layers reached), but if not, we should always have an EditLayer
 		}
 	}
 		
diff --git a/Engine/Source/Runtime/Landscape/Private/LandscapeEdit.cpp b/Engine/Source/Runtime/Landscape/Private/LandscapeEdit.cpp
--- a/Engine/Source/Runtime/Landscape/Private/LandscapeEdit.cpp
+++ b/Engine/Source/Runtime/Landscape/Private/LandscapeEdit.cpp
@@ -3748,13 +3748,16 @@ void ALandscapeProxy::Import(const FGuid& InGuid, int32 InMinX, int32 InMinY, in
 		{
 			for (const FLandscapeLayer& OldLayer : InImportLayers)
 			{
-				FLandscapeLayer* NewLayer = LandscapeActor->DuplicateLayerAndMoveBrushes(OldLayer);
-				check(NewLayer != nullptr);
+				const FLandscapeLayer* NewLayer = LandscapeActor->DuplicateLayerAndMoveBrushes(OldLayer);
+				if (NewLayer != nullptr)
+				{
+					check(NewLayer->EditLayer != nullptr); // it's possible DuplicateLayerAndMoveBrushes fails (e.g. max number of layers reached), but if not, we should always have an EditLayer
 
-				FLayerImportSettings ImportSettings;
-				ImportSettings.SourceLayerGuid = OldLayer.Guid;
-				ImportSettings.DestinationLayerGuid = NewLayer->Guid;
-				LayerImportSettings.Add(ImportSettings);
+					FLayerImportSettings ImportSettings;
+					ImportSettings.SourceLayerGuid = OldLayer.Guid;
+					ImportSettings.DestinationLayerGuid = NewLayer->Guid;
+					LayerImportSettings.Add(ImportSettings);
+				}
 			}
 
 			LandscapeInfo->GetComponentsInRegion(InMinX, InMinY, InMaxX, InMaxY, ComponentsToProcess);
diff --git a/Engine/Source/Runtime/Landscape/Private/LandscapeEditLayers.cpp b/Engine/Source/Runtime/Landscape/Private/LandscapeEditLayers.cpp
--- a/Engine/Source/Runtime/Landscape/Private/LandscapeEditLayers.cpp
+++ b/Engine/Source/Runtime/Landscape/Private/LandscapeEditLayers.cpp
@@ -12906,6 +12906,12 @@ FLandscapeLayer* ALandscape::DuplicateLayerAndMoveBrushes(const FLandscapeLayer&
 		return nullptr;
 	}
 
+	if (IsMaxLayersReached())
+	{
+		UE_LOG(LogLandscape, Warning, TEXT("Cannot duplicate layer : %s as the max number of layers (%i) has been reached"), *InOtherLayer.Name.ToString(), GetDefault<ULandscapeSettings>()->MaxNumberOfLayers);
+		return nullptr;
+	}
+
 	Modify();
 
 	FLandscapeLayer NewLayer(InOtherLayer);
@@ -12922,13 +12928,7 @@ FLandscapeLayer* ALandscape::DuplicateLayerAndMoveBrushes(const FLandscapeLayer&
 	}
 
 	int32 AddedIndex = LandscapeEditLayers.Add(NewLayer);
-
-	// Create associated layer data in each landscape proxy
-	LandscapeInfo->ForEachLandscapeProxy([&NewLayer](ALandscapeProxy* Proxy)
-	{
-		Proxy->AddLayer(NewLayer.Guid);
-		return true;
-	});
+	OnLayerCreatedInternal(NewLayer);
 
 	return &LandscapeEditLayers[AddedIndex];
 }
@@ -12951,6 +12951,7 @@ int32 ALandscape::CreateLayer(FName InName, const TSubclassOf<ULandscapeEditLaye
 	return LayerIndex;
 }
 
+// Deprecated : use DuplicateLayerAndMoveBrushes 
 int32 ALandscape::CreateLayerFrom(const FLandscapeLayer& InLayer)
 {
 	if (IsMaxLayersReached() || !CanHaveLayersContent())
