Removed internal details about how target layers are handled from client-side structs ILandscapeEditLayerRenderer::FRenderParams and FMergeRenderContext::FOnRenderBatchTargetGroupDoneParams (replaced by 2 synchronized TArrays of target layer names and infos)

diff --git a/Engine/Source/Runtime/Landscape/Private/LandscapeBlueprintBrushBase.cpp b/Engine/Source/Runtime/Landscape/Private/LandscapeBlueprintBrushBase.cpp
--- a/Engine/Source/Runtime/Landscape/Private/LandscapeBlueprintBrushBase.cpp
+++ b/Engine/Source/Runtime/Landscape/Private/LandscapeBlueprintBrushBase.cpp
@@ -286,11 +286,10 @@ void ALandscapeBlueprintBrushBase::RenderLayer(ILandscapeEditLayerRenderer::FRen
 		CurrentLayerReadRT->TransitionTo(ERHIAccess::CopySrc);
 	}
 
-	TArray<FName> TargetLayerNames = InRenderParams.MergeRenderContext->ConvertTargetLayerBitIndicesToNamesChecked(InRenderParams.RenderGroupBitIndices);
-	const int32 NumTargetLayers = TargetLayerNames.Num();
+	const int32 NumTargetLayers = InRenderParams.RenderGroupTargetLayerNames.Num();
 	for (int32 TargetLayerIndex = 0; TargetLayerIndex < NumTargetLayers; ++TargetLayerIndex)
 	{
-		const FName TargetLayerName = TargetLayerNames[TargetLayerIndex];
+		const FName TargetLayerName = InRenderParams.RenderGroupTargetLayerNames[TargetLayerIndex];
 		SCOPED_DRAW_EVENTF_GAMETHREAD(LandscapeLayers, TEXT("Render %s"), TargetLayerName);
 
 		// If necessary, copy from the texture array's slice to the scratch render target 2D : 
diff --git a/Engine/Source/Runtime/Landscape/Private/LandscapeEditLayerRenderer.cpp b/Engine/Source/Runtime/Landscape/Private/LandscapeEditLayerRenderer.cpp
--- a/Engine/Source/Runtime/Landscape/Private/LandscapeEditLayerRenderer.cpp
+++ b/Engine/Source/Runtime/Landscape/Private/LandscapeEditLayerRenderer.cpp
@@ -774,10 +774,8 @@ void FMergeRenderContext::Render(TFunction<void(const FOnRenderBatchTargetGroupD
 
 	const FTransform& LandscapeTransform = Landscape->GetTransform();
 	// For visual logging, start at the top of the landscape's bounding box :
-	FTransform LandscapeWorldTransformForVisLog = LandscapeTransform;
 	FVector LandscapeTopPosition(0.0, 0.0, MaxLocalHeight);
-	LandscapeWorldTransformForVisLog.SetTranslation(LandscapeWorldTransformForVisLog.GetLocation() + LandscapeTransform.TransformVector(LandscapeTopPosition));
-	// TODO [jonathan.bard] : this should simply be: FTransform LandscapeWorldTransformForVisLog = FTransform(LandscapeTopPosition).ToMatrixNoScale() * LandscapeTransform.ToMatrixWithScale()
+	FTransform LandscapeWorldTransformForVisLog = FTransform(LandscapeTopPosition) * LandscapeTransform;
 
 	AllocateResources();
 
@@ -816,7 +814,8 @@ void FMergeRenderContext::Render(TFunction<void(const FOnRenderBatchTargetGroupD
 		for (int32 RenderStepIndex = 0; RenderStepIndex < NumRenderSteps; ++RenderStepIndex)
 		{
 			const FMergeRenderStep& RenderStep = RenderBatch.RenderSteps[RenderStepIndex];
-			const TArray<FName> TargetLayerNames = ConvertTargetLayerBitIndicesToNamesChecked(RenderStep.RenderGroupBitIndices);
+			TArray<FName> RenderGroupTargetLayerNames = ConvertTargetLayerBitIndicesToNamesChecked(RenderStep.RenderGroupBitIndices);
+			TArray<ULandscapeLayerInfoObject*> RenderGroupTargetLayerInfos = bIsHeightmapMerge ? TArray<ULandscapeLayerInfoObject*> { nullptr } : ConvertTargetLayerBitIndicesToLayerInfosChecked(RenderStep.RenderGroupBitIndices);
 
 			// Compute all necessary info about the components affected by this renderer at this step
 			TArray<FComponentMergeRenderInfo> SortedComponentMergeRenderInfos;
@@ -849,7 +848,7 @@ void FMergeRenderContext::Render(TFunction<void(const FOnRenderBatchTargetGroupD
 					// TODO[jonathan.bard] offset the world transform to account for the half-pixel offset?
 					//RenderParams.RenderAreaWorldTransform = FTransform(LandscapeTransform.GetRotation(), LandscapeTransform.GetTranslation() + FVector(RenderBatch.SectionRect.Min) /** (double)ComponentSizeQuads - FVector(0.5, 0.5, 0)*/, LandscapeTransform.GetScale3D());
 
-					FString RenderStepProfilingEventName = FString::Format(TEXT("Step [{0}] ({1}): Render {2}"), { RenderStepIndex, *ConvertTargetLayerNamesToString(TargetLayerNames), *Renderer->GetEditLayerRendererDebugName() });
+					FString RenderStepProfilingEventName = FString::Format(TEXT("Step [{0}] ({1}): Render {2}"), { RenderStepIndex, *ConvertTargetLayerNamesToString(RenderGroupTargetLayerNames), *Renderer->GetEditLayerRendererDebugName() });
 					SCOPED_DRAW_EVENTF_GAMETHREAD(LandscapeLayers, TEXT("%s"), RenderStepProfilingEventName);
 
 					// TODO [jonathan.bard] : this is more of a Batch world transform/section rect at the moment. Shall we have a RenderAreaWorldTransform/RenderAreaSectionRect in FRenderParams and a BatchRenderAreaWorldTransform in FMergeRenderBatch?
@@ -876,17 +875,17 @@ void FMergeRenderContext::Render(TFunction<void(const FOnRenderBatchTargetGroupD
 							}
 						});
 			
-					ILandscapeEditLayerRenderer::FRenderParams RenderParams(this, RenderStep.RenderGroupBitIndices, RenderStep.RendererState, SortedComponentMergeRenderInfos, RenderAreaWorldTransform, RenderAreaSectionRect);
+					ILandscapeEditLayerRenderer::FRenderParams RenderParams(this, RenderGroupTargetLayerNames, RenderGroupTargetLayerInfos, RenderStep.RendererState, SortedComponentMergeRenderInfos, RenderAreaWorldTransform, RenderAreaSectionRect);
 					Renderer->RenderLayer(RenderParams);
 				}
 			}
 			else if ((RenderStep.Type == FMergeRenderStep::EType::SignalBatchMergeGroupDone))
 			{
 				TRACE_CPUPROFILER_EVENT_SCOPE(MergeGroupDone);
-				SCOPED_DRAW_EVENTF_GAMETHREAD(LandscapeLayers, TEXT("Step [%i] (%s) : Render Group Done"), RenderStepIndex, ConvertTargetLayerNamesToString(TargetLayerNames));
+				SCOPED_DRAW_EVENTF_GAMETHREAD(LandscapeLayers, TEXT("Step [%i] (%s) : Render Group Done"), RenderStepIndex, ConvertTargetLayerNamesToString(RenderGroupTargetLayerNames));
 
 				// The last render target we wrote to is the one containing the batch group's merge result : 
-				FOnRenderBatchTargetGroupDoneParams Params(this, RenderBatch, RenderStep.RenderGroupBitIndices, SortedComponentMergeRenderInfos);
+				FOnRenderBatchTargetGroupDoneParams Params(this, RenderBatch, RenderGroupTargetLayerNames, RenderGroupTargetLayerInfos, SortedComponentMergeRenderInfos);
 				OnBatchTargetGroupDone(Params);
 			}
 			else
diff --git a/Engine/Source/Runtime/Landscape/Private/LandscapeEditLayers.cpp b/Engine/Source/Runtime/Landscape/Private/LandscapeEditLayers.cpp
--- a/Engine/Source/Runtime/Landscape/Private/LandscapeEditLayers.cpp
+++ b/Engine/Source/Runtime/Landscape/Private/LandscapeEditLayers.cpp
@@ -5530,8 +5530,7 @@ void ULandscapeEditLayerPersistent::RenderLayer(ILandscapeEditLayerRenderer::FRe
 		}
 	};
 
-	TArray<FName> TargetLayerNames = InRenderParams.MergeRenderContext->ConvertTargetLayerBitIndicesToNamesChecked(InRenderParams.RenderGroupBitIndices);
-	const int32 NumTargetLayers = TargetLayerNames.Num();
+	const int32 NumTargetLayers = InRenderParams.RenderGroupTargetLayerNames.Num();
 
 	// The first step is to copy all of the necessary components' textures to WriteRT and remove the duplicate borders when doing so (subsection by subsection). This is done with a 
 	//  "copy from multiple sources" shader instead of several texture copies, in order to reduce the amount of copy texture commands, which can have a big impact on render-thread performance
@@ -5565,7 +5564,7 @@ void ULandscapeEditLayerPersistent::RenderLayer(ILandscapeEditLayerRenderer::FRe
 	{
 		for (int32 TargetLayerIndex = 0; TargetLayerIndex < NumTargetLayers; ++TargetLayerIndex)
 		{
-			const FName TargetLayerName = TargetLayerNames[TargetLayerIndex];
+			const FName TargetLayerName = InRenderParams.RenderGroupTargetLayerNames[TargetLayerIndex];
 			UTexture2D* SourceTexture = nullptr;
 			FVector2D SourceTextureBias(ForceInit);
 			// Copy all channels by default :
@@ -5814,7 +5813,7 @@ void ULandscapeEditLayerPersistent::RenderLayer(ILandscapeEditLayerRenderer::FRe
 		if (InRenderParams.MergeRenderContext->IsHeightmapMerge())
 		{
 			ENQUEUE_RENDER_COMMAND(LandscapeLayers_Cmd_HeightmapsMergeEditLayer)(
-				[ TargetLayerName = TargetLayerNames[0]
+				[ TargetLayerName = InRenderParams.RenderGroupTargetLayerNames[0]
 				, OutputResource = WriteRT->GetRenderTarget2D()->GetResource()
 				, OutputResourceName = WriteRT->GetDebugName()
 				, CurrentEditLayerResource = CurrentLayerReadRT->GetRenderTarget2D()->GetResource()
@@ -5857,7 +5856,7 @@ void ULandscapeEditLayerPersistent::RenderLayer(ILandscapeEditLayerRenderer::FRe
 			// Prepare a list of all paint layers to merge and do it all in 1 command :
 			TArray<FMergePaintLayerParams> MergePaintLayersParams;
 			MergePaintLayersParams.Reserve(NumTargetLayers);
-			for (const FName& TargetLayerName : TargetLayerNames)
+			for (const FName& TargetLayerName : InRenderParams.RenderGroupTargetLayerNames)
 			{
 				FMergePaintLayerParams& MergePaintLayerParams = MergePaintLayersParams.Emplace_GetRef(TargetLayerName);
 
@@ -6388,13 +6387,12 @@ void ULandscapeHeightmapNormalsEditLayerRenderer::RenderLayer(FRenderParams& InR
 	const FMergeRenderBatch* RenderBatch = RenderContext->GetCurrentRenderBatch();
 	ALandscape* Landscape = RenderContext->GetLandscape();
 
-	TArray<FName> TargetLayerNames = InRenderParams.MergeRenderContext->ConvertTargetLayerBitIndicesToNamesChecked(InRenderParams.RenderGroupBitIndices);
-	checkf(TargetLayerNames.Num() == 1, TEXT("Normals should only be generated on heightmap merge, which should have 1 and only target layer"));
+	checkf(InRenderParams.RenderGroupTargetLayerNames.Num() == 1, TEXT("Normals should only be generated on heightmap merge, which should have 1 and only target layer"));
 
 	RenderContext->CycleBlendRenderTargets(/*InDesiredWriteAccess = */ERHIAccess::RTV);
 	ULandscapeScratchRenderTarget* WriteRT = RenderContext->GetBlendRenderTargetWrite();
 	ULandscapeScratchRenderTarget* ReadRT = RenderContext->GetBlendRenderTargetRead();
-	ULandscapeScratchRenderTarget* ValidityRT = RenderContext->GetValidityRenderTarget(TargetLayerNames[0]);
+	ULandscapeScratchRenderTarget* ValidityRT = RenderContext->GetValidityRenderTarget(InRenderParams.RenderGroupTargetLayerNames[0]);
 	check((WriteRT != nullptr) && (ReadRT != nullptr) && (ValidityRT != nullptr));
 
 	ValidityRT->TransitionTo(ERHIAccess::SRVMask);
@@ -6445,11 +6443,9 @@ void ULandscapeWeightmapWeightBlendedLayersRenderer::RenderLayer(FRenderParams&
 	ULandscapeScratchRenderTarget* WriteRT = InRenderParams.MergeRenderContext->GetBlendRenderTargetWrite();
 	ULandscapeScratchRenderTarget* ReadRT = InRenderParams.MergeRenderContext->GetBlendRenderTargetRead();
 
-	TArray<FName> TargetLayerNames = InRenderParams.MergeRenderContext->ConvertTargetLayerBitIndicesToNamesChecked(InRenderParams.RenderGroupBitIndices);
-	TArray<ULandscapeLayerInfoObject*> LayerInfos = InRenderParams.MergeRenderContext->ConvertTargetLayerBitIndicesToLayerInfosChecked(InRenderParams.RenderGroupBitIndices);
 	TArray<FWeightmapPaintLayerInfo> WeightmapPaintLayerInfos;
-	WeightmapPaintLayerInfos.Reserve(LayerInfos.Num());
-	Algo::Transform(LayerInfos, WeightmapPaintLayerInfos, [](ULandscapeLayerInfoObject* InLayerInfo) 
+	WeightmapPaintLayerInfos.Reserve(InRenderParams.RenderGroupTargetLayerInfos.Num());
+	Algo::Transform(InRenderParams.RenderGroupTargetLayerInfos, WeightmapPaintLayerInfos, [](ULandscapeLayerInfoObject* InLayerInfo)
 		{
 			check(InLayerInfo != nullptr); // There should only be valid layer infos at this point
 			FWeightmapPaintLayerInfo WeightmapPaintLayerInfo;
@@ -6465,7 +6461,7 @@ void ULandscapeWeightmapWeightBlendedLayersRenderer::RenderLayer(FRenderParams&
 		});
 
 	ENQUEUE_RENDER_COMMAND(LandscapeLayers_Cmd_WeightmapsPerformLegacyWeightBlending)(
-		[TargetLayerNames
+		[TargetLayerNames = InRenderParams.RenderGroupTargetLayerNames
 		, WeightmapPaintLayerInfos
 		, OutputResource = WriteRT->GetRenderTarget2DArray()->GetResource()
 		, OutputResourceName = WriteRT->GetDebugName()
@@ -10347,7 +10343,6 @@ int32 ALandscape::PerformLayersWeightmapsBatchedMerge(FUpdateLayersContentContex
 		const int32 TotalNumSubsections = Landscape->NumSubsections * Landscape->NumSubsections;
 		FIntPoint MinWeightmapResolution(MAX_int32, MAX_int32);
 		FIntPoint MaxWeightmapResolution(MIN_int32, MIN_int32);
-		TArray<FName> TargetLayerNames = InParams.MergeRenderContext->ConvertTargetLayerBitIndicesToNamesChecked(InParams.RenderGroupBitIndices);
 
 		struct FWeightmapResolveInfo
 		{
@@ -10440,7 +10435,7 @@ int32 ALandscape::PerformLayersWeightmapsBatchedMerge(FUpdateLayersContentContex
 						{
 							check((AllocationInfo.WeightmapTextureChannel >= 0) && (AllocationInfo.WeightmapTextureChannel < 4));
 							checkf(((WeightmapResolveInfo->ChannelMask & (1 << AllocationInfo.WeightmapTextureChannel)) == 0), TEXT("This channel has already been resolved, it shouldn't happen, it would mean that 2 allocations are using the same channel"));
-							int32 SliceIndex = TargetLayerNames.IndexOfByKey(AllocationInfo.GetLayerName());
+							int32 SliceIndex = InParams.RenderGroupTargetLayerNames.IndexOfByKey(AllocationInfo.GetLayerName());
 							checkf(SliceIndex != INDEX_NONE, TEXT("Couldn't find %s in the list of weightmaps that have been produced"), *AllocationInfo.GetLayerName().ToString());
 							WeightmapResolveInfo->SourceSliceIndexPerChannel[AllocationInfo.WeightmapTextureChannel] = SliceIndex;
 							WeightmapResolveInfo->SourceSubsectionRectsPerChannel[AllocationInfo.WeightmapTextureChannel] = SourceSubsectionRects;
diff --git a/Engine/Source/Runtime/Landscape/Public/LandscapeEditLayerRenderer.h b/Engine/Source/Runtime/Landscape/Public/LandscapeEditLayerRenderer.h
--- a/Engine/Source/Runtime/Landscape/Public/LandscapeEditLayerRenderer.h
+++ b/Engine/Source/Runtime/Landscape/Public/LandscapeEditLayerRenderer.h
@@ -394,15 +394,19 @@ public:
 	LANDSCAPE_API ULandscapeScratchRenderTarget* GetBlendRenderTargetWrite() const;
 	LANDSCAPE_API ULandscapeScratchRenderTarget* GetBlendRenderTargetRead() const;
 	LANDSCAPE_API ULandscapeScratchRenderTarget* GetBlendRenderTargetReadPrevious() const;
-	// TODO [jonathan.bard] : pass a target layer index instead of a name here?
 	LANDSCAPE_API ULandscapeScratchRenderTarget* GetValidityRenderTarget(const FName& InTargetLayerName) const;
 	
 	struct FOnRenderBatchTargetGroupDoneParams
 	{
-		FOnRenderBatchTargetGroupDoneParams(FMergeRenderContext* InMergeRenderContext, const FMergeRenderBatch& InRenderBatch, const TBitArray<>& InRenderGroupBitIndices, const TArrayView<FComponentMergeRenderInfo>& InSortedComponentMergeRenderInfos)
+		FOnRenderBatchTargetGroupDoneParams(FMergeRenderContext* InMergeRenderContext, 
+			const FMergeRenderBatch& InRenderBatch, 
+			const TArrayView<FName>& InRenderGroupTargetLayerNames,
+			const TArrayView<ULandscapeLayerInfoObject*>& InRenderGroupTargetLayerInfos, 
+			const TArrayView<FComponentMergeRenderInfo>& InSortedComponentMergeRenderInfos)
 			: MergeRenderContext(InMergeRenderContext)
 			, RenderBatch(&InRenderBatch)
-			, RenderGroupBitIndices(InRenderGroupBitIndices)
+			, RenderGroupTargetLayerNames(InRenderGroupTargetLayerNames)
+			, RenderGroupTargetLayerInfos(InRenderGroupTargetLayerInfos)
 			, SortedComponentMergeRenderInfos(InSortedComponentMergeRenderInfos)
 		{}
 
@@ -412,8 +416,11 @@ public:
 		/** Batch to that was just rendered for this render group */
 		const FMergeRenderBatch* RenderBatch = nullptr;
 
-		/** Render group that was just rendered for this batch. Each bit in that bit array corresponds to an entry in FMergeRenderContext's AllTargetLayerNames  */
-		const TBitArray<> RenderGroupBitIndices;
+		/** List of target layers being involved in this step */
+		TArray<FName> RenderGroupTargetLayerNames;
+
+		/** List of target layer info objects being involved in this step (same size as RenderGroupTargetLayerNames) */
+		TArray<ULandscapeLayerInfoObject*> RenderGroupTargetLayerInfos;
 
 		/** Additional info about the components that have been processed in this batch render */
 		const TArray<FComponentMergeRenderInfo> SortedComponentMergeRenderInfos;
@@ -739,6 +746,7 @@ class ULandscapeEditLayerRenderer :
 	public UInterface
 {
 	GENERATED_BODY()
+
 };
 
 /** 
@@ -755,12 +763,14 @@ public:
 	struct FRenderParams
 	{
 		FRenderParams(UE::Landscape::EditLayers::FMergeRenderContext* InMergeRenderContext, 
-			const TBitArray<>& InRenderGroupBitIndices,
+			const TArrayView<FName>& InRenderGroupTargetLayerNames,
+			const TArrayView<ULandscapeLayerInfoObject*>& InRenderGroupTargetLayerInfos,
 			const UE::Landscape::EditLayers::FEditLayerRendererState& InRendererState, 
 			const TArrayView<UE::Landscape::EditLayers::FComponentMergeRenderInfo>& InSortedComponentMergeRenderInfos, 
 			const FTransform& InRenderAreaWorldTransform, const FIntRect& InRenderAreaSectionRect)
 			: MergeRenderContext(InMergeRenderContext)
-			, RenderGroupBitIndices(InRenderGroupBitIndices)
+			, RenderGroupTargetLayerNames(InRenderGroupTargetLayerNames)
+			, RenderGroupTargetLayerInfos(InRenderGroupTargetLayerInfos)
 			, RendererState(InRendererState)
 			, SortedComponentMergeRenderInfos(InSortedComponentMergeRenderInfos)
 			, RenderAreaWorldTransform(InRenderAreaWorldTransform)
@@ -770,8 +780,11 @@ public:
 		/** Merge context */
 		UE::Landscape::EditLayers::FMergeRenderContext* MergeRenderContext = nullptr;
 
-		/** List of target layers being involved in this step. Each bit in that bit array corresponds to an entry in FMergeRenderContext's AllTargetLayerNames  */
-		TBitArray<> RenderGroupBitIndices;
+		/** List of target layers being involved in this step */
+		TArray<FName> RenderGroupTargetLayerNames;
+
+		/** List of target layer info objects being involved in this step (same size as RenderGroupTargetLayerNames) */
+		TArray<ULandscapeLayerInfoObject*> RenderGroupTargetLayerInfos;
 
 		/** Full state for the renderer involved in this step. This allows to retrieve the exact state of this renderer (e.g. enabled weightmaps, which can be different than the render group, in that 
 		 target layers A, B and C might belong to the same group but this renderer actually only has A enabled). This is therefore the renderer's responsibility to check that a given target layer from the 
