Made collision settings LandscapeOverridable

diff --git a/Engine/Source/Runtime/Core/Public/UObject/FortniteMainBranchObjectVersion.h b/Engine/Source/Runtime/Core/Public/UObject/FortniteMainBranchObjectVersion.h
--- a/Engine/Source/Runtime/Core/Public/UObject/FortniteMainBranchObjectVersion.h
+++ b/Engine/Source/Runtime/Core/Public/UObject/FortniteMainBranchObjectVersion.h
@@ -514,6 +514,9 @@ struct FFortniteMainBranchObjectVersion
 		// Enable SkipOnlyEditorOnly style cooking of UStaticMeshComponent::MeshPaintTexture
 		MeshPaintTextureUsesEditorOnly,
 
+		// Fixup and synchronize some landscape properties that have moved to the property sharing/overriding system :
+		LandscapeBodyInstanceAsSharedProperty,
+
 		// -----<new versions can be added above this line>-------------------------------------------------
 		VersionPlusOne,
 		LatestVersion = VersionPlusOne - 1
diff --git a/Engine/Source/Runtime/Landscape/Classes/LandscapeProxy.h b/Engine/Source/Runtime/Landscape/Classes/LandscapeProxy.h
--- a/Engine/Source/Runtime/Landscape/Classes/LandscapeProxy.h
+++ b/Engine/Source/Runtime/Landscape/Classes/LandscapeProxy.h
@@ -68,6 +68,15 @@ LANDSCAPE_API extern bool GLandscapeEditModeActive;
 extern int32 GGrassMapUseRuntimeGeneration;
 
 DECLARE_MULTICAST_DELEGATE_TwoParams(FOnLandscapeProxyComponentDataChanged, ALandscapeProxy*, const FLandscapeProxyComponentDataChangedParams&);
+
+struct FOnLandscapeProxyFixupSharedDataParams
+{
+	/** Parent landscape actor for the landscape proxy on which OnLandscapeProxyFixupSharedData is triggered */
+	ALandscape* Landscape = nullptr;
+	/** Indicates whether UpgradeSharedProperties has been called on this proxy already (before this call) */
+	bool bUpgradeSharedPropertiesPerformed = false;
+};
+DECLARE_MULTICAST_DELEGATE_TwoParams(FOnLandscapeProxyFixupSharedDataDelegate, ALandscapeProxy* /*InProxy*/, const FOnLandscapeProxyFixupSharedDataParams& /*InParams*/);
 #endif // WITH_EDITOR
 
 
@@ -462,6 +471,8 @@ protected:
 
 	UPROPERTY(EditAnywhere, Category = Nanite, AdvancedDisplay, meta = (EditCondition = "bEnableNanite", LandscapeInherited))
 	float NaniteMaxEdgeLengthFactor = 16.0f;
+
+	static FOnLandscapeProxyFixupSharedDataDelegate OnLandscapeProxyFixupSharedDataDelegate;
 #endif // WITH_EDITORONLY_DATA
 
 	/** Disable runtime grass data generation.  If disabled, the grass maps will be serialized at cook time. Do not set directly, use ALandscape::SetDisableRuntimeGrassMapGeneration to ensure it is set on all loaded proxies. */
@@ -817,7 +828,7 @@ public:
 	int32 SimpleCollisionMipLevel;
 
 	/** Collision profile settings for this landscape */
-	UPROPERTY(EditAnywhere, BlueprintReadOnly, Category=Collision, meta=(ShowOnlyInnerProperties))
+	UPROPERTY(EditAnywhere, BlueprintReadOnly, Category = Collision, meta = (LandscapeOverridable))
 	FBodyInstance BodyInstance;
 
 	/**
@@ -1246,6 +1257,9 @@ PRAGMA_ENABLE_DEPRECATION_WARNINGS
 
 	// Modifies the override state of the property given as argument.
 	virtual void SetSharedPropertyOverride(const FName& InPropertyName, const bool bIsOverriden) { }
+
+	/** Delegate that will be called whenever FixupSharedData is called, to inject some custom logic */
+	LANDSCAPE_API static FOnLandscapeProxyFixupSharedDataDelegate::RegistrationType& OnLandscapeProxyFixupSharedData() { return OnLandscapeProxyFixupSharedDataDelegate; }
 #endif // WITH_EDITOR
 
 	// Get Landscape Material assigned to this Landscape
diff --git a/Engine/Source/Runtime/Landscape/Private/Landscape.cpp b/Engine/Source/Runtime/Landscape/Private/Landscape.cpp
--- a/Engine/Source/Runtime/Landscape/Private/Landscape.cpp
+++ b/Engine/Source/Runtime/Landscape/Private/Landscape.cpp
@@ -4805,7 +4805,10 @@ void ALandscapeProxy::FixupSharedData(ALandscape* Landscape, const bool bMapChec
 		return;
 	}
 
-	if (!bUpgradeSharedPropertiesPerformed && GetLinkerCustomVersion(FFortniteReleaseBranchCustomObjectVersion::GUID) < FFortniteReleaseBranchCustomObjectVersion::LandscapeSharedPropertiesEnforcement)
+	const bool bUpgradeSharedPropertiesPerformedBefore = bUpgradeSharedPropertiesPerformed;
+	if (!bUpgradeSharedPropertiesPerformed && 
+		((GetLinkerCustomVersion(FFortniteReleaseBranchCustomObjectVersion::GUID) < FFortniteReleaseBranchCustomObjectVersion::LandscapeSharedPropertiesEnforcement)
+		|| (GetLinkerCustomVersion(FFortniteMainBranchObjectVersion::GUID) < FFortniteMainBranchObjectVersion::LandscapeBodyInstanceAsSharedProperty)))
 	{
 		UpgradeSharedProperties(Landscape);
 		bUpgradeSharedPropertiesPerformed = true;
@@ -4841,6 +4844,8 @@ void ALandscapeProxy::FixupSharedData(ALandscape* Landscape, const bool bMapChec
 			}
 		}
 	}
+
+	OnLandscapeProxyFixupSharedDataDelegate.Broadcast(/*Proxy = */this, FOnLandscapeProxyFixupSharedDataParams { .Landscape = Landscape, .bUpgradeSharedPropertiesPerformed = bUpgradeSharedPropertiesPerformedBefore });
 }
 
 void ALandscapeProxy::SetAbsoluteSectionBase(FIntPoint InSectionBase)
diff --git a/Engine/Source/Runtime/Landscape/Private/LandscapeEdit.cpp b/Engine/Source/Runtime/Landscape/Private/LandscapeEdit.cpp
--- a/Engine/Source/Runtime/Landscape/Private/LandscapeEdit.cpp
+++ b/Engine/Source/Runtime/Landscape/Private/LandscapeEdit.cpp
@@ -126,6 +126,8 @@ extern const size_t ChannelOffsets[4];
 extern float LandscapeNaniteBuildLag;
 
 ULandscapeLayerInfoObject* ALandscapeProxy::VisibilityLayer = nullptr;
+
+FOnLandscapeProxyFixupSharedDataDelegate ALandscapeProxy::OnLandscapeProxyFixupSharedDataDelegate;
 #endif //WITH_EDITOR
 
 void ULandscapeComponent::Init(int32 InBaseX, int32 InBaseY, int32 InComponentSizeQuads, int32 InNumSubsections, int32 InSubsectionSizeQuads)
