Fixed incorrect parallel-for when retrieving the list of outdated landscape proxies

diff --git a/Engine/Source/Runtime/Landscape/Private/LandscapeSubsystem.cpp b/Engine/Source/Runtime/Landscape/Private/LandscapeSubsystem.cpp
--- a/Engine/Source/Runtime/Landscape/Private/LandscapeSubsystem.cpp
+++ b/Engine/Source/Runtime/Landscape/Private/LandscapeSubsystem.cpp
@@ -695,22 +695,27 @@ TArray<TTuple<ALandscapeProxy*, UE::Landscape::EOutdatedDataFlags>> ULandscapeSu
 		return {};
 	}
 
-	TArray<TTuple<ALandscapeProxy*, UE::Landscape::EOutdatedDataFlags>> OutdatedProxies;
-	
-	ParallelForWithTaskContext(OutdatedProxies, Proxies.Num(), [this, InMatchingOutdatedDataFlags, bInMustMatchAllFlags](TTuple<ALandscapeProxy*, UE::Landscape::EOutdatedDataFlags>& Context, int32 Index)
-	{
-		FTaskTagScope Scope(ETaskTag::EParallelGameThread);
-		ALandscapeProxy* ValidProxy = Proxies[Index].Get();
-		const UE::Landscape::EOutdatedDataFlags ProxyOutdatedDataFlags = ValidProxy->GetOutdatedDataFlags();
+	using ProxyAndFlagsArray = TArray<TTuple<ALandscapeProxy*, UE::Landscape::EOutdatedDataFlags>>;
 
-		if ((bInMustMatchAllFlags && EnumHasAllFlags(ProxyOutdatedDataFlags, InMatchingOutdatedDataFlags))
-			|| (!bInMustMatchAllFlags && EnumHasAnyFlags(ProxyOutdatedDataFlags, InMatchingOutdatedDataFlags)))
+	// Parallelize the retrieval of the outdated proxies and their flags by "fork and join". Each task in the parallel-for handles a certain number of proxies so task contexts simply consists in an array of proxy+flags : 
+	TArray<ProxyAndFlagsArray> TaskContexts;
+	ParallelForWithTaskContext(TaskContexts, Proxies.Num(), [this, InMatchingOutdatedDataFlags, bInMustMatchAllFlags](ProxyAndFlagsArray& InTaskContext, int32 InIndex)
 		{
-			Context = {ValidProxy, ProxyOutdatedDataFlags};
-		}
-	});
+			FTaskTagScope Scope(ETaskTag::EParallelGameThread);
+			ALandscapeProxy* ValidProxy = Proxies[InIndex].Get();
+			const UE::Landscape::EOutdatedDataFlags ProxyOutdatedDataFlags = ValidProxy->GetOutdatedDataFlags();
+
+			if ((bInMustMatchAllFlags && EnumHasAllFlags(ProxyOutdatedDataFlags, InMatchingOutdatedDataFlags))
+				|| (!bInMustMatchAllFlags && EnumHasAnyFlags(ProxyOutdatedDataFlags, InMatchingOutdatedDataFlags)))
+			{
+				InTaskContext.Add({ ValidProxy, ProxyOutdatedDataFlags });
+			}
+		});
 
-	OutdatedProxies.RemoveAllSwap([](const TTuple<ALandscapeProxy*, UE::Landscape::EOutdatedDataFlags>& OutdatedProxy){ return OutdatedProxy.Key == nullptr; });
+	// Join all outdated proxies that have been found by the different tasks :
+	ProxyAndFlagsArray OutdatedProxies;
+	OutdatedProxies.Reserve(Proxies.Num());
+	Algo::ForEach(TaskContexts, [&OutdatedProxies](ProxyAndFlagsArray& InTaskContext) { OutdatedProxies.Append(InTaskContext); });
 	return OutdatedProxies;
 }
 
