World Partition - HLOD - Landscape: Added "HLOD Material Override" property
* Will allow users to specify a custom material for their HLOD landscape (possibly fetching from an SVT)
* Specifying a material override will bypass material baking for the HLOD mesh, improving build times.
* Improved doc & removed meta keyword "EditConditionHides" on related properties in order to make the UI more consistent
#tests Used the option in the open world template to validate the result.

diff --git a/Engine/Source/Runtime/Landscape/Classes/LandscapeProxy.h b/Engine/Source/Runtime/Landscape/Classes/LandscapeProxy.h
--- a/Engine/Source/Runtime/Landscape/Classes/LandscapeProxy.h
+++ b/Engine/Source/Runtime/Landscape/Classes/LandscapeProxy.h
@@ -923,20 +923,24 @@ PRAGMA_ENABLE_DEPRECATION_WARNINGS
 	bool bUseLandscapeForCullingInvisibleHLODVertices;
 
 #if WITH_EDITORONLY_DATA
-	/** Specify how to choose the texture size of the resulting HLOD mesh */
-	UPROPERTY(EditAnywhere, Category = HLOD, meta = (DisplayName = "HLOD Texture Size Policy", LandscapeOverridable))
+	/** Specify how to choose the texture size of the resulting HLOD mesh. Specifying an HLOD Material Override will disable this option as no texture will be baked. */
+	UPROPERTY(EditAnywhere, Category = HLOD, meta = (DisplayName = "HLOD Texture Size Policy", LandscapeOverridable, EditCondition = "HLODMaterialOverride == nullptr"))
 	ELandscapeHLODTextureSizePolicy HLODTextureSizePolicy;
 
-	/** Specify the texture size to use for the HLOD mesh if HLODTextureSizePolicy is set to SpecificSize */
-	UPROPERTY(EditAnywhere, Category = HLOD, meta = (DisplayName = "HLOD Texture Size", LandscapeOverridable, EditCondition = "HLODTextureSizePolicy == ELandscapeHLODTextureSizePolicy::SpecificSize", EditConditionHides, ClampMin = "16", ClampMax = "8192"))
+	/** Specify the texture size to use for the HLOD mesh if HLODTextureSizePolicy is set to SpecificSize. Specifying an HLOD Material Override will disable this option as no texture will be baked. */
+	UPROPERTY(EditAnywhere, Category = HLOD, meta = (DisplayName = "HLOD Texture Size", LandscapeOverridable, EditCondition = "HLODMaterialOverride == nullptr && HLODTextureSizePolicy == ELandscapeHLODTextureSizePolicy::SpecificSize", ClampMin = "16", ClampMax = "8192"))
 	int32 HLODTextureSize;
 
-	/** Specify how to choose the LOD used as input for the HLOD mesh */
+	/** Specify a custom HLOD material to apply to the HLOD mesh. Specifying an HLOD Material Override will result in no texture being baked for the HLOD mesh. */
+	UPROPERTY(EditAnywhere, Category = HLOD, meta = (DisplayName = "HLOD Material Override", LandscapeOverridable))
+	TObjectPtr<UMaterialInterface> HLODMaterialOverride;
+
+	/** Specify how to choose the LOD used as input for the HLOD mesh. */
 	UPROPERTY(EditAnywhere, Category = HLOD, meta = (DisplayName = "HLOD Mesh Source LOD Policy", LandscapeOverridable))
 	ELandscapeHLODMeshSourceLODPolicy HLODMeshSourceLODPolicy;
 
-	/** Specify which LOD to use for the HLOD mesh if HLODMeshSourceLODPolicy is set to SpecificLOD */
-	UPROPERTY(EditAnywhere, Category = HLOD, meta = (DisplayName = "HLOD Mesh Source LOD", LandscapeOverridable, EditCondition = "HLODMeshSourceLODPolicy == ELandscapeHLODMeshSourceLODPolicy::SpecificLOD", EditConditionHides, ClampMin = "0"))
+	/** Specify which LOD to use for the HLOD mesh if HLODMeshSourceLODPolicy is set to SpecificLOD. */
+	UPROPERTY(EditAnywhere, Category = HLOD, meta = (DisplayName = "HLOD Mesh Source LOD", LandscapeOverridable, EditCondition = "HLODMeshSourceLODPolicy == ELandscapeHLODMeshSourceLODPolicy::SpecificLOD", ClampMin = "0"))
 	int32 HLODMeshSourceLOD;
 #endif
 
diff --git a/Engine/Source/Runtime/Landscape/Private/LandscapeHLODBuilder.cpp b/Engine/Source/Runtime/Landscape/Private/LandscapeHLODBuilder.cpp
--- a/Engine/Source/Runtime/Landscape/Private/LandscapeHLODBuilder.cpp
+++ b/Engine/Source/Runtime/Landscape/Private/LandscapeHLODBuilder.cpp
@@ -420,8 +420,16 @@ TArray<UActorComponent*> ULandscapeHLODBuilder::Build(const FHLODBuildContext& I
 
 		// Material
 		{
-			int32 TextureSize = ComputeRequiredTextureSize(LandscapeProxy, static_cast<float>(InHLODBuildContext.MinVisibleDistance), MeshDescription);
-			UMaterialInterface* LandscapeMaterial = BakeLandscapeMaterial(InHLODBuildContext, *MeshDescription, LandscapeProxy, TextureSize);
+			UMaterialInterface* LandscapeMaterial;
+			if (LandscapeProxy->HLODMaterialOverride)
+			{
+				LandscapeMaterial = LandscapeProxy->HLODMaterialOverride.Get();
+			}
+			else
+			{
+				int32 TextureSize = ComputeRequiredTextureSize(LandscapeProxy, static_cast<float>(InHLODBuildContext.MinVisibleDistance), MeshDescription);
+				LandscapeMaterial = BakeLandscapeMaterial(InHLODBuildContext, *MeshDescription, LandscapeProxy, TextureSize);
+			}
 
 			//Assign the proxy material to the static mesh
 			StaticMesh->GetStaticMaterials().Add(FStaticMaterial(LandscapeMaterial));
