[Landscape] Fix Landscape proxies registering their landscape info object into the main world instead of their outer world.
The logic behind this changes comes from the fact that WorldCreationEventFunction and WorldDestroyEventFunction uses the outer world to register and unregister landscape info objects.
Also outered landscape info objects to the world they are registered to just to make it clear to what object they belongs to.

diff --git a/Engine/Source/Runtime/Landscape/Classes/LandscapeInfoMap.h b/Engine/Source/Runtime/Landscape/Classes/LandscapeInfoMap.h
--- a/Engine/Source/Runtime/Landscape/Classes/LandscapeInfoMap.h
+++ b/Engine/Source/Runtime/Landscape/Classes/LandscapeInfoMap.h
@@ -18,14 +18,8 @@ class ULandscapeInfoMap : public UObject
 	GENERATED_UCLASS_BODY()
 
 public:
-	void BeginDestroy() override;
-	void PostDuplicate(bool bDuplicateForPIE) override;
-	void Serialize(FArchive& Ar) override;
-
-	static void AddReferencedObjects(UObject* InThis, FReferenceCollector& Collector);
-
+	UPROPERTY()
 	TMap<FGuid, TObjectPtr<ULandscapeInfo>> Map;
-	TWeakObjectPtr<UWorld> World;
 
 	/**
 	* Gets landscape-specific data for given world.
diff --git a/Engine/Source/Runtime/Landscape/Private/Landscape.cpp b/Engine/Source/Runtime/Landscape/Private/Landscape.cpp
--- a/Engine/Source/Runtime/Landscape/Private/Landscape.cpp
+++ b/Engine/Source/Runtime/Landscape/Private/Landscape.cpp
@@ -1762,14 +1762,14 @@ void ALandscape::MarkAllLandscapeRenderStateDirty()
 
 ULandscapeInfo* ALandscapeProxy::CreateLandscapeInfo(bool bMapCheck, bool bUpdateAllAddCollisions)
 {
-	ULandscapeInfo* LandscapeInfo = ULandscapeInfo::FindOrCreate(GetWorld(), LandscapeGuid);
+	ULandscapeInfo* LandscapeInfo = ULandscapeInfo::FindOrCreate(GetTypedOuter<UWorld>(), LandscapeGuid);
 	LandscapeInfo->RegisterActor(this, bMapCheck, bUpdateAllAddCollisions);
 	return LandscapeInfo;
 }
 
 ULandscapeInfo* ALandscapeProxy::GetLandscapeInfo() const
 {
-	return ULandscapeInfo::Find(GetWorld(), LandscapeGuid);
+	return ULandscapeInfo::Find(GetTypedOuter<UWorld>(), LandscapeGuid);
 }
 
 FTransform ALandscapeProxy::LandscapeActorToWorld() const
@@ -5345,7 +5345,8 @@ ULandscapeInfo* ULandscapeInfo::FindOrCreate(UWorld* InWorld, const FGuid& Lands
 
 	if (!LandscapeInfo)
 	{
-		LandscapeInfo = NewObject<ULandscapeInfo>(GetTransientPackage(), NAME_None, RF_Transactional | RF_Transient);
+		const FName LandscapeInfoName = MakeUniqueObjectName(InWorld, ULandscapeInfo::StaticClass(), TEXT("LandscapeInfo"));
+		LandscapeInfo = NewObject<ULandscapeInfo>(InWorld, LandscapeInfoName, RF_Transactional | RF_Transient);
 		LandscapeInfoMap.Modify(false);
 		LandscapeInfo->Initialize(InWorld, LandscapeGuid);
 		LandscapeInfoMap.Map.Add(LandscapeGuid, LandscapeInfo);
diff --git a/Engine/Source/Runtime/Landscape/Private/LandscapeInfoMap.cpp b/Engine/Source/Runtime/Landscape/Private/LandscapeInfoMap.cpp
--- a/Engine/Source/Runtime/Landscape/Private/LandscapeInfoMap.cpp
+++ b/Engine/Source/Runtime/Landscape/Private/LandscapeInfoMap.cpp
@@ -8,50 +8,13 @@
 
 ULandscapeInfoMap::ULandscapeInfoMap(const FObjectInitializer& ObjectInitializer)
 	: Super(ObjectInitializer)
-	, World(nullptr)
-{
-}
-
-void ULandscapeInfoMap::PostDuplicate(bool bDuplicateForPIE)
-{
-	Super::PostDuplicate(bDuplicateForPIE);
-
-	check(Map.Num() == 0);
-}
-
-void ULandscapeInfoMap::Serialize(FArchive& Ar)
-{
-	Super::Serialize(Ar);
-
-	if (Ar.IsTransacting() || Ar.IsObjectReferenceCollector())
-	{
-		Ar << Map;
-	}
-}
-
-void ULandscapeInfoMap::BeginDestroy()
-{
-	if (World.IsValid())
-	{
-		World->PerModuleDataObjects.Remove(this);
-	}
-
-	Super::BeginDestroy();
-}
-
-void ULandscapeInfoMap::AddReferencedObjects(UObject* InThis, FReferenceCollector& Collector)
-{
-	ULandscapeInfoMap* This = CastChecked<ULandscapeInfoMap>(InThis);
-	Collector.AddReferencedObjects(This->Map, This);
-}
+{}
 
 ULandscapeInfoMap& ULandscapeInfoMap::GetLandscapeInfoMap(const UWorld* World)
 {
 	ULandscapeInfoMap *FoundObject = nullptr;
 	World->PerModuleDataObjects.FindItemByClass(&FoundObject);
-
 	checkf(FoundObject, TEXT("ULandscapInfoMap object was not created for this UWorld."));
-
 	return *FoundObject;
 }
 
diff --git a/Engine/Source/Runtime/Landscape/Private/LandscapeModule.cpp b/Engine/Source/Runtime/Landscape/Private/LandscapeModule.cpp
--- a/Engine/Source/Runtime/Landscape/Private/LandscapeModule.cpp
+++ b/Engine/Source/Runtime/Landscape/Private/LandscapeModule.cpp
@@ -55,15 +55,9 @@ private:
  */
 void AddPerWorldLandscapeData(UWorld* World)
 {
-	EObjectFlags NewLandscapeDataFlags = RF_NoFlags;
 	if (!World->PerModuleDataObjects.FindItemByClass<ULandscapeInfoMap>())
 	{
-		if (World->HasAnyFlags(RF_Transactional))
-		{
-			NewLandscapeDataFlags = RF_Transactional;
-		}
-		ULandscapeInfoMap* InfoMap = NewObject<ULandscapeInfoMap>(GetTransientPackage(), NAME_None, NewLandscapeDataFlags);
-		InfoMap->World = World;
+		ULandscapeInfoMap* InfoMap = NewObject<ULandscapeInfoMap>(World, TEXT("LandscapeInfoMap"), World->GetMaskedFlags(RF_Transactional) | RF_Transient);
 		World->PerModuleDataObjects.Add(InfoMap);
 	}
 }
@@ -167,8 +161,6 @@ void WorldDuplicateEventFunction(UWorld* World, bool bDuplicateForPIE, TMap<UObj
 	if (World->PerModuleDataObjects.FindItemByClass(&InfoMap, &Index))
 	{
 		ULandscapeInfoMap* NewInfoMap = Cast<ULandscapeInfoMap>( StaticDuplicateObject(InfoMap, InfoMap->GetOuter()) );
-		NewInfoMap->World = World;
-
 		World->PerModuleDataObjects[Index] = NewInfoMap;
 	}
 	else
diff --git a/Engine/Source/Runtime/Landscape/Private/LandscapeSplineActor.cpp b/Engine/Source/Runtime/Landscape/Private/LandscapeSplineActor.cpp
--- a/Engine/Source/Runtime/Landscape/Private/LandscapeSplineActor.cpp
+++ b/Engine/Source/Runtime/Landscape/Private/LandscapeSplineActor.cpp
@@ -114,15 +114,14 @@ void ALandscapeSplineActor::PostRegisterAllComponents()
 		}
 #endif // WITH_EDITORONLY_DATA
 
-		UWorld* World = GetWorld();
 		if (LandscapeGuid.IsValid())
 		{
-			ULandscapeInfo* LandscapeInfo = ULandscapeInfo::FindOrCreate(World, LandscapeGuid);
+			ULandscapeInfo* LandscapeInfo = ULandscapeInfo::FindOrCreate(GetTypedOuter<UWorld>(), LandscapeGuid);
 			LandscapeInfo->RegisterSplineActor(this);
 		}
 
 		// If Landscape uses generated LandscapeSplineMeshesActors, ensure SplineMeshComponents & ControlPointMeshComponents are hidden in PIE
-		if (World->IsGameWorld() && HasGeneratedLandscapeSplineMeshesActors())
+		if (GetWorld()->IsGameWorld() && HasGeneratedLandscapeSplineMeshesActors())
 		{
 			ForEachComponent<UStaticMeshComponent>(true, [](UStaticMeshComponent* Component)
 			{
