Added tooltips to Landscape Dirtying Mode settings and made InLandscapeModeAndUserTriggeredChanges the default, now that it's been battle-tested

diff --git a/Engine/Source/Runtime/Landscape/Public/LandscapeSettings.h b/Engine/Source/Runtime/Landscape/Public/LandscapeSettings.h
--- a/Engine/Source/Runtime/Landscape/Public/LandscapeSettings.h
+++ b/Engine/Source/Runtime/Landscape/Public/LandscapeSettings.h
@@ -13,8 +13,19 @@ class ULandscapeLayerInfoObject;
 UENUM()
 enum class ELandscapeDirtyingMode : uint8
 {
+	/** Auto : 
+	 Landscapes that are marked as needing to be resaved will appear in the Choose files to save dialog.
+	 Changes are saved whenever the Landscape requires it.*/
 	Auto,
+	/** In Landscape Mode Only : 
+	 Landscapes that are marked as needing to be resaved will not appear in the Choose files to save dialog.
+	 This is a manual saving mode that puts the responsibility on the user to avoid file contention with other team members.
+	 The viewport will display an error message indicating that landscape actors are not up-to-date and need to be resaved. This is done using Build > Save Modified Landscapes (or Build > Build Landscape). */
 	InLandscapeModeOnly,
+	/** In Landscape Mode And User Triggered Changes : 
+	 Landscapes that are marked as needing to be resaved will not appear in the Choose files to save dialog.
+	 However, any user-triggered changes (direct or indirect) will require the Landscape to be resaved.
+	 This mode is recommended for team collaboration as it provides the best features of the other two modes while ensuring that modified landscape actors are still saved and properly submitted to source control. */
 	InLandscapeModeAndUserTriggeredChanges
 };
 
@@ -74,8 +85,9 @@ public:
 	UPROPERTY(config, EditAnywhere, Category = "Configuration", meta = (ToolTip = "Disable Painting Startup Slowdown"))
 	bool bDisablePaintingStartupSlowdown = true;
 	
+	/** Defines when the engine requires the landscape actors to be resaved */
 	UPROPERTY(Config, Category = "Configuration", EditAnywhere)
-	ELandscapeDirtyingMode LandscapeDirtyingMode;
+	ELandscapeDirtyingMode LandscapeDirtyingMode = ELandscapeDirtyingMode::InLandscapeModeAndUserTriggeredChanges;
 
 protected:
 	UPROPERTY(config)
