Fix for server crash in landscape PostLoad

diff --git a/Engine/Source/Runtime/Landscape/Private/LandscapeEdit.cpp b/Engine/Source/Runtime/Landscape/Private/LandscapeEdit.cpp
--- a/Engine/Source/Runtime/Landscape/Private/LandscapeEdit.cpp
+++ b/Engine/Source/Runtime/Landscape/Private/LandscapeEdit.cpp
@@ -3017,11 +3017,14 @@ TMap<FName, ULandscapeLayerInfoObject*> ALandscapeProxy::RetrieveTargetLayerInfo
 	TMap<FName, ULandscapeLayerInfoObject*> InfoObjects;
 	for (const TObjectPtr<ULandscapeComponent>& Component : LandscapeComponents)
 	{
-		for (const FWeightmapLayerAllocationInfo& Allocation : Component->GetWeightmapLayerAllocations())
+		if (Component != nullptr) // can be null in server builds that strip landscape components
 		{
-			if (Allocation.LayerInfo)
+			for (const FWeightmapLayerAllocationInfo& Allocation : Component->GetWeightmapLayerAllocations())
 			{
-				InfoObjects.Add(Allocation.GetLayerName(), Allocation.LayerInfo);
+				if (Allocation.LayerInfo)
+				{
+					InfoObjects.Add(Allocation.GetLayerName(), Allocation.LayerInfo);
+				}
 			}
 		}
 	}
