Landscape batched merge fixes:
* Fixed landscape edit layer renderers being renderered in all batches instead of just the ones they intersect with
* Fixed OOBBox landscape edit layer render item debug display and misc visual logger improvements

diff --git a/Engine/Source/Runtime/Landscape/Private/LandscapeEditLayers.cpp b/Engine/Source/Runtime/Landscape/Private/LandscapeEditLayers.cpp
--- a/Engine/Source/Runtime/Landscape/Private/LandscapeEditLayers.cpp
+++ b/Engine/Source/Runtime/Landscape/Private/LandscapeEditLayers.cpp
@@ -5988,9 +5988,9 @@ namespace UE::Landscape::EditLayers::Private
 	};
 
 	/** Struct that holds all the per-render batch information needed when preparing the batched merge context */
-	struct FComponentToRenderBatchInfo
+	struct FRenderBatchInfo
 	{
-		FComponentToRenderBatchInfo(int32 InNumComponentsToRender, int32 InBatchIndex)
+		FRenderBatchInfo(int32 InNumComponentsToRender, int32 InBatchIndex)
 			: BatchIndex(InBatchIndex)
 			, ComponentToRenderInfoBitIndices(false, InNumComponentsToRender)
 		{
@@ -6075,7 +6075,7 @@ namespace UE::Landscape::EditLayers::Private
 	 *   - When a component is being rendered by a given batch, all components needed for rendering this given component are present in the batch
 	 *   - All components end up being rendered in at least one of the batches
 	 */
-	TArray<FComponentToRenderBatchInfo> DivideIntoBatches(const TBitArray<>& InFinalComponentsToRenderInfoBitIndices, const TArray<FComponentToRenderInfo>& InAllComponentsToRenderInfos)
+	TArray<FRenderBatchInfo> DivideIntoBatches(const TBitArray<>& InFinalComponentsToRenderInfoBitIndices, const TArray<FComponentToRenderInfo>& InAllComponentsToRenderInfos)
 	{
 		TRACE_CPUPROFILER_EVENT_SCOPE(DivideIntoBatches);
 
@@ -6124,7 +6124,7 @@ namespace UE::Landscape::EditLayers::Private
 		});
 
 		// Iterate on all the work items and organize them into batches as large as possible (within the maximum allowed resolution)
-		TArray<FComponentToRenderBatchInfo> AllComponentToRenderBatchInfos;
+		TArray<FRenderBatchInfo> AllBatchInfos;
 		if (!RemainingComponentsToRenderInfoIndices.IsEmpty())
 		{
 			TBitArray<TInlineAllocator<1>> TempBitArray(false, InAllComponentsToRenderInfos.Num());
@@ -6140,17 +6140,17 @@ namespace UE::Landscape::EditLayers::Private
 				int32 MaxNumComponentsInCommonWithBatch = -1;
 
 				// Iterate through all batches and try to find which would be able to accept it and amongst those, which would have the minimal overall resolution: 
-				const int32 NumBatches = AllComponentToRenderBatchInfos.Num();
+				const int32 NumBatches = AllBatchInfos.Num();
 				for (int32 BatchIndex = 0; BatchIndex < NumBatches; ++BatchIndex)
 				{
-					const FComponentToRenderBatchInfo& Batch = AllComponentToRenderBatchInfos[BatchIndex];
-					FIntRect ProjectedBatchSectionRect = Batch.GetProjectedSectionRect(ComponentToRenderInfo);
+					const FRenderBatchInfo& BatchInfo = AllBatchInfos[BatchIndex];
+					FIntRect ProjectedBatchSectionRect = BatchInfo.GetProjectedSectionRect(ComponentToRenderInfo);
 
 					// If after adding this component and its dependent components, the batch still fits within the max allowed resolution, it can accept it : 
 					if ((ProjectedBatchSectionRect.Width() <= MaxBatchResolution.X) && (ProjectedBatchSectionRect.Height() <= MaxBatchResolution.Y))
 					{
 						// Favor the batch that has the most components in common with what we're trying to render : 
-						TempBitArray = TBitArray<>::BitwiseAND(Batch.ComponentToRenderInfoBitIndices, ComponentToRenderInfo.DependentComponentBitIndices, EBitwiseOperatorFlags::MinSize);
+						TempBitArray = TBitArray<>::BitwiseAND(BatchInfo.ComponentToRenderInfoBitIndices, ComponentToRenderInfo.DependentComponentBitIndices, EBitwiseOperatorFlags::MinSize);
 						const int32 NumComponentsInCommonWithBatch = TempBitArray.CountSetBits();
 						// If the batch already has all the components we need, it's a perfect match, we won't ever find a better batch so just stop the search there :
 						if (NumComponentsInCommonWithBatch == ComponentToRenderInfo.DependentComponentBitIndices.CountSetBits())
@@ -6178,14 +6178,14 @@ namespace UE::Landscape::EditLayers::Private
 					}
 				}
 
-				// If we have found a batch, just add the FComponentToRenderBatchInfo to it, otherwise, add a new batch:
-				FComponentToRenderBatchInfo& SelectedBatch = (BestBatchIndex != INDEX_NONE) ? AllComponentToRenderBatchInfos[BestBatchIndex]
-					: AllComponentToRenderBatchInfos.Add_GetRef(FComponentToRenderBatchInfo(InAllComponentsToRenderInfos.Num(), /*InBatchIndex = */AllComponentToRenderBatchInfos.Num()));
+				// If we have found a batch, just add the FRenderBatchInfo to it, otherwise, add a new batch:
+				FRenderBatchInfo& SelectedBatchInfo = (BestBatchIndex != INDEX_NONE) ? AllBatchInfos[BestBatchIndex]
+					: AllBatchInfos.Add_GetRef(FRenderBatchInfo(InAllComponentsToRenderInfos.Num(), /*InBatchIndex = */AllBatchInfos.Num()));
 
-				SelectedBatch.AddToBatch(ComponentToRenderInfo);
+				SelectedBatchInfo.AddToBatch(ComponentToRenderInfo);
 			}
 		}
-		return AllComponentToRenderBatchInfos;
+		return AllBatchInfos;
 	}
 
 #if ENABLE_VISUAL_LOG
@@ -6254,6 +6254,7 @@ namespace UE::Landscape::EditLayers::Private
 		{
 			const FVector SourceCenter = AddNode(InSourceComponentRenderInfo, InSourceRendererRenderInfo);
 			const FVector DestinationCenter = AddNode(InDestinationComponentRenderInfo, InDestinationRendererRenderInfo);
+			// TODO [jonathan.bard] : UE_VLOG_ARROW_MAG(Landscape, LogLandscape, Log, SourceCenter, DestinationCenter, InSourceRendererRenderInfo.VisualLogColor, TEXT(""), 80.0f); // TODO [jonathan.bard] : use proper mag here
 			UE_VLOG_ARROW(Landscape, LogLandscape, Log, SourceCenter, DestinationCenter, InSourceRendererRenderInfo.VisualLogColor, TEXT(""));
 		}
 
@@ -6318,10 +6319,7 @@ namespace UE::Landscape::EditLayers::Private
 		{
 			FOOBox2D OOBox = InInputWorldArea.GetOOBox();
 			FBox VisualBounds(-FVector(OOBox.Extents, 0.0), FVector(OOBox.Extents, 0.0));
-			FTransform OOBoxLocalTransform = BaseTransform.GetRelativeTransform(OOBox.Transform);
-			OOBoxLocalTransform.SetTranslation(FVector::ZeroVector);
-			FTransform FinalTransform = BaseTransform * OOBoxLocalTransform;
-			UE_VLOG_WIREOBOX(InLandscape, LogLandscape, Log, VisualBounds, FinalTransform.ToMatrixWithScale(), InRendererRenderInfo.VisualLogColor, TEXT(""));
+			UE_VLOG_WIREOBOX(InLandscape, LogLandscape, Log, VisualBounds, OOBox.Transform.ToMatrixWithScale(), InRendererRenderInfo.VisualLogColor, TEXT(""));
 			break;
 		}
 		default:
@@ -6371,10 +6369,7 @@ namespace UE::Landscape::EditLayers::Private
 		{
 			FOOBox2D OOBox = InOutputWorldArea.GetOOBox();
 			FBox VisualBounds(-FVector(OOBox.Extents, 0.0), FVector(OOBox.Extents, 0.0));
-			FTransform OOBoxLocalTransform = BaseTransform.GetRelativeTransform(OOBox.Transform);
-			OOBoxLocalTransform.SetTranslation(FVector::ZeroVector);
-			FTransform FinalTransform = BaseTransform * OOBoxLocalTransform;
-			UE_VLOG_OBOX(InLandscape, LogLandscape, Log, VisualBounds, FinalTransform.ToMatrixWithScale(), InRendererRenderInfo.VisualLogColor, TEXT("%s"), *LogMessage);
+			UE_VLOG_OBOX(InLandscape, LogLandscape, Log, VisualBounds, OOBox.Transform.ToMatrixWithScale(), InRendererRenderInfo.VisualLogColor, TEXT("%s"), *LogMessage);
 			break;
 		}
 		default:
@@ -6516,19 +6511,6 @@ UE::Landscape::EditLayers::FMergeRenderContext ALandscape::PrepareEditLayersMerg
 	const FString VisualLogShowRenderItemsEditLayerRendererFilter = CVarLandscapeBatchedMergeVisualLogShowRenderItemsEditLayerRendererFilter.GetValueOnGameThread();
 	const int32 VisualLogShowComponentDependencies = CVarLandscapeBatchedMergeVisualLogShowComponentDependencies.GetValueOnGameThread();
 	const FString VisualLogShowComponentDependenciesFilter = CVarLandscapeBatchedMergeVisualLogShowComponentDependenciesFilter.GetValueOnGameThread();
-	TOptional<FIntPoint> VisualLogShowComponentDependenciesKey;
-	if (!VisualLogShowComponentDependenciesFilter.IsEmpty())
-	{
-		FIntPoint ComponentKey;
-		if (ComponentKey.InitFromString(VisualLogShowComponentDependenciesFilter))
-		{
-			VisualLogShowComponentDependenciesKey = ComponentKey;
-		}
-		else
-		{
-			UE_LOG(LogLandscape, Warning, TEXT("Cannot parse string \"%s\". Ignoring show component dependencies filter"), *VisualLogShowComponentDependenciesFilter);
-		}
-	}
 
 	ULandscapeInfo* Info = GetLandscapeInfo();
 	check(Info != nullptr);
@@ -6753,22 +6735,42 @@ UE::Landscape::EditLayers::FMergeRenderContext ALandscape::PrepareEditLayersMerg
 	TArray<FEditLayerRendererRenderInfo> OrderedEditLayerRendererRenderInfos;
 
 #if ENABLE_VISUAL_LOG
+	int32 VisualLogShowComponentDependenciesIndex = INDEX_NONE;
+	if (!VisualLogShowComponentDependenciesFilter.IsEmpty())
+	{
+		FIntPoint ComponentKey;
+		if (ComponentKey.InitFromString(VisualLogShowComponentDependenciesFilter))
+		{
+			VisualLogShowComponentDependenciesIndex = Component2DIndexer.GetValueIndexForKeySafe(ComponentKey);
+			if (VisualLogShowComponentDependenciesIndex == INDEX_NONE)
+			{
+				UE_LOG(LogLandscape, Warning, TEXT("Component key \"%s\" specified for dependencies filter does not correspond to a valid component. Ignoring show component dependencies filter"), *VisualLogShowComponentDependenciesFilter);
+			}
+		}
+		else
+		{
+			UE_LOG(LogLandscape, Warning, TEXT("Cannot parse string \"%s\". Ignoring show component dependencies filter"), *VisualLogShowComponentDependenciesFilter);
+		}
+	}
+
 	// Helper for debugging component dependencies : only if the CVar requires it
 	TOptional<FComponentDependenciesVisLogHelper> VisualLogDependencyHelper;
 	UE_IFVLOG(
-			if ((MergeRenderContext.IsVisualLogEnabled()) && ((VisualLogShowComponentDependencies > 0) || VisualLogShowComponentDependenciesKey.IsSet()))
+			if ((MergeRenderContext.IsVisualLogEnabled()) && ((VisualLogShowComponentDependencies > 0) || (VisualLogShowComponentDependenciesIndex != INDEX_NONE)))
 			{
 				// Force the display of all info when we show the dependencies of one component in particular : 
-				const FComponentDependenciesVisLogHelper::EShowNodeInfo ShowNodeInfo = VisualLogShowComponentDependenciesKey.IsSet() 
+				const FComponentDependenciesVisLogHelper::EShowNodeInfo ShowNodeInfo = (VisualLogShowComponentDependenciesIndex != INDEX_NONE)
 					? FComponentDependenciesVisLogHelper::EShowNodeInfo::Detailed
 					: static_cast<FComponentDependenciesVisLogHelper::EShowNodeInfo>(VisualLogShowComponentDependencies);
 				VisualLogDependencyHelper.Emplace(this, InMergeRenderParams.bIsHeightmapMerge, ShowNodeInfo, MergeRenderContext);
 			}
 		);
-	auto VisLogDependency = [&VisualLogDependencyHelper, &VisualLogShowComponentDependenciesKey, &AllComponentsToRenderInfos, &OrderedEditLayerRendererRenderInfos]
+	auto VisLogDependency = [VisualLogShowComponentDependenciesIndex, &VisualLogDependencyHelper, &AllComponentsToRenderInfos, &OrderedEditLayerRendererRenderInfos]
 		(int32 InSourceComponentIndex, int32 InSourceRendererIndex, int32 InDestinationComponentIndex, int32 InDestinationRendererIndex)
 		{
-			if (VisualLogDependencyHelper.IsSet() && (InSourceRendererIndex >= 0) && VisualLogShowComponentDependenciesKey.IsSet())
+			if (VisualLogDependencyHelper.IsSet() 
+				&& (InSourceRendererIndex >= 0) 
+				&& ((VisualLogShowComponentDependenciesIndex == INDEX_NONE) || (InSourceComponentIndex == VisualLogShowComponentDependenciesIndex)))
 			{
 				VisualLogDependencyHelper->AddDependency(AllComponentsToRenderInfos[InSourceComponentIndex], OrderedEditLayerRendererRenderInfos[InSourceRendererIndex],
 					AllComponentsToRenderInfos[InDestinationComponentIndex], OrderedEditLayerRendererRenderInfos[InDestinationRendererIndex]);
@@ -7045,19 +7047,19 @@ UE::Landscape::EditLayers::FMergeRenderContext ALandscape::PrepareEditLayersMerg
 #endif // ENABLE_VISUAL_LOG
 
 	// Now divide the work into batches as large as possible (but fitting in the desired max batch resolution, if possible): 
-	TArray<FComponentToRenderBatchInfo> AllComponentToRenderBatchInfos = DivideIntoBatches(FinalComponentsToRenderInfoBitIndices, AllComponentsToRenderInfos);
+	TArray<FRenderBatchInfo> AllBatchInfos = DivideIntoBatches(FinalComponentsToRenderInfoBitIndices, AllComponentsToRenderInfos);
 	{
 		TRACE_CPUPROFILER_EVENT_SCOPE(PrepareBatches);
-		MergeRenderContext.RenderBatches.Reserve(AllComponentToRenderBatchInfos.Num());
+		MergeRenderContext.RenderBatches.Reserve(AllBatchInfos.Num());
 		MergeRenderContext.TargetLayersToComponents.AddDefaulted(MergeRenderContext.AllTargetLayerNames.Num());
-		for (const FComponentToRenderBatchInfo& ComponentToRenderBatchInfo : AllComponentToRenderBatchInfos)
+		for (const FRenderBatchInfo& BatchInfo : AllBatchInfos)
 		{
 			int32 MergeBatchIndex = MergeRenderContext.RenderBatches.Num();
 			FMergeRenderBatch& MergeRenderBatch = MergeRenderContext.RenderBatches.Emplace_GetRef();
 			MergeRenderBatch.Landscape = this;
-			MergeRenderBatch.SectionRect = ComponentToRenderBatchInfo.CombinedSectionRect;
-			MergeRenderBatch.MinComponentKey = MergeRenderBatch.MinComponentKey.ComponentMin(ComponentToRenderBatchInfo.MinComponentKey);
-			MergeRenderBatch.MaxComponentKey = MergeRenderBatch.MaxComponentKey.ComponentMax(ComponentToRenderBatchInfo.MaxComponentKey);
+			MergeRenderBatch.SectionRect = BatchInfo.CombinedSectionRect;
+			MergeRenderBatch.MinComponentKey = MergeRenderBatch.MinComponentKey.ComponentMin(BatchInfo.MinComponentKey);
+			MergeRenderBatch.MaxComponentKey = MergeRenderBatch.MaxComponentKey.ComponentMax(BatchInfo.MaxComponentKey);
 			// We'll reuse the same merge render targets in order to generate the mips, which include the duplicate borders, so we need to expand the render target's size to accommodate for this : 
 			MergeRenderBatch.Resolution = (MergeRenderBatch.MaxComponentKey - MergeRenderBatch.MinComponentKey + 1) * NumSubsections * (SubsectionSizeQuads + 1);
 			MergeRenderBatch.TargetLayersToComponents.AddDefaulted(MergeRenderContext.AllTargetLayerNames.Num());
@@ -7073,11 +7075,11 @@ UE::Landscape::EditLayers::FMergeRenderContext ALandscape::PrepareEditLayersMerg
 			{
 				const FEditLayerRendererRenderInfo& EditLayerRendererRenderInfo = OrderedEditLayerRendererRenderInfos[EditLayerRendererIndex];
 				// Find out all components that are in common between the renderer's list and the batch's list : only these will need to be rendered in that render step : 
-				TBitArray<> ComponentToRenderInBatchBitIndices = TBitArray<>::BitwiseAND(EditLayerRendererRenderInfo.ComponentToRenderInfoBitIndices, ComponentToRenderBatchInfo.ComponentToRenderInfoBitIndices, EBitwiseOperatorFlags::MinSize);
+				TBitArray<> ComponentToRenderForRendererBitIndices = TBitArray<>::BitwiseAND(EditLayerRendererRenderInfo.ComponentToRenderInfoBitIndices, BatchInfo.ComponentToRenderInfoBitIndices, EBitwiseOperatorFlags::MinSize);
 				TArray<ULandscapeComponent*>& ComponentsToRender = ComponentsForRenderers[EditLayerRendererIndex];
-				ComponentsToRender.Reserve(ComponentToRenderInBatchBitIndices.CountSetBits());
+				ComponentsToRender.Reserve(ComponentToRenderForRendererBitIndices.CountSetBits());
 				// Transform the bit indices back into a proper component list : 
-				for (TConstSetBitIterator ItComponent(ComponentToRenderBatchInfo.ComponentToRenderInfoBitIndices); ItComponent; ++ItComponent)
+				for (TConstSetBitIterator ItComponent(ComponentToRenderForRendererBitIndices); ItComponent; ++ItComponent)
 				{
 					FComponentToRenderInfo& ComponentToRenderInfo = AllComponentsToRenderInfos[ItComponent.GetIndex()];
 					ComponentsToRender.Add(ComponentToRenderInfo.Component);
