Fixed improper backout (35373335), where instead of backing out 29959327 (where I had incorrectly tried to fix the GetHeightmapData function), I ended up breaking GetWeightmapData in the same fashion. Both are now back to what they were wrt "ValidX1/Y1/X2/Y2" computation. This API will probably disappear soon anyway since it's so error-prone.
Also, use ALandscape::CalcComponentIndicesNoOverlap instead of ALandscape::CalcComponentIndicesOverlap in those functions (they were used in some but not all...) in order to prevent from processing additional components when trying the end of the region we're trying to read is located on the east/south border of components

diff --git a/Engine/Source/Runtime/Landscape/Private/LandscapeEditInterface.cpp b/Engine/Source/Runtime/Landscape/Private/LandscapeEditInterface.cpp
--- a/Engine/Source/Runtime/Landscape/Private/LandscapeEditInterface.cpp
+++ b/Engine/Source/Runtime/Landscape/Private/LandscapeEditInterface.cpp
@@ -275,7 +275,7 @@ void FLandscapeEditDataInterface::SetHeightData(int32 X1, int32 Y1, int32 X2, in
 	check(ComponentSizeQuads > 0);
 	// Find component range for this block of data
 	int32 ComponentIndexX1, ComponentIndexY1, ComponentIndexX2, ComponentIndexY2;
-	ALandscape::CalcComponentIndicesOverlap(X1, Y1, X2, Y2, ComponentSizeQuads, ComponentIndexX1, ComponentIndexY1, ComponentIndexX2, ComponentIndexY2);
+	ALandscape::CalcComponentIndicesNoOverlap(X1, Y1, X2, Y2, ComponentSizeQuads, ComponentIndexX1, ComponentIndexY1, ComponentIndexX2, ComponentIndexY2);
 
 	FVector* VertexNormals = nullptr;
 	if (InCalcNormals)
@@ -721,10 +721,10 @@ void FLandscapeEditDataInterface::GetHeightDataTemplFast(const int32 X1, const i
 }
 
 template<typename TData, typename TStoreData, typename FType>
-void FLandscapeEditDataInterface::CalcMissingValues(const int32& X1, const int32& X2, const int32& Y1, const int32& Y2, 
-								   const int32& ComponentIndexX1, const int32& ComponentIndexX2, const int32& ComponentIndexY1, const int32& ComponentIndexY2, 
-								   const int32& ComponentSizeX, const int32& ComponentSizeY, TData* CornerValues, 
-								   TArray<bool>& NoBorderY1, TArray<bool>& NoBorderY2, TArray<bool>& ComponentDataExist, TStoreData& StoreData)
+void FLandscapeEditDataInterface::CalcMissingValues(const int32 X1, const int32 X2, const int32 Y1, const int32 Y2,
+	const int32 ComponentIndexX1, const int32 ComponentIndexX2, const int32 ComponentIndexY1, const int32 ComponentIndexY2,
+	const int32 ComponentSizeX, const int32 ComponentSizeY, TData* CornerValues,
+	TArray<bool>& NoBorderY1, TArray<bool>& NoBorderY2, TArray<bool>& ComponentDataExist, TStoreData& StoreData)
 {
 	bool NoBorderX1 = false, NoBorderX2 = false;
 	// Init data...
@@ -1050,8 +1050,8 @@ void FLandscapeEditDataInterface::GetHeightDataInternal(int32& ValidX1, int32& V
 	ValidX1 = INT_MAX; ValidX2 = INT_MIN; ValidY1 = INT_MAX; ValidY2 = INT_MIN;
 
 	int32 ComponentIndexX1, ComponentIndexY1, ComponentIndexX2, ComponentIndexY2;
-	ALandscape::CalcComponentIndicesOverlap(X1, Y1, X2, Y2, ComponentSizeQuads, ComponentIndexX1, ComponentIndexY1, ComponentIndexX2, ComponentIndexY2);
-
+	ALandscape::CalcComponentIndicesNoOverlap(X1, Y1, X2, Y2, ComponentSizeQuads, ComponentIndexX1, ComponentIndexY1, ComponentIndexX2, ComponentIndexY2);
+	
 	int32 ComponentSizeX = ComponentIndexX2-ComponentIndexX1+1;
 	int32 ComponentSizeY = ComponentIndexY2-ComponentIndexY1+1;
 
@@ -1471,10 +1471,10 @@ void FLandscapeEditDataInterface::GetHeightDataInternal(int32& ValidX1, int32& V
 			ComponentSizeX, ComponentSizeY, CornerValues,
 			NoBorderY1, NoBorderY2, ComponentDataExist, StoreData);
 		// Update valid region
-		ValidX1 = FMath::Min<int32>(X1, ValidX1);
-		ValidX2 = FMath::Max<int32>(X2, ValidX2);
-		ValidY1 = FMath::Min<int32>(Y1, ValidY1);
-		ValidY2 = FMath::Max<int32>(Y2, ValidY2);
+		ValidX1 = FMath::Max<int32>(X1, ValidX1);
+		ValidX2 = FMath::Min<int32>(X2, ValidX2);
+		ValidY1 = FMath::Max<int32>(Y1, ValidY1);
+		ValidY2 = FMath::Min<int32>(Y2, ValidY2);
 	}
 	else
 	{
@@ -4277,10 +4277,10 @@ void FLandscapeEditDataInterface::GetWeightDataTempl(ULandscapeLayerInfoObject*
 			ComponentSizeX, ComponentSizeY, CornerValues,
 			NoBorderY1, NoBorderY2, ComponentDataExist, StoreData );
 		// Update valid region
-		ValidX1 = FMath::Min<int32>(X1, ValidX1);
-		ValidX2 = FMath::Max<int32>(X2, ValidX2);
-		ValidY1 = FMath::Min<int32>(Y1, ValidY1);
-		ValidY2 = FMath::Max<int32>(Y2, ValidY2);
+		ValidX1 = FMath::Max<int32>(X1, ValidX1);
+		ValidX2 = FMath::Min<int32>(X2, ValidX2);
+		ValidY1 = FMath::Max<int32>(Y1, ValidY1);
+		ValidY2 = FMath::Min<int32>(Y2, ValidY2);
 	}
 	else
 	{
@@ -5022,7 +5022,7 @@ void FLandscapeEditDataInterface::GetXYOffsetDataTempl(int32& ValidX1, int32& Va
 	ValidX1 = INT_MAX; ValidX2 = INT_MIN; ValidY1 = INT_MAX; ValidY2 = INT_MIN;
 
 	int32 ComponentIndexX1, ComponentIndexY1, ComponentIndexX2, ComponentIndexY2;
-	ALandscape::CalcComponentIndicesOverlap(X1, Y1, X2, Y2, ComponentSizeQuads, ComponentIndexX1, ComponentIndexY1, ComponentIndexX2, ComponentIndexY2);
+	ALandscape::CalcComponentIndicesNoOverlap(X1, Y1, X2, Y2, ComponentSizeQuads, ComponentIndexX1, ComponentIndexY1, ComponentIndexX2, ComponentIndexY2);
 
 	int32 ComponentSizeX = ComponentIndexX2 - ComponentIndexX1 + 1;
 	int32 ComponentSizeY = ComponentIndexY2 - ComponentIndexY1 + 1;
diff --git a/Engine/Source/Runtime/Landscape/Public/LandscapeEdit.h b/Engine/Source/Runtime/Landscape/Public/LandscapeEdit.h
--- a/Engine/Source/Runtime/Landscape/Public/LandscapeEdit.h
+++ b/Engine/Source/Runtime/Landscape/Public/LandscapeEdit.h
@@ -311,9 +311,9 @@ private:
 
 	// Only for Missing Data interpolation... only internal usage
 	template<typename TData, typename TStoreData, typename FType>
-	FORCEINLINE void CalcMissingValues(const int32& X1, const int32& X2, const int32& Y1, const int32& Y2,
-		const int32& ComponentIndexX1, const int32& ComponentIndexX2, const int32& ComponentIndexY1, const int32& ComponentIndexY2,
-		const int32& ComponentSizeX, const int32& ComponentSizeY, TData* CornerValues,
+	FORCEINLINE void CalcMissingValues(const int32 X1, const int32 X2, const int32 Y1, const int32 Y2,
+		const int32 ComponentIndexX1, const int32 ComponentIndexX2, const int32 ComponentIndexY1, const int32 ComponentIndexY2,
+		const int32 ComponentSizeX, const int32 ComponentSizeY, TData* CornerValues,
 		TArray<bool>& NoBorderY1, TArray<bool>& NoBorderY2, TArray<bool>& ComponentDataExist, TStoreData& StoreData);
 
 	// Generic Height Data access
@@ -362,9 +362,9 @@ struct FLandscapeDoNotDirtyScope
 template<typename T>
 void FLandscapeEditDataInterface::ShrinkData(TArray<T>& Data, int32 OldMinX, int32 OldMinY, int32 OldMaxX, int32 OldMaxY, int32 NewMinX, int32 NewMinY, int32 NewMaxX, int32 NewMaxY)
 {
-	checkSlow(OldMinX <= OldMaxX && OldMinY <= OldMaxY);
-	checkSlow(NewMinX >= OldMinX && NewMaxX <= OldMaxX);
-	checkSlow(NewMinY >= OldMinY && NewMaxY <= OldMaxY);
+	check(OldMinX <= OldMaxX && OldMinY <= OldMaxY);
+	check(NewMinX >= OldMinX && NewMaxX <= OldMaxX);
+	check(NewMinY >= OldMinY && NewMaxY <= OldMaxY);
 
 	if (NewMinX != OldMinX || NewMinY != OldMinY ||
 		NewMaxX != OldMaxX || NewMaxY != OldMaxY)
