Add validation to landscape subsystem assumptions to investigate crashes.

diff --git a/Engine/Source/Runtime/Landscape/Private/LandscapeGrass.cpp b/Engine/Source/Runtime/Landscape/Private/LandscapeGrass.cpp
--- a/Engine/Source/Runtime/Landscape/Private/LandscapeGrass.cpp
+++ b/Engine/Source/Runtime/Landscape/Private/LandscapeGrass.cpp
@@ -2404,11 +2404,12 @@ void ALandscapeProxy::FlushGrassComponents(const TSet<ULandscapeComponent*>* Onl
 			Component->DestroyComponent();
 		}
 
+		check(RootComponent != nullptr);
 		TArray<USceneComponent*> AttachedFoliageComponents = RootComponent->GetAttachChildren().FilterByPredicate(
 			[](USceneComponent* Component)
-		{
-			return Cast<UHierarchicalInstancedStaticMeshComponent>(Component);
-		});
+			{
+				return Cast<UHierarchicalInstancedStaticMeshComponent>(Component);
+			});
 
 		// Destroy any attached but un-owned foliage components
 		for (USceneComponent* Component : AttachedFoliageComponents)
diff --git a/Engine/Source/Runtime/Landscape/Private/LandscapeSubsystem.cpp b/Engine/Source/Runtime/Landscape/Private/LandscapeSubsystem.cpp
--- a/Engine/Source/Runtime/Landscape/Private/LandscapeSubsystem.cpp
+++ b/Engine/Source/Runtime/Landscape/Private/LandscapeSubsystem.cpp
@@ -295,7 +295,8 @@ void ULandscapeSubsystem::RemoveGrassInstances(const TSet<ULandscapeComponent*>*
 	TRACE_CPUPROFILER_EVENT_SCOPE(ULandscapeSubsystem::RemoveGrassInstances);
 	for (TObjectPtr<ALandscapeProxy> ProxyPtr : Proxies)
 		{
-		ALandscapeProxy* Proxy = ProxyPtr.Get();
+			ALandscapeProxy* Proxy = ProxyPtr.Get();
+			check(IsValid(Proxy)); // Validate our expectation that proxies will have Unregister() called before a Proxy is flagged as garbage.
 			Proxy->FlushGrassComponents(ComponentsToRemoveGrassInstances, /*bFlushGrassMaps = */false);
 		}
 	}
