Fixed landscape Select tool that started to be broken after CL 30603716. Landscape tool was requesting a TSF_G8 texture but was served a BGRA8 one because CompressionSettings were set to TC_Default and even though CompressionNone was set to true, the texture code interpreted this as a request for a 4-channel texture. Using TC_Grayscale allows to revert to the previous behavior

diff --git a/Engine/Source/Runtime/Landscape/Classes/LandscapeProxy.h b/Engine/Source/Runtime/Landscape/Classes/LandscapeProxy.h
--- a/Engine/Source/Runtime/Landscape/Classes/LandscapeProxy.h
+++ b/Engine/Source/Runtime/Landscape/Classes/LandscapeProxy.h
@@ -1491,13 +1491,13 @@ PRAGMA_ENABLE_DEPRECATION_WARNINGS
 	/** @return Current size of bounding rectangle in quads space */
 	LANDSCAPE_API FIntRect GetBoundingRect() const;
 
-	/** Creates a Texture2D for use by this landscape proxy or one of it's components. If OptionalOverrideOuter is not specified, the proxy is used. */
+	/** Creates a Texture2D for use by this landscape proxy or one of its components. If OptionalOverrideOuter is not specified, the proxy is used. */
 	LANDSCAPE_API UTexture2D* CreateLandscapeTexture(int32 InSizeX, int32 InSizeY, TextureGroup InLODGroup, ETextureSourceFormat InFormat, UObject* OptionalOverrideOuter = nullptr, bool bCompress = false, bool bMipChain = true) const;
 
-	/** Creates a Texture2DArray for use by this landscape proxy or one of it's components. If OptionalOverrideOuter is not specified, the proxy is used. */
+	/** Creates a Texture2DArray for use by this landscape proxy or one of its components. If OptionalOverrideOuter is not specified, the proxy is used. */
 	LANDSCAPE_API UTexture2DArray* CreateLandscapeTextureArray(int32 InSizeX, int32 InSizeY, int32 Slices, TextureGroup InLODGroup, ETextureSourceFormat InFormat, UObject* OptionalOverrideOuter = nullptr);
 	
-	/** Creates a Texture2D for use by this landscape proxy or one of it's components for tools .*/ 
+	/** Creates a Texture2D for use by this landscape proxy or one of its components for tools .*/ 
 	LANDSCAPE_API UTexture2D* CreateLandscapeToolTexture(int32 InSizeX, int32 InSizeY, TextureGroup InLODGroup, ETextureSourceFormat InFormat) const;
 
 	/** Creates a LandscapeWeightMapUsage object outered to this proxy. */
diff --git a/Engine/Source/Runtime/Landscape/Private/LandscapeEdit.cpp b/Engine/Source/Runtime/Landscape/Private/LandscapeEdit.cpp
--- a/Engine/Source/Runtime/Landscape/Private/LandscapeEdit.cpp
+++ b/Engine/Source/Runtime/Landscape/Private/LandscapeEdit.cpp
@@ -8084,6 +8084,8 @@ UTexture2D* ALandscapeProxy::CreateLandscapeToolTexture(int32 InSizeX, int32 InS
 	NewTexture->Source.Init(InSizeX, InSizeY, 1, 1, InFormat, ZeroBuffer);
 	FMemory::Free(ZeroBuffer);
 
+	// Note : for TSF_G8, use TC_Grayscale even though CompressionNone is set to true, otherwise, the texture ends up being created as BGRA instead 
+	NewTexture->CompressionSettings = (InFormat == TSF_G8) ? TC_Grayscale : TC_Default;
 	NewTexture->SRGB = false;
 	NewTexture->CompressionNone = true;
 	NewTexture->MipGenSettings = TMGS_NoMipmaps;
diff --git a/Engine/Source/Runtime/Landscape/Private/LandscapeEditInterface.cpp b/Engine/Source/Runtime/Landscape/Private/LandscapeEditInterface.cpp
--- a/Engine/Source/Runtime/Landscape/Private/LandscapeEditInterface.cpp
+++ b/Engine/Source/Runtime/Landscape/Private/LandscapeEditInterface.cpp
@@ -4574,7 +4574,7 @@ void FLandscapeEditDataInterface::SetSelectData(int32 X1, int32 Y1, int32 X2, in
 		check(Component);
 		return static_cast<UTexture2D*&>(Component->EditToolRenderData.DataTexture);
 	};
-	SetEditToolTextureData(X1, Y1, X2, Y2, Data, Stride, ReturnComponentTexture);
+	SetEditToolTextureData(X1, Y1, X2, Y2, Data, Stride, ReturnComponentTexture, TEXTUREGROUP_8BitData);
 }
 
 template<typename TStoreData>
@@ -4612,7 +4612,7 @@ void FLandscapeEditDataInterface::SetLayerContributionData(int32 X1, int32 Y1, i
 		check(Component);
 		return static_cast<UTexture2D*&>(Component->EditToolRenderData.LayerContributionTexture);
 	};
-	SetEditToolTextureData(X1, Y1, X2, Y2, Data, Stride, ReturnComponentTexture);
+	SetEditToolTextureData(X1, Y1, X2, Y2, Data, Stride, ReturnComponentTexture, TEXTUREGROUP_8BitData);
 }
 
 template<typename TStoreData>
