Fixed landscape.ShowDirty/TrackDirty, which was iterating over vertices without taking into account (subsection) duplicate vertices

diff --git a/Engine/Source/Runtime/Landscape/Private/LandscapeEditLayers.cpp b/Engine/Source/Runtime/Landscape/Private/LandscapeEditLayers.cpp
--- a/Engine/Source/Runtime/Landscape/Private/LandscapeEditLayers.cpp
+++ b/Engine/Source/Runtime/Landscape/Private/LandscapeEditLayers.cpp
@@ -7827,7 +7827,7 @@ void ALandscape::ClearDirtyData(ULandscapeComponent* InLandscapeComponent)
 	const int32 X2 = X1 + ComponentSizeQuads;
 	const int32 Y1 = InLandscapeComponent->GetSectionBase().Y;
 	const int32 Y2 = Y1 + ComponentSizeQuads;
-	const int32 ComponentWidth = (SubsectionSizeQuads + 1) * NumSubsections;
+	const int32 ComponentWidth = ComponentSizeQuads + 1;
 	const int32 DirtyDataSize = ComponentWidth * ComponentWidth;
 	TUniquePtr<uint8[]> DirtyData = MakeUnique<uint8[]>(DirtyDataSize);
 	FMemory::Memzero(DirtyData.Get(), DirtyDataSize);
@@ -7843,23 +7843,30 @@ void ALandscape::UpdateWeightDirtyData(ULandscapeComponent* InLandscapeComponent
 	const int32 X2 = X1 + ComponentSizeQuads;
 	const int32 Y1 = InLandscapeComponent->GetSectionBase().Y;
 	const int32 Y2 = Y1 + ComponentSizeQuads;
-	const int32 ComponentWidth = (SubsectionSizeQuads + 1) * NumSubsections;
+	const int32 ComponentWidth = ComponentSizeQuads + 1;
 	const int32 DirtyDataSize = ComponentWidth * ComponentWidth;
-	TUniquePtr<uint8[]> DirtyData = MakeUnique<uint8[]>(DirtyDataSize);
 	const int32 SizeU = InWeightmap->Source.GetSizeX();
 	const int32 SizeV = InWeightmap->Source.GetSizeY();
-	check(DirtyDataSize == SizeU * SizeV);
-
 	const uint8 DirtyWeight = 1 << 1;
+
+	TUniquePtr<uint8[]> DirtyData = MakeUnique<uint8[]>(DirtyDataSize);
 	LandscapeEdit.GetDirtyData(X1, Y1, X2, Y2, DirtyData.Get(), 0);
 
-	for (int32 Index = 0; Index < DirtyDataSize; ++Index)
+	// COMMENT [jonathan.bard] : this isn't quite working, because of weightmap re-assignment during painting, which can lead to InOldData to be totally different than the previous frame, which 
+	//  will mark pretty much everything as dirty. This will be this way until we stop using weightmap sharing in the tool
+	FLandscapeComponentDataInterface CDI(InLandscapeComponent);
+	for (int32 X = 0; X < ComponentWidth; ++X)
 	{
-		uint8* OldChannelValue = (uint8*)&InOldData[Index] + ChannelOffsets[InChannel];
-		uint8* NewChannelValue = (uint8*)&InNewData[Index] + ChannelOffsets[InChannel];
-		if (*OldChannelValue != *NewChannelValue)
+		for (int32 Y = 0; Y < ComponentWidth; ++Y)
 		{
-			DirtyData[Index] |= DirtyWeight;
+			int32 TexX, TexY;
+			CDI.VertexXYToTexelXY(X, Y, TexX, TexY);
+			int32 TexIndex = TexX + TexY * SizeU;
+			check(TexIndex < SizeU * SizeV);
+			if (InOldData[TexIndex] != InNewData[TexIndex])
+			{
+				DirtyData[X + Y * ComponentWidth] |= DirtyWeight;
+			}
 		}
 	}
 
@@ -7988,7 +7995,7 @@ void ALandscape::UpdateHeightDirtyData(ULandscapeComponent* InLandscapeComponent
 	const int32 X2 = X1 + ComponentSizeQuads;
 	const int32 Y1 = InLandscapeComponent->GetSectionBase().Y;
 	const int32 Y2 = Y1 + ComponentSizeQuads;
-	const int32 ComponentWidth = (SubsectionSizeQuads + 1) * NumSubsections;
+	const int32 ComponentWidth = ComponentSizeQuads + 1;
 	const int32 DirtyDataSize = ComponentWidth * ComponentWidth;
 	TUniquePtr<uint8[]> DirtyData = MakeUnique<uint8[]>(DirtyDataSize);
 	const int32 SizeU = InHeightmap->Source.GetSizeX();
@@ -7998,14 +8005,17 @@ void ALandscape::UpdateHeightDirtyData(ULandscapeComponent* InLandscapeComponent
 	const uint8 DirtyHeight = 1 << 0;
 	LandscapeEdit.GetDirtyData(X1, Y1, X2, Y2, DirtyData.Get(), 0);
 
+	FLandscapeComponentDataInterface CDI(InLandscapeComponent);
 	for (int32 X = 0; X < ComponentWidth; ++X)
 	{
 		for (int32 Y = 0; Y < ComponentWidth; ++Y)
 		{
-			int32 TexX = HeightmapOffsetX + X;
-			int32 TexY = HeightmapOffsetY + Y;
+			int32 TexX, TexY;
+			CDI.VertexXYToTexelXY(X, Y, TexX, TexY);
+			TexX += HeightmapOffsetX;
+			TexY += HeightmapOffsetY;
 			int32 TexIndex = TexX + TexY * SizeU;
-			check(TexIndex < SizeU* SizeV);
+			check(TexIndex < SizeU * SizeV);
 			if (InOldData[TexIndex] != InNewData[TexIndex])
 			{
 				DirtyData[X + Y * ComponentWidth] |= DirtyHeight;
