Removed editor-only data from landscape/water classes in order to avoid issues with missing data in cooked editor.
Note that the associated textures have been changed to lazy-load on use (i.e. in the editor only, when the landscape brushes run) so that this doesn't affect the data loaded at runtime

diff --git a/Engine/Plugins/Experimental/Water/Source/Editor/Private/WaterBrushManager.cpp b/Engine/Plugins/Experimental/Water/Source/Editor/Private/WaterBrushManager.cpp
--- a/Engine/Plugins/Experimental/Water/Source/Editor/Private/WaterBrushManager.cpp
+++ b/Engine/Plugins/Experimental/Water/Source/Editor/Private/WaterBrushManager.cpp
@@ -924,7 +924,7 @@ void AWaterBrushManager::DisplacementSettings(const FBrushActorRenderContext& Br
 	float BrushTexHeightValue = DisableBrushTextureEffects ? 0.0f : Displacement.DisplacementHeight;
 	BrushActorRenderContext.MID->SetScalarParameterValue(FName(TEXT("BrushTexHeight")), BrushTexHeightValue);
 	BrushActorRenderContext.MID->SetScalarParameterValue(FName(TEXT("T")), Displacement.DisplacementTiling);
-	BrushActorRenderContext.MID->SetTextureParameterValue(FName(TEXT("BrushRT")), Displacement.Texture);
+	BrushActorRenderContext.MID->SetTextureParameterValue(FName(TEXT("BrushRT")), Displacement.Texture.Get());
 	BrushActorRenderContext.MID->SetScalarParameterValue(FName(TEXT("Displacement Midpoint")), Displacement.Midpoint);
 	BrushActorRenderContext.MID->SetVectorParameterValue(FName(TEXT("DisplacementChannel")), Displacement.Channel);
 }
@@ -941,7 +941,7 @@ void AWaterBrushManager::ApplyWeightmapSettings(const FBrushRenderContext& Brush
 	float EdgeOffsetValue = WMSettings.EdgeOffset - (SplineMeshExtension / 2.0f);
 	WeightmapMID->SetScalarParameterValue(FName(TEXT("EdgeOffset")), EdgeOffsetValue);
 
-	WeightmapMID->SetTextureParameterValue(FName(TEXT("BrushRT")), WMSettings.ModulationTexture);
+	WeightmapMID->SetTextureParameterValue(FName(TEXT("BrushRT")), WMSettings.ModulationTexture.Get());
 	WeightmapMID->SetScalarParameterValue(FName(TEXT("T")), WMSettings.TextureTiling);
 
 	float WeightmapInfluenceValue = WMSettings.TextureInfluence * (float)!DisableBrushTextureEffects;
diff --git a/Engine/Plugins/Experimental/Water/Source/Runtime/Private/WaterBodyActor.cpp b/Engine/Plugins/Experimental/Water/Source/Runtime/Private/WaterBodyActor.cpp
--- a/Engine/Plugins/Experimental/Water/Source/Runtime/Private/WaterBodyActor.cpp
+++ b/Engine/Plugins/Experimental/Water/Source/Runtime/Private/WaterBodyActor.cpp
@@ -539,8 +539,30 @@ void AWaterBody::DeprecateData()
 
 	if (GetLinkerCustomVersion(FWaterCustomVersion::GUID) < FWaterCustomVersion::MoveTerrainCarvingSettingsToWater)
 	{
-		static_assert(sizeof(WaterHeightmapSettings_DEPRECATED) == sizeof(TerrainCarvingSettings_DEPRECATED), "Both old and old water heightmap settings struct should be exactly similar");
-		FMemory::Memcpy((void*)&WaterHeightmapSettings_DEPRECATED, (void*)&TerrainCarvingSettings_DEPRECATED, sizeof(WaterHeightmapSettings_DEPRECATED));
+		// This horrific piece of code is for salvaging data from the old struct (TerrainCarvingSettings_DEPRECATED, of type FLandmassTerrainCarvingSettings), to the 
+		//  "new" struct (WaterHeightmapSettings_DEPRECATED, of type FWaterBodyHeightmapSettings) (which is also deprecated, mind you, since it has moved into UWaterBodyComponent). 
+		//  A straight-up Memcpy used to be enough, but now the Effects.Displacement.Texture has gone from TObjectPtr to TSoftObjectPtr, so we cannot do this anymore and we have to copy 
+		//  each sub-struct one by one... The code is still there for perfect backwards-compatibility, but should not ever be exercised, since FLandmassTerrainCarvingSettings was only used 
+		//  before Water became public... 
+#define COPY_SIMILAR_SETTINGS(MemberName) \
+		static_assert(sizeof(WaterHeightmapSettings_DEPRECATED.MemberName) == sizeof(TerrainCarvingSettings_DEPRECATED.MemberName), "Both old and new settings struct should be exactly similar"); \
+		FMemory::Memcpy((void*)&WaterHeightmapSettings_DEPRECATED.MemberName, (void*)&TerrainCarvingSettings_DEPRECATED.MemberName, sizeof(WaterHeightmapSettings_DEPRECATED.MemberName));
+
+		COPY_SIMILAR_SETTINGS(BlendMode);
+		WaterHeightmapSettings_DEPRECATED.bInvertShape = TerrainCarvingSettings_DEPRECATED.bInvertShape;
+		COPY_SIMILAR_SETTINGS(FalloffSettings);
+		COPY_SIMILAR_SETTINGS(Effects.Blurring);
+		COPY_SIMILAR_SETTINGS(Effects.CurlNoise);
+		// Copy FWaterBrushEffects from FLandmassBrushEffectsList by hand because it is not binary-compatible : 
+		WaterHeightmapSettings_DEPRECATED.Effects.Displacement.DisplacementHeight = TerrainCarvingSettings_DEPRECATED.Effects.Displacement.DisplacementHeight;
+		WaterHeightmapSettings_DEPRECATED.Effects.Displacement.DisplacementTiling = TerrainCarvingSettings_DEPRECATED.Effects.Displacement.DisplacementTiling;
+		WaterHeightmapSettings_DEPRECATED.Effects.Displacement.Texture = TerrainCarvingSettings_DEPRECATED.Effects.Displacement.Texture;
+		WaterHeightmapSettings_DEPRECATED.Effects.Displacement.Midpoint = TerrainCarvingSettings_DEPRECATED.Effects.Displacement.Midpoint;
+		WaterHeightmapSettings_DEPRECATED.Effects.Displacement.Channel = TerrainCarvingSettings_DEPRECATED.Effects.Displacement.Channel;
+		WaterHeightmapSettings_DEPRECATED.Effects.Displacement.WeightmapInfluence = TerrainCarvingSettings_DEPRECATED.Effects.Displacement.WeightmapInfluence;
+		COPY_SIMILAR_SETTINGS(Effects.SmoothBlending);
+		COPY_SIMILAR_SETTINGS(Effects.Terracing);
+		WaterHeightmapSettings_DEPRECATED.Priority_DEPRECATED = TerrainCarvingSettings_DEPRECATED.Priority;
 	}
 
 	if (GIsEditor && !HasAnyFlags(RF_ClassDefaultObject | RF_ArchetypeObject | RF_DefaultSubObject))
diff --git a/Engine/Plugins/Experimental/Water/Source/Runtime/Private/WaterBodyComponent.cpp b/Engine/Plugins/Experimental/Water/Source/Runtime/Private/WaterBodyComponent.cpp
--- a/Engine/Plugins/Experimental/Water/Source/Runtime/Private/WaterBodyComponent.cpp
+++ b/Engine/Plugins/Experimental/Water/Source/Runtime/Private/WaterBodyComponent.cpp
@@ -259,17 +259,18 @@ ETextureRenderTargetFormat UWaterBodyComponent::GetBrushRenderTargetFormat() con
 
 void UWaterBodyComponent::GetBrushRenderDependencies(TSet<UObject*>& OutDependencies) const
 {
+	// Lazy-load the referenced textures : this allows the settings to not be editor-only but to only actually load data if GetBrushRenderDependencies is called (which is : only in the editor)
 	for (const TPair<FName, FWaterBodyWeightmapSettings>& Pair : LayerWeightmapSettings)
 	{
-		if (Pair.Value.ModulationTexture)
+		if (UTexture2D* Texture = Pair.Value.ModulationTexture.LoadSynchronous())
 		{
-			OutDependencies.Add(Pair.Value.ModulationTexture);
+			OutDependencies.Add(Texture);
 		}
 	}
 
-	if (WaterHeightmapSettings.Effects.Displacement.Texture)
+	if (UTexture2D* Texture = WaterHeightmapSettings.Effects.Displacement.Texture.LoadSynchronous())
 	{
-		OutDependencies.Add(WaterHeightmapSettings.Effects.Displacement.Texture);
+		OutDependencies.Add(Texture);
 	}
 }
 #endif //WITH_EDITOR
diff --git a/Engine/Plugins/Experimental/Water/Source/Runtime/Private/WaterBodyIslandActor.cpp b/Engine/Plugins/Experimental/Water/Source/Runtime/Private/WaterBodyIslandActor.cpp
--- a/Engine/Plugins/Experimental/Water/Source/Runtime/Private/WaterBodyIslandActor.cpp
+++ b/Engine/Plugins/Experimental/Water/Source/Runtime/Private/WaterBodyIslandActor.cpp
@@ -56,17 +56,18 @@ ETextureRenderTargetFormat AWaterBodyIsland::GetBrushRenderTargetFormat() const
 
 void AWaterBodyIsland::GetBrushRenderDependencies(TSet<UObject*>& OutDependencies) const 
 {
+	// Lazy-load the referenced textures : this allows the settings to not be editor-only but to only actually load data if GetBrushRenderDependencies is called (which is : only in the editor)
 	for (const TPair<FName, FWaterBodyWeightmapSettings>& Pair : WaterWeightmapSettings)
 	{
-		if (Pair.Value.ModulationTexture)
+		if (UTexture2D* Texture = Pair.Value.ModulationTexture.LoadSynchronous())
 		{
-			OutDependencies.Add(Pair.Value.ModulationTexture);
+			OutDependencies.Add(Texture);
 		}
 	}
 
-	if (WaterHeightmapSettings.Effects.Displacement.Texture)
+	if (UTexture2D* Texture = WaterHeightmapSettings.Effects.Displacement.Texture.LoadSynchronous())
 	{
-		OutDependencies.Add(WaterHeightmapSettings.Effects.Displacement.Texture);
+		OutDependencies.Add(Texture);
 	}
 }
 #endif //WITH_EDITOR
diff --git a/Engine/Plugins/Experimental/Water/Source/Runtime/Public/WaterBodyComponent.h b/Engine/Plugins/Experimental/Water/Source/Runtime/Public/WaterBodyComponent.h
--- a/Engine/Plugins/Experimental/Water/Source/Runtime/Public/WaterBodyComponent.h
+++ b/Engine/Plugins/Experimental/Water/Source/Runtime/Public/WaterBodyComponent.h
@@ -597,10 +597,8 @@ public:
 	UPROPERTY(Category = Terrain, EditAnywhere, BlueprintReadWrite, meta = (EditCondition = "bAffectsLandscape"))
 	FWaterBodyHeightmapSettings WaterHeightmapSettings;
 
-#if WITH_EDITORONLY_DATA
 	UPROPERTY(Category = Terrain, EditAnywhere, BlueprintReadWrite, meta = (EditCondition = "bAffectsLandscape"))
 	TMap<FName, FWaterBodyWeightmapSettings> LayerWeightmapSettings;
-#endif
 
 	UPROPERTY(EditAnywhere, BlueprintReadWrite, AdvancedDisplay, Category = Rendering)
 	float ShapeDilation = 4096.0f;
diff --git a/Engine/Plugins/Experimental/Water/Source/Runtime/Public/WaterBodyWeightmapSettings.h b/Engine/Plugins/Experimental/Water/Source/Runtime/Public/WaterBodyWeightmapSettings.h
--- a/Engine/Plugins/Experimental/Water/Source/Runtime/Public/WaterBodyWeightmapSettings.h
+++ b/Engine/Plugins/Experimental/Water/Source/Runtime/Public/WaterBodyWeightmapSettings.h
@@ -29,7 +29,7 @@ struct FWaterBodyWeightmapSettings
 	float EdgeOffset;
 
 	UPROPERTY(EditAnywhere, BlueprintReadWrite, Category = WaterBodyWeightmapSettings)
-	TObjectPtr<UTexture2D> ModulationTexture;
+	TSoftObjectPtr<UTexture2D> ModulationTexture;
 
 	UPROPERTY(EditAnywhere, BlueprintReadWrite, Category = WaterBodyWeightmapSettings)
 	float TextureTiling;
diff --git a/Engine/Plugins/Experimental/Water/Source/Runtime/Public/WaterBrushEffects.h b/Engine/Plugins/Experimental/Water/Source/Runtime/Public/WaterBrushEffects.h
--- a/Engine/Plugins/Experimental/Water/Source/Runtime/Public/WaterBrushEffects.h
+++ b/Engine/Plugins/Experimental/Water/Source/Runtime/Public/WaterBrushEffects.h
@@ -101,7 +101,7 @@ struct FWaterBrushEffectDisplacement
 	float DisplacementTiling;
 
 	UPROPERTY(EditAnywhere, BlueprintReadWrite, Category = BrushEffects)
-	TObjectPtr<UTexture2D> Texture;
+	TSoftObjectPtr<UTexture2D> Texture;
 
 	UPROPERTY(EditAnywhere, BlueprintReadWrite, Category = BrushEffects)
 	float Midpoint;
diff --git a/Engine/Source/Runtime/Landscape/Private/LandscapeBlueprintBrushBase.cpp b/Engine/Source/Runtime/Landscape/Private/LandscapeBlueprintBrushBase.cpp
--- a/Engine/Source/Runtime/Landscape/Private/LandscapeBlueprintBrushBase.cpp
+++ b/Engine/Source/Runtime/Landscape/Private/LandscapeBlueprintBrushBase.cpp
@@ -47,12 +47,12 @@ FLandscapeBrushParameters::FLandscapeBrushParameters(bool bInIsHeightmapMerge, c
 // ----------------------------------------------------------------------------------
 
 ALandscapeBlueprintBrushBase::ALandscapeBlueprintBrushBase(const FObjectInitializer& ObjectInitializer)
-#if WITH_EDITORONLY_DATA
-	: OwningLandscape(nullptr)
-	, UpdateOnPropertyChange(true)
+	: UpdateOnPropertyChange(true)
 	, AffectHeightmap(false)
 	, AffectWeightmap(false)
 	, AffectVisibilityLayer(false)
+#if WITH_EDITORONLY_DATA
+	, OwningLandscape(nullptr)
 	, bIsVisible(true)
 	, LastRequestLayersContentUpdateFrameNumber(InvalidLastRequestLayersContentUpdateFrameNumber)
 #endif // WITH_EDITORONLY_DATA
@@ -174,6 +174,51 @@ void ALandscapeBlueprintBrushBase::RequestLandscapeUpdate(bool bInUserTriggered)
 #endif // WITH_EDITOR
 }
 
+void ALandscapeBlueprintBrushBase::SetCanAffectHeightmap(bool bInCanAffectHeightmap)
+{
+#if WITH_EDITORONLY_DATA
+	if (bInCanAffectHeightmap != AffectHeightmap)
+	{
+		Modify();
+		AffectHeightmap = bInCanAffectHeightmap;
+		if (OwningLandscape)
+		{
+			OwningLandscape->OnBlueprintBrushChanged();
+		}
+	}
+#endif // WITH_EDITORONLY_DATA
+}
+
+void ALandscapeBlueprintBrushBase::SetCanAffectWeightmap(bool bInCanAffectWeightmap)
+{
+#if WITH_EDITORONLY_DATA
+	if (bInCanAffectWeightmap != AffectWeightmap)
+	{
+		Modify();
+		AffectWeightmap = bInCanAffectWeightmap;
+		if (OwningLandscape)
+		{
+			OwningLandscape->OnBlueprintBrushChanged();
+		}
+	}
+#endif // WITH_EDITORONLY_DATA
+}
+
+void ALandscapeBlueprintBrushBase::SetCanAffectVisibilityLayer(bool bInCanAffectVisibilityLayer)
+{
+#if WITH_EDITORONLY_DATA
+	if (bInCanAffectVisibilityLayer != AffectVisibilityLayer)
+	{
+		Modify();
+		AffectVisibilityLayer = bInCanAffectVisibilityLayer;
+		if (OwningLandscape)
+		{
+			OwningLandscape->OnBlueprintBrushChanged();
+		}
+	}
+#endif // WITH_EDITORONLY_DATA
+}
+
 #if WITH_EDITOR
 
 void ALandscapeBlueprintBrushBase::GetRendererStateInfo(const ULandscapeInfo* InLandscapeInfo,
@@ -433,54 +478,7 @@ void ALandscapeBlueprintBrushBase::SetIsVisible(bool bInIsVisible)
 		OwningLandscape->OnBlueprintBrushChanged();
 	}
 }
-#endif // WITH_EDITOR
-
-void ALandscapeBlueprintBrushBase::SetCanAffectHeightmap(bool bInCanAffectHeightmap)
-{
-#if WITH_EDITOR
-	if (bInCanAffectHeightmap != AffectHeightmap)
-	{
-		Modify();
-		AffectHeightmap = bInCanAffectHeightmap;
-		if (OwningLandscape)
-		{
-			OwningLandscape->OnBlueprintBrushChanged();
-		}
-	}
-#endif // WITH_EDITOR
-}
 
-void ALandscapeBlueprintBrushBase::SetCanAffectWeightmap(bool bInCanAffectWeightmap)
-{
-#if WITH_EDITOR
-	if (bInCanAffectWeightmap != AffectWeightmap)
-	{
-		Modify();
-		AffectWeightmap = bInCanAffectWeightmap;
-		if (OwningLandscape)
-		{
-			OwningLandscape->OnBlueprintBrushChanged();
-		}
-	}
-#endif // WITH_EDITOR
-}
-
-void ALandscapeBlueprintBrushBase::SetCanAffectVisibilityLayer(bool bInCanAffectVisibilityLayer)
-{
-#if WITH_EDITOR
-	if (bInCanAffectVisibilityLayer != AffectVisibilityLayer)
-	{
-		Modify();
-		AffectVisibilityLayer = bInCanAffectVisibilityLayer;
-		if (OwningLandscape)
-		{
-			OwningLandscape->OnBlueprintBrushChanged();
-		}
-	}
-#endif // WITH_EDITOR
-}
-
-#if WITH_EDITOR
 bool ALandscapeBlueprintBrushBase::CanAffectWeightmapLayer(const FName& InLayerName) const
 {
 	if (!CanAffectWeightmap())
diff --git a/Engine/Source/Runtime/Landscape/Public/LandscapeBlueprintBrushBase.h b/Engine/Source/Runtime/Landscape/Public/LandscapeBlueprintBrushBase.h
--- a/Engine/Source/Runtime/Landscape/Public/LandscapeBlueprintBrushBase.h
+++ b/Engine/Source/Runtime/Landscape/Public/LandscapeBlueprintBrushBase.h
@@ -55,10 +55,6 @@ class ALandscapeBlueprintBrushBase : public AActor
 	GENERATED_UCLASS_BODY()
 
 protected:
-#if WITH_EDITORONLY_DATA
-	UPROPERTY(Transient)
-	TObjectPtr<class ALandscape> OwningLandscape;
-
 	UPROPERTY(Category = "Settings", EditAnywhere, BlueprintReadWrite)
 	bool UpdateOnPropertyChange;
 
@@ -74,6 +70,10 @@ protected:
 	UPROPERTY(Category = "Settings", EditAnywhere, BlueprintReadWrite)
 	TArray<FName> AffectedWeightmapLayers;
 
+#if WITH_EDITORONLY_DATA
+	UPROPERTY(Transient)
+	TObjectPtr<class ALandscape> OwningLandscape;
+
 	UPROPERTY(Transient)
 	bool bIsVisible;
 
@@ -84,7 +84,7 @@ protected:
 	FIntPoint CurrentRenderTargetSize = FIntPoint(MAX_int32, MAX_int32);
 
 	bool bCaptureBoundaryNormals = false;	// HACK [chris.tchou] remove once we have a better boundary normal solution
-#endif
+#endif // WITH_EDITORONLY_DATA
 
 public:
 	UFUNCTION(BlueprintNativeEvent, meta = (DeprecatedFunction, DeprecationMessage = "Please use RenderLayer instead."))
@@ -158,5 +158,5 @@ public:
 	virtual bool CanChangeIsSpatiallyLoadedFlag() const override { return false; }
 
 	bool GetCaptureBoundaryNormals() { return bCaptureBoundaryNormals; }	// HACK [chris.tchou] remove once we have a better boundary normal solution
-#endif
+#endif // WITH_EDITOR
 };
