Fix an issue when cooking levels containing multiple instances of the same level in a commandlet.

Previously it would warn about duplicate landscape GUIDs during cook, because it would load the level with the level instances, but not apply the unique instance GUIDs to each landscape.

This change applies the instance GUID when level instances are loaded in a commandlet, but then restores the original GUID on save.

diff --git a/Engine/Source/Runtime/Landscape/Private/Landscape.cpp b/Engine/Source/Runtime/Landscape/Private/Landscape.cpp
--- a/Engine/Source/Runtime/Landscape/Private/Landscape.cpp
+++ b/Engine/Source/Runtime/Landscape/Private/Landscape.cpp
@@ -3759,6 +3759,17 @@ void ALandscapeProxy::PreSave(FObjectPreSaveContext ObjectSaveContext)
 
 void ALandscapeProxy::Serialize(FArchive& Ar)
 {
+	FGuid InstanceLandscapeGuid = this->LandscapeGuid;
+	if (Ar.IsSaving() && Ar.IsPersistent())
+	{
+		// if we're using an instance-modified landscape guid, we need to restore the original before saving to persistent storage
+		// (this can happen when you are cooking a level containing level instances in a commandlet)
+		if ((LandscapeGuid != OriginalLandscapeGuid) && OriginalLandscapeGuid.IsValid())
+		{
+			this->LandscapeGuid = this->OriginalLandscapeGuid;
+		}
+	}
+
 	Super::Serialize(Ar);
 
 	Ar.UsingCustomVersion(FLandscapeCustomVersion::GUID);
@@ -3790,8 +3801,13 @@ PRAGMA_DISABLE_DEPRECATION_WARNINGS
 		}
 PRAGMA_ENABLE_DEPRECATION_WARNINGS
 	}
-
 #endif
+
+	if (Ar.IsSaving() && Ar.IsPersistent())
+	{
+		// restore the instance guid
+		this->LandscapeGuid = InstanceLandscapeGuid;
+	}
 }
 
 void ALandscapeProxy::AddReferencedObjects(UObject* InThis, FReferenceCollector& Collector)
@@ -4003,7 +4019,7 @@ bool ULandscapeInfo::UpdateLayerInfoMap(ALandscapeProxy* Proxy /*= nullptr*/, bo
 void ChangeLandscapeGuidIfObjectIsInstanced(FGuid& InOutGuid, UObject* InObject)
 {
 	// we shouldn't be dealing with any instanced landscapes in these cases, early out
-	if (InObject->IsTemplate() || IsRunningCookCommandlet())
+	if (InObject->IsTemplate())
 	{
 		return;
 	}
