[Landscape]
* Fixed linker error when building editor on linux by adding dll export of static variable

diff --git a/Engine/Source/Runtime/Landscape/Classes/LandscapeProxy.h b/Engine/Source/Runtime/Landscape/Classes/LandscapeProxy.h
--- a/Engine/Source/Runtime/Landscape/Classes/LandscapeProxy.h
+++ b/Engine/Source/Runtime/Landscape/Classes/LandscapeProxy.h
@@ -472,7 +472,7 @@ protected:
 	UPROPERTY(EditAnywhere, Category = Nanite, AdvancedDisplay, meta = (EditCondition = "bEnableNanite", LandscapeInherited))
 	float NaniteMaxEdgeLengthFactor = 16.0f;
 
-	static FOnLandscapeProxyFixupSharedDataDelegate OnLandscapeProxyFixupSharedDataDelegate;
+	LANDSCAPE_API static FOnLandscapeProxyFixupSharedDataDelegate OnLandscapeProxyFixupSharedDataDelegate;
 #endif // WITH_EDITORONLY_DATA
 
 	/** Disable runtime grass data generation.  If disabled, the grass maps will be serialized at cook time. Do not set directly, use ALandscape::SetDisableRuntimeGrassMapGeneration to ensure it is set on all loaded proxies. */
