Fix for crash when trying to make a landscape proxy "modified" on PostLoad when the main landscape actor has not registered to its ULandscapeInfo yet

diff --git a/Engine/Source/Runtime/Landscape/Private/Landscape.cpp b/Engine/Source/Runtime/Landscape/Private/Landscape.cpp
--- a/Engine/Source/Runtime/Landscape/Private/Landscape.cpp
+++ b/Engine/Source/Runtime/Landscape/Private/Landscape.cpp
@@ -4799,7 +4799,9 @@ void ALandscapeProxy::UpgradeSharedProperties(ALandscape* InParentLandscape)
 
 	if (!SynchronizedProperties.IsEmpty())
 	{
-		LandscapeInfo->MarkObjectDirty(/*InObject = */this, /*bInForceResave = */true);
+		// This function may be called from PostLoad, in which case InParentLandscape will be non-null. Pass it along to LandscapeInfo so that if the landscape actor has not registered to the 
+		//  landscape info yet, it can still retrieve it via this direct pointer : 
+		LandscapeInfo->MarkObjectDirty(/*InObject = */this, /*bInForceResave = */true, InParentLandscape);
 
 		if (!CVarSilenceSharedPropertyDeprecationFixup->GetBool())
 		{
