Fix an issue with landscapes getting disconnected from their streaming proxies after loading a level in template/instancing mode, converting to world partition, and then saving as a new level.
Here we flag the OriginalLandscapeGuid as a landscape inherited property, so it is synced to any streaming proxies created by the conversion to world partition.

This was a regression caused by CL 35806854 (previously we did not consider OriginalLandscapeGuid when saving)

diff --git a/Engine/Source/Runtime/Landscape/Classes/LandscapeProxy.h b/Engine/Source/Runtime/Landscape/Classes/LandscapeProxy.h
--- a/Engine/Source/Runtime/Landscape/Classes/LandscapeProxy.h
+++ b/Engine/Source/Runtime/Landscape/Classes/LandscapeProxy.h
@@ -441,8 +441,9 @@ protected:
 	FGuid LandscapeGuid;
 
 	/** 
-	  * The original unmutated LandscapeGuid on the source asset, before instancing modfications.
+	  * The original unmutated LandscapeGuid on the source asset, before instancing modifications.
 	  **/
+	UPROPERTY(Transient, DuplicateTransient, meta = (LandscapeInherited))
 	FGuid OriginalLandscapeGuid;
 
 	/** Use Nanite to render landscape as a mesh on supported platforms. */
