Fix a rare crash when saving landscape with grass, and grass maps are invalidated _while_ synchronously generating them.

The check() being hit is converted into an ensure(), to allow users to continue to save instead of crashing.  Added more debug logging to report more information about the cause of the invalidation.

There is some possibility the grass generation will continue to fail (until the synchronous generation loop times out) but this issue shouldn't stop it from saving anymore.

diff --git a/Engine/Source/Runtime/Landscape/Private/LandscapeGrassMapsBuilder.cpp b/Engine/Source/Runtime/Landscape/Private/LandscapeGrassMapsBuilder.cpp
--- a/Engine/Source/Runtime/Landscape/Private/LandscapeGrassMapsBuilder.cpp
+++ b/Engine/Source/Runtime/Landscape/Private/LandscapeGrassMapsBuilder.cpp
@@ -866,8 +866,28 @@ bool FLandscapeGrassMapsBuilder::BuildGrassMapsNowForComponents(
 			{
 				check(Component->GrassData->HasValidData()); // guaranteed by UpdateTrackedComponents(), as long as we update all of the components
 #if WITH_EDITOR
-				// guaranteed by UpdateTrackedComponents(), as long as we update all of the components
-				check(Component->ComputeGrassMapGenerationHash() == Component->GrassData->GenerationHash);
+				// this should be guaranteed by UpdateTrackedComponents(), as long as we update all of the components
+				// if not, then issue a report to try to gather more information as to what is causing this.
+				uint32 CurrentGrassMapGenHash = Component->ComputeGrassMapGenerationHash();
+				if (CurrentGrassMapGenHash != Component->GrassData->GenerationHash)
+				{
+					UE_LOG(LogGrass, Warning, TEXT("Unexpected content change while generating grass maps synchronously (Component:%s Ticks:%d WorldType:%d OldHash:%x NewHash:%x StateHash:%x NumElems:%d)"),
+						*Component->GetPathName(),
+						State->TickCount,
+						World->WorldType, Component->GrassData->GenerationHash, CurrentGrassMapGenHash, State->GrassMapGenerationHash,
+						Component->GrassData->NumElements);
+
+					UE_LOG(LogGrass, Warning, TEXT("  GrassBuilder State: Pend:%d (Heap:%d) Strm:%d Rend:%d Fetch:%d Pop:%d NR:%d Total:%d"),
+						PendingCount, PendingComponentsHeap.Num(), StreamingCount, RenderingCount, AsyncFetchCount, PopulatedCount, NotReadyCount, ComponentStates.Num());
+
+					check(State->Component == Component);
+
+					// report the error
+					ensure(CurrentGrassMapGenHash == Component->GrassData->GenerationHash);
+
+					// don't count this as an UpToDate component -- the hash mismatch should be picked up by Update and evicted in the next iteration
+					continue;
+				}
 #endif // WITH_EDITOR
 				if (FailedComponents.Contains(Component))
 				{
