Landscape Runtime Edge Fixup to Match Neighbor Normals

This feature matches the normals (and heights) across the border between neighboring landscape components at runtime by dynamically modifying the edges of the heightmap textures to ensure they match.

diff --git a/Engine/Source/Runtime/Landscape/Classes/LandscapeComponent.h b/Engine/Source/Runtime/Landscape/Classes/LandscapeComponent.h
--- a/Engine/Source/Runtime/Landscape/Classes/LandscapeComponent.h
+++ b/Engine/Source/Runtime/Landscape/Classes/LandscapeComponent.h
@@ -40,6 +40,8 @@ struct FLandscapeTextureDataInfo;
 struct FStaticLightingPrimitiveInfo;
 struct FLandscapeEditDataInterface;
 struct FLandscapeMobileRenderData;
+class ULandscapeHeightmapTextureEdgeFixup;
+struct FLandscapeGroup;
 
 //
 // FLandscapeEditToolRenderData
@@ -612,6 +614,15 @@ public:
 	UPROPERTY()
 	FGuid StateId;
 
+	/** When the component has registered a heightmap texture edge fixup,
+	  * a reference to the edge fixup and the group are stored here for easy access
+	  *  (the fixup is also referenced via the heightmap AssetUserDatas,
+	  *  and the group can be accessed via the ULandcapeSubsystem)
+	  */
+	UPROPERTY(Transient, DuplicateTransient, NonTransactional)
+	TObjectPtr<ULandscapeHeightmapTextureEdgeFixup> RegisteredEdgeFixup;
+	FLandscapeGroup* RegisteredLandscapeGroup;
+
 #if WITH_EDITORONLY_DATA
 	/**	Legacy irrelevant lights */
 	UE_DEPRECATED(5.5, "IrrelevantLights is officially deprecated now and will be removed in a future version")
@@ -874,6 +885,26 @@ public:
 	/** If the LandscapeMaterial has changed, updates the GrassTypes array. Returns true if the GrassTypes array was updated. */
 	LANDSCAPE_API bool UpdateGrassTypes(bool bForceUpdate = false);
 
+	/** Installs new or updates existing TextureUserDatas for this component: up to one ULandscapeHeightmapTextureEdgeFixup and
+	  * up to one UTextureMipDataProviderFactory (either ULandscapeTextureStorageProviderFactory or ULandscapeTextureMipEdgeOverrideFactory).
+	  * 
+	  * The edge fixup class ULandscapeHeightmapTextureEdgeFixup is used to store edge snapshots and manage edge blending with neighboring heightmap textures.
+	  * 
+	  * ULandscapeTextureStorageProviderFactory is used to stream custom-compressed heightmap texture mip data,
+	  *   and it supports also applying edge blending from the edge fixup class while streaming its mips.
+	  * 
+	  * If bUseEdgeFixup is enabled, but bUseCompression is not, the ULandscapeTextureMipEdgeOverrideFactory is used to apply the edge blending
+	  *   on top of the standard built-in texture mip streaming.
+	  * 
+	  * This function will ensure an edge fixup class exists if bUseEdgeFixup is true, or it will remove any existing ones otherwise.
+	  * 
+	  * [EDITOR ONLY] It will capture a new snapshot for the edge fixup, when bUpdateSnapshotNow is true.  Ignored in non-editor builds.
+	  * 
+	  * [EDITOR ONLY] It will compress heightmaps to a ULandscapeTextureStorageProviderFactory when bUseCompression is true.  Ignored in non-editor builds.
+	  *    NOTE: It will NOT remove an existing one if false -- applying compression is a one-way street, and should only be used in cooked content.
+	  */
+	ULandscapeHeightmapTextureEdgeFixup* InstallOrUpdateTextureUserDatas(bool bUseEdgeFixup, bool bUseCompression = false, bool bUpdateSnapshotNow = true);
+
 #if WITH_EDITOR
 	/** Deletes a layer from this component if it does not contain data, calling DeleteLayerAllocation. */
 	bool DeleteLayerIfAllZero(const FGuid& InEditLayerGuid, const uint8* const TexDataPtr, int32 TexSize, int32 LayerIdx, bool bShouldDirtyPackage);
diff --git a/Engine/Source/Runtime/Landscape/Classes/LandscapeEdgeFixup.h b/Engine/Source/Runtime/Landscape/Classes/LandscapeEdgeFixup.h
new file mode 100644
--- /dev/null
+++ b/Engine/Source/Runtime/Landscape/Classes/LandscapeEdgeFixup.h
@@ -0,0 +1,304 @@
+// Copyright Epic Games, Inc. All Rights Reserved.
+
+#pragma once
+
+#include "Containers/ArrayView.h"
+#include "Containers/StaticArray.h"
+#include "Misc/EnumRange.h"
+#include "Engine/AssetUserData.h"
+#include "Streaming/TextureMipDataProvider.h"
+#include "LandscapeEdgeFixup.generated.h"
+
+struct FLandscapeGroup;
+class ULandscapeComponent;
+class ULandscapeInfo;
+class UTexture2D;
+
+namespace UE::Landscape
+{
+	// Enumerates directions, for the edges or neighbors of a tile in the the landscape group grid.
+	// When specifically referencing an edge or a neighbor, use EEdgeIndex or ENeighborIndex
+	enum class EDirectionIndex : uint8
+	{
+		Bottom = 0,
+		BottomRight = 1,
+		Right = 2,
+		TopRight = 3,
+		Top = 4,
+		TopLeft = 5,
+		Left = 6,
+		BottomLeft = 7,
+
+		First = 0,
+		Last = 7,
+		Count = 8,
+
+		FirstEdge = 0,
+		LastEdge = 6,
+		EdgeCount = 4,
+
+		FirstCorner = 1,
+		LastCorner = 7,
+		CornerCount = 4,
+	};
+
+	// Specifies a set of edges or neighbors.
+	enum class EDirectionFlags : uint8
+	{
+		Bottom = 0x00000001,
+		BottomRight = 0x00000002,
+		Right = 0x00000004,
+		TopRight = 0x00000008,
+		Top = 0x00000010,
+		TopLeft = 0x00000020,
+		Left = 0x00000040,
+		BottomLeft = 0x00000080,
+
+		None = 0,
+		AllCorners = BottomRight | TopRight | TopLeft | BottomLeft,
+		AllEdges = Bottom | Right | Top | Left,
+		All = AllEdges | AllCorners
+	};
+
+	typedef EDirectionIndex EEdgeIndex;		// EEdgeIndex specifies an edge or corner on the local landscape component.
+	typedef EDirectionFlags EEdgeFlags;		// EEdgeFlags specifies a set of edges and corners on the local landscape component.
+	typedef EDirectionIndex ENeighborIndex;	// ENeighborIndex specifies a neighbor landscape component, relative to a local landscape component.
+	typedef EDirectionFlags ENeighborFlags;	// ENeighborFlags specifies a set of neighboring landscape components.
+
+	struct FHeightmapTexel
+	{
+		static_assert(PLATFORM_LITTLE_ENDIAN, "This code needs to be adapted to big endian");
+
+		union
+		{
+			uint32 Data32;
+			uint8 Data[4];
+			struct 
+			{
+				uint8 NormalX;
+				uint8 HeightL;	// low 8 bits
+				uint8 HeightH;	// high 8 bits
+				uint8 NormalY;
+			};
+		};
+
+		uint32 GetHeight16() const
+		{
+			return (((uint32)HeightH) << 8) + HeightL;
+		}
+
+		void SetHeight16(uint16 NewHeight)
+		{
+			HeightL = NewHeight & 0xff;
+			HeightH = (NewHeight >> 8) & 0xff;
+		}
+
+		void SetNormal(const FVector& NewNormal)
+		{
+			const FVector Normal = NewNormal.GetSafeNormal();
+			// these seem to produce the nearest result to the GPU normal calculation
+			// (matches 'straight up' of 0x7f coming from the GPU)
+			NormalX = static_cast<uint8>(FMath::RoundToInt32(127.49999f + Normal.X * 127.5f));
+			NormalY = static_cast<uint8>(FMath::RoundToInt32(127.49999f + Normal.Y * 127.5f));
+		}
+
+		bool IsSameHeight(const FHeightmapTexel& Other) const
+		{
+			uint32 HeightBitsXOR = (Data32 ^ Other.Data32) & 0x00FFFF00;
+			return HeightBitsXOR == 0;
+		}
+	};
+
+	// states required to perform edge patching
+	struct FNeighborSnapshots
+	{
+		ENeighborFlags ExistingNeighbors;
+		EEdgeFlags EdgesWithAnyModifiedNeighbor;
+		FHeightmapTextureEdgeSnapshot* NeighborSnapshots[8];
+		FHeightmapTextureEdgeSnapshot* LocalSnapshot;
+		TStaticArray<uint32, 8> GPUEdgeHashes;
+	};
+}
+ENUM_RANGE_BY_FIRST_AND_LAST(UE::Landscape::EDirectionIndex, UE::Landscape::EDirectionIndex::First, UE::Landscape::EDirectionIndex::Last);
+ENUM_CLASS_FLAGS(UE::Landscape::EDirectionFlags);
+
+namespace UE::Landscape
+{
+	// Returns the relative offset (in group tile coords) of a neighbor
+ 	FIntPoint GetNeighborRelativePosition(ENeighborIndex NeighborIndex);
+
+	// Returns the set of neighbors that blend with any local edge in EdgeFlags
+	ENeighborFlags EdgesToAffectedNeighbors(EEdgeFlags LocalEdgeFlags);
+
+	// Converts a neighbor index to a neighbor flag (or an edge index to an edge flag)
+	EDirectionFlags ToFlag(EDirectionIndex Index);
+
+	// Returns a debug string describing the neighbor (or edge) direction
+	const FString& GetDirectionString(EDirectionIndex Index);
+}
+
+
+/**
+  * The snapshots contain a copy of the heightmap edge texels (both height and normal info).
+  * It is filled out in editor or at cook time, to make available at runtime for dynamic edge fixup.
+  */
+USTRUCT()
+struct FHeightmapTextureEdgeSnapshot
+{
+	GENERATED_USTRUCT_BODY()
+	friend class ULandscapeHeightmapTextureEdgeFixup;
+
+private:
+	int32 EdgeLength = 0;						// edge length for recorded edge data here - when up to date, should match texture resolution (width AND height)
+	TArray<uint32> EdgeData;					// height and normal data for each edge & mip (in heightmap texture format 32bpp). Use GetEdgeData() to access specific edges and mips.
+	TStaticArray<uint32, 8> SnapshotEdgeHashes;	// hash of each edge / corner (at full resolution) in the EdgeData
+	TStaticArray<uint32, 8> InitialEdgeHashes;	// hash of each edge / corner (at full resolution) in the GPU Texture Resource (at initial unpatched state)
+
+#if WITH_EDITOR
+	FGuid TextureSourceID;						// used to detect when this is out of date with texture sourcesource
+#endif // WITH_EDITOR
+
+public:
+	/** Return edge snapshot data for this component, for the specified neighbor direction and mip.
+	  * Horizontal edges are stored left to right, and vertical edges bottom to top.
+	  */
+	TArrayView<UE::Landscape::FHeightmapTexel> GetEdgeData(UE::Landscape::EEdgeIndex InEdgeIndex, int32 InMipIndex);
+	UE::Landscape::FHeightmapTexel GetCornerData(UE::Landscape::EEdgeIndex InEdgeIndex);
+
+#if WITH_EDITOR
+	// Create a snapshot from the heightmap source
+	static TSharedRef<FHeightmapTextureEdgeSnapshot> CreateEdgeSnapshotFromHeightmapSource(UTexture2D* InHeightmap, const FVector& LandscapeGridScale);
+#endif // WITH_EDITOR
+	// Create a snapshot from an explicit byte array (in heightmap source standard layout)
+	static TSharedRef<FHeightmapTextureEdgeSnapshot> CreateEdgeSnapshotFromTextureData(const TArrayView64<UE::Landscape::FHeightmapTexel>& InHeightmapTextureData, int32 InEdgeLength, const FVector& LandscapeGridScale);
+
+	// Return the set of edges that are different (according to edge hashes) -- i.e. changes that could cause neighbors to patch
+	UE::Landscape::EEdgeFlags CompareEdges(const FHeightmapTextureEdgeSnapshot& OldSnapshot) const;
+
+	friend FArchive& operator<<(FArchive& Ar, FHeightmapTextureEdgeSnapshot& Data);
+
+	FString GetTextureSourceIDAsString()
+	{
+#if WITH_EDITOR
+		return TextureSourceID.ToString();
+#else
+		return FString("<NONE>");
+#endif // WITH_EDITOR
+	}
+
+private:
+	void ResizeForEdgeLength(const int32 InEdgeLength);
+
+#if WITH_EDITOR
+	// These CaptureEdgeData* functions are NOT THREAD SAFE --
+	// they stomp existing data and should only be used on newly allocated FHeightmapTextureEdgeSnapshot
+	void CaptureEdgeDataFromHeightmapSource_Internal(UTexture2D* InHeightmap, const FVector& LandscapeGridScale);
+#endif // WITH_EDITOR
+	void CaptureEdgeDataFromTextureData_Internal(const TArrayView64<UE::Landscape::FHeightmapTexel>& InHeightmapTextureData, int32 InEdgeLength, const FVector& LandscapeGridScale);
+	
+	// must call ResizeForEdgeLength to set the texture size before calling CopyTextureEdgeDataAndRecomputeNormals
+	void CaptureSingleEdgeDataAndComputeNormalsAndHashes(const UE::Landscape::FHeightmapTexel* TextureData, const UE::Landscape::EEdgeIndex EdgeOrCorner, const FVector& LandscapeGridScale);
+};
+
+// This UAssetUserData is attached to landscape heightmap UTexture2D's and tracks the heightmap texture's edge fixup state
+// This is used by mip providers to apply edge fixup on mip streaming/creation, and
+// also used by runtime dynamic fixup when neighboring landscape components are pulled in
+UCLASS(NotBlueprintable, MinimalAPI, Within = Texture2D)
+class ULandscapeHeightmapTextureEdgeFixup : public UAssetUserData
+{
+	GENERATED_UCLASS_BODY()
+	friend struct FLandscapeGroup;
+	friend class FLandscapeTextureMipEdgeOverrideProvider;
+	friend class FLandscapeTextureStorageMipProvider;
+
+private:
+	// SERIALIZED snapshot of the heightmap edge data
+	// COPY-ON-WRITE so we can use it safely from other threads. do not modify an existing snapshot, create a new snapshot and replace this reference.
+	TSharedRef<FHeightmapTextureEdgeSnapshot, ESPMode::ThreadSafe> EdgeSnapshot;
+
+	// transient runtime tracking data
+	TObjectPtr<UTexture2D> HeightmapTexture;			// heightmap texture (set to our parent heightmap, on first registration)
+
+	TObjectPtr<ULandscapeComponent> ActiveComponent;	// the active component, that is patching HeightmapTexture
+	FLandscapeGroup* ActiveGroup = nullptr;				// the active group, that is patching HeightmapTexture
+
+	TStaticArray<uint32, 8> GPUEdgeHashes;				// hash for the current GPU edge state, initialized on first registration
+	UE::Landscape::EEdgeFlags GPUEdgeModifiedFlags =	// edges that have been patched / modified from the initial state
+		UE::Landscape::EEdgeFlags::None;
+
+#if WITH_EDITOR
+	bool bDoNotPatchUntilGPUEdgeHashesUpdated = false;
+	bool bUpdateGPUEdgeHashes = false;
+#endif // WITH_EDITOR
+
+	// per-group settings (apply to the active group/component)
+	bool bMapped = false;
+	FIntPoint GroupCoord = FIntPoint::ZeroValue;		// coordinate of this heightmap in the active group (when bMapped)
+
+	// components that also want to use this heightmap & edge fixup, but were disabled as we can only support one active component at a time.
+	// these must be weak object pointers, as they can be unregistered while still in this list, and can be garbage collected out from under us.
+	// Note that this array is generally empty except in scenarios where there are multiple active worlds sharing the same landscape textures (PIE)
+	// so the expense of accessing the TWeakObjectPtr is a PIE-only cost.
+	TArray<TWeakObjectPtr<ULandscapeComponent>> DisabledComponents;
+
+public:
+	inline bool IsActive() { return ActiveGroup != nullptr; }
+	inline UTexture2D* GetHeightmapTexture() { return HeightmapTexture; }
+	inline bool IsComponentActive(ULandscapeComponent* Component) { return ActiveComponent == Component; }
+	inline const FIntPoint& GetGroupCoord() { return GroupCoord; }
+
+#if WITH_EDITOR
+	inline bool IsTextureEdgePatchingPaused() { return bDoNotPatchUntilGPUEdgeHashesUpdated; }
+	inline void PauseTextureEdgePatchingUntilGPUEdgeHashesUpdated() { bDoNotPatchUntilGPUEdgeHashesUpdated = true; }
+#endif // WITH_EDITOR
+
+	virtual void Serialize(FArchive& Ar) override;
+	static void AddReferencedObjects(UObject* InThis, FReferenceCollector& Collector);
+
+	// [main thread] -- Find or Create an edge fixup class for the given heightmap texture.  Will only create in editor (the data used to create is only available in editor)
+	static ULandscapeHeightmapTextureEdgeFixup* FindOrCreateFor(UTexture2D* TargetTexture);
+
+	// set the heightmap texture (can only be called once)
+	void SetHeightmapTexture(UTexture2D* InHeightmapTexture);
+
+	// set the active landscape component, handling unmapping the old, and mapping the new.
+	// if bDisableCurrentActive, it will move the current active component, if any, to the disabled list.
+	void SetActiveComponent(ULandscapeComponent* InComponent, FLandscapeGroup* InGroup, const bool bDisableCurrentActive = true);
+
+	// Request edge texture patching on a set of neighbors.
+	void RequestEdgeTexturePatchingForNeighbors(UE::Landscape::ENeighborFlags NeighborsNeedingPatching);
+
+#if WITH_EDITOR
+	// Request an update to the edge snapshot (capturing from the heightmap source).
+	// When bUpdateGPUEdgeHashes is true, it will also update the GPU Edge Hashes to match.
+	// This is set when we know the heightmap source exactly reflects the GPU texture state
+	// i.e. after reading back from GPU to CPU after a layer merge, or after UpdateResource().
+	void RequestEdgeSnapshotUpdateFromHeightmapSource(bool bUpdateGPUEdgeHashes = false);
+
+	// Update the edge snapshot from heightmap source.  Returns the set of edges that changed since the previous snapshot.
+	UE::Landscape::EEdgeFlags UpdateEdgeSnapshotFromHeightmapSource(const FVector& LandscapeGridScale);
+#endif // WITH_EDITOR
+
+	// Patch the GPU texture edges if needed, using the current snapshot and corresponding neighbor snapshots as source data.
+	// Returns the number of edges patched.
+	int32 CheckAndPatchTextureEdgesFromEdgeSnapshots();
+
+	// Get the set of neighbor snapshots (nullptr if they don't exist), and gather existence and modified flags
+	// this is the data necessary to perform patching on a component
+	void GetNeighborSnapshots(UE::Landscape::FNeighborSnapshots& OutSnapshots);
+
+	// Patch all of the edges for a single texture mip.
+	// Called during streaming operations to patch a newly streamed mip in flight.
+	// Returns the number of edges and corners patched
+	static int32 PatchTextureEdgesForSingleMip(int32 MipIndex, FTextureMipInfo& DestMipInfo, const UE::Landscape::FNeighborSnapshots& NeighborSnapshots);
+	static int32 PatchTextureEdgesForStreamingMips(int32 FirstMipIndexInclusive, int32 LastMipIndexExclusive, FTextureMipInfoArray& DestMipInfos, const UE::Landscape::FNeighborSnapshots& NeighborSnapshots);
+
+private:
+	void PatchTextureEdge_Internal(UE::Landscape::EEdgeIndex EdgeIndex);
+	void PatchTextureCorner_Internal(UE::Landscape::EEdgeIndex CornerIndex, UE::Landscape::FHeightmapTexel Texel);
+
+	// Helper functions that generate blended edge or corner data from snapshots, to use in texture patching
+	static void BlendEdgeData(FHeightmapTextureEdgeSnapshot& EdgeSnapshot, UE::Landscape::EEdgeIndex EdgeIndex, int32 MipIndex, FHeightmapTextureEdgeSnapshot& NeighborEdgeSnapshot, TStridedView<UE::Landscape::FHeightmapTexel>& OutDestView);
+	static void BlendCornerData(UE::Landscape::FHeightmapTexel& OutTexel, UE::Landscape::EEdgeIndex CornerIndex, const UE::Landscape::FNeighborSnapshots& NeighborSnapshots);
+};
+
diff --git a/Engine/Source/Runtime/Landscape/Classes/LandscapeProxy.h b/Engine/Source/Runtime/Landscape/Classes/LandscapeProxy.h
--- a/Engine/Source/Runtime/Landscape/Classes/LandscapeProxy.h
+++ b/Engine/Source/Runtime/Landscape/Classes/LandscapeProxy.h
@@ -875,20 +875,26 @@ PRAGMA_ENABLE_DEPRECATION_WARNINGS
 	UPROPERTY(Transient, NonTransactional)
 	TMap<TObjectPtr<UTexture2D>, TObjectPtr<ULandscapeWeightmapUsage>> WeightmapUsageMap;
 
-	/** True when this Proxy is registered with the LandscapeInfo */
-	bool bIsRegisteredWithLandscapeInfo = false;
-
-	/** Set to true when on undo, when it's necessary to completely regenerate weightmap usages (since some weightmap allocations are transactional and others not, e.g. splines edit layer) */
-	bool bNeedsWeightmapUsagesUpdate = false;
-
 	/** CurrentVersion is bumped whenever a landscape component has an undo/redo operation applied. This lets us detect when the weightmap fixup needs to be run. */
 	uint32 CurrentVersion = 1;
 	uint32 WeightmapFixupVersion = 0;
 
+	/** True when this Proxy is registered with the LandscapeInfo */
+	bool bIsRegisteredWithLandscapeInfo : 1 = false;
+
+	/** Set to true when on undo, when it's necessary to completely regenerate weightmap usages (since some weightmap allocations are transactional and others not, e.g. splines edit layer) */
+	bool bNeedsWeightmapUsagesUpdate : 1 = false;
+
 	/** Set to true when we know that weightmap usages are being reconstructed and might be temporarily invalid as a result (ValidateProxyLayersWeightmapUsage should be called after setting this back to false) */
-	bool bTemporarilyDisableWeightmapUsagesValidation = false;
+	bool bTemporarilyDisableWeightmapUsagesValidation : 1 = false;
 #endif // WITH_EDITORONLY_DATA
 
+	/** True when this Proxy's LandscapeActor has registered with the LandscapeInfo, (and our shared properties are fully updated) */
+	bool bIsLandscapeActorRegisteredWithLandscapeInfo : 1 = false;
+
+	/** True when this Proxy is registered with the ULandscapeSubsystem */
+	bool bIsRegisteredWithSubsystem : 1 = false;
+
 	/** Data set at creation time */
 	UPROPERTY(meta = (LandscapeInherited))
 	int32 ComponentSizeQuads;    // Total number of quads in each component
diff --git a/Engine/Source/Runtime/Landscape/Classes/LandscapeTextureStorageProvider.h b/Engine/Source/Runtime/Landscape/Classes/LandscapeTextureStorageProvider.h
--- a/Engine/Source/Runtime/Landscape/Classes/LandscapeTextureStorageProvider.h
+++ b/Engine/Source/Runtime/Landscape/Classes/LandscapeTextureStorageProvider.h
@@ -9,9 +9,11 @@
 #include "Algo/Accumulate.h"
 #include "Serialization/BulkData.h"
 #include "Engine/TextureAllMipDataProviderFactory.h"
+#include "LandscapeEdgeFixup.h"
 #include "LandscapeTextureStorageProvider.generated.h"
 
 class ULandscapeTextureStorageProviderFactory;
+class ULandscapeHeightmapTextureEdgeFixup;
 class UTexture2D;
 
 USTRUCT()
@@ -67,6 +69,9 @@ private:
 	// Pending async requests created in GetMips().
 	TArray<FIORequest, TInlineAllocator<MAX_TEXTURE_MIP_COUNT>> IORequests;
 
+	// We make a copy of this, as it is provided in GetMips, and we fill it out in PollMips
+	FTextureMipInfoArray DestMipInfos;
+
 	int32 FirstRequestedMipIndex = INDEX_NONE;
 
 	// The asset name, used to log IO errors.
@@ -89,7 +94,82 @@ private:
 
 
 
-UCLASS(MinimalAPI)		// MinimalAPI -- do not export methods (methods not available in other modules)
+class FLandscapeTextureMipEdgeOverrideProvider : public FTextureMipDataProvider
+{
+	// This class is actually a "Modifier" in that it doesn't provide any mips, it only modifies mips provided by the default providers
+	// It relies on some assumptions about the existing provider execution that are currently true, but may not remain true forever:
+	// 1) that all providers are executed, even if they are not actually handling any mips
+	// 2) that all default providers do NOT modify mip data in PollMips  (this is true for both FTexture2DMipDataProvider_IO and FTexture2DMipDataProvider_DDC currently)
+	//    this allows us to modify mip data in PollMips and not worry about being stomped on.
+	// IF either of these cease to be true, we should probably write a first class FTextureMipDataModifier path, and convert this to be a modifier instead of a provider.
+
+public:
+	FLandscapeTextureMipEdgeOverrideProvider(ULandscapeHeightmapTextureEdgeFixup* InEdgeFixup, UTexture2D* InTexture);
+	~FLandscapeTextureMipEdgeOverrideProvider();
+
+	// ****************************************************************************
+	// ********* FLandscapeTextureMipEdgeOverrideProvider implementation **********
+	// ****************************************************************************
+
+	void Init(const FTextureUpdateContext& Context, const FTextureUpdateSyncOptions& SyncOptions) final override;
+	int32 GetMips(const FTextureUpdateContext& Context, int32 StartingMipIndex, const FTextureMipInfoArray& MipInfos, const FTextureUpdateSyncOptions& SyncOptions) final override;
+	bool PollMips(const FTextureUpdateSyncOptions& SyncOptions) final override;
+	void CleanUp(const FTextureUpdateSyncOptions& SyncOptions) final override;
+	void Cancel(const FTextureUpdateSyncOptions& SyncOptions) final override;
+	ETickThread GetCancelThread() const final override;
+
+private:
+	// the Edge Fixup (that actually has the edge override data)
+	ULandscapeHeightmapTextureEdgeFixup* EdgeFixup = nullptr;
+
+	// We make a copy of this, as it is provided in GetMips, and we fill it out in PollMips
+	FTextureMipInfoArray DestMipInfos;
+
+	// The asset name, used to log IO errors.
+	FName TextureName;
+};
+
+
+// serves as a hook into the texture streaming process, allowing us to inject edge data
+UCLASS(NotBlueprintable, MinimalAPI, Within = Texture2D)
+class ULandscapeTextureMipEdgeOverrideFactory : public UTextureMipDataProviderFactory
+{
+	GENERATED_UCLASS_BODY()
+public:
+	TObjectPtr<UTexture2D> Texture;									// should be the same as GetOuter()
+	TObjectPtr<ULandscapeHeightmapTextureEdgeFixup> EdgeFixup;		// assigned when the EdgeFixup is Active
+
+	static ULandscapeTextureMipEdgeOverrideFactory* AddTo(UTexture2D* TargetTexture);
+
+	void SetupEdgeFixup(ULandscapeHeightmapTextureEdgeFixup* InEdgeFixup);
+
+	virtual void Serialize(FArchive& Ar) override;
+	static void AddReferencedObjects(UObject* InThis, FReferenceCollector& Collector);
+
+	/**
+	 * Create a FTextureMipDataProvider to handle a single StreamIn mip operation.
+	 * The object lifetime will be managed by FRenderAssetUpdate.
+	 *
+	 * @param InAsset - the texture on which the stream in operation will be performed.
+	 */
+	virtual FTextureMipDataProvider* AllocateMipDataProvider(UTexture* Asset) override
+	{
+		check(Asset == (UTexture*) Texture.Get());
+		return new FLandscapeTextureMipEdgeOverrideProvider(EdgeFixup, Texture);
+	}
+
+	/**
+	 * Returns true if TextureMipDataProviders allocated by this factory can provide MipData by themselves,
+	 * even without loading from disk at all, so streaming can be enabled for their textures.
+	 */
+	virtual bool WillProvideMipDataWithoutDisk() const
+	{
+		return true;
+	}
+};
+
+
+UCLASS(NotBlueprintable, MinimalAPI, Within = Texture2D)
 class ULandscapeTextureStorageProviderFactory : public UTextureAllMipDataProviderFactory
 {
 	GENERATED_UCLASS_BODY()
@@ -103,10 +183,16 @@ public:
 	TArray<FLandscapeTexture2DMipMap> Mips;
 	TObjectPtr<UTexture2D> Texture;
 
+	// Reference to the edge fixup on the Texture, if any.  (This can be null, if edge fixup is disabled)
+	// It is set when the Texture's landscape component is registered.
+	TObjectPtr<ULandscapeHeightmapTextureEdgeFixup> EdgeFixup;
+
 	#if WITH_EDITORONLY_DATA
-	static ULandscapeTextureStorageProviderFactory* ApplyTo(UTexture2D* TargetTexture, const FVector& InLandsapeGridScale);
+	static ULandscapeTextureStorageProviderFactory* ApplyTo(UTexture2D* InTargetTexture, const FVector& InLandsapeGridScale);
 	#endif // WITH_EDITORONLY_DATA
 
+	void SetupEdgeFixup(ULandscapeHeightmapTextureEdgeFixup* EdgeFixup);
+
 	// compress the mip data.
 	static void CompressMipToBulkData(int32 MipIndex, int32 MipSizeX, int32 MipSizeY, uint8* SourceData, int32 SourceDataBytes, FByteBulkData& DestBulkData);
 	static void CopyMipToBulkData(int32 MipIndex, int32 MipSizeX, int32 MipSizeY, uint8* SourceData, int32 SourceDataBytes, FByteBulkData& DestBulkData);
@@ -115,6 +201,7 @@ public:
 	void DecompressMip(uint8* SourceData, int64 SourceDataBytes, uint8* DestData, int64 DestDataBytes, int32 MipIndex);
 
 	virtual void Serialize(FArchive& Ar) override;
+	static void AddReferencedObjects(UObject* InThis, FReferenceCollector& Collector);
 
 	/**
 	* Create a FTextureMipDataProvider to handle a single StreamIn mip operation.
diff --git a/Engine/Source/Runtime/Landscape/Private/Landscape.cpp b/Engine/Source/Runtime/Landscape/Private/Landscape.cpp
--- a/Engine/Source/Runtime/Landscape/Private/Landscape.cpp
+++ b/Engine/Source/Runtime/Landscape/Private/Landscape.cpp
@@ -64,6 +64,7 @@ Landscape.cpp: Terrain rendering
 #include "ComponentRecreateRenderStateContext.h"
 #include "LandscapeWeightmapUsage.h"
 #include "LandscapeSubsystem.h"
+#include "LandscapeGroup.h"
 #include "LandscapeGrassMapsBuilder.h"
 #include "LandscapeCulling.h"
 #include "ContentStreaming.h"
@@ -2103,14 +2104,116 @@ static void OnStaticMeshLODDistanceScaleChanged()
 
 FAutoConsoleVariableSink OnStaticMeshLODDistanceScaleChangedSink(FConsoleCommandDelegate::CreateStatic(&OnStaticMeshLODDistanceScaleChanged));
 
+ULandscapeHeightmapTextureEdgeFixup* ULandscapeComponent::InstallOrUpdateTextureUserDatas(bool bUseEdgeFixup, bool bUseCompression, bool bUpdateSnapshotNow)
+{
+	if (HeightmapTexture == nullptr)
+	{
+		UE_LOG(LogLandscape, Warning, TEXT("Tried to install EdgeFixup on component %s (proxy %s), but it had NO heightmap"), *GetPathName(), *GetLandscapeProxy()->GetPathName());
+		return nullptr;
+	}
+
+	// grid scale must be updated to get proper normal calculations on the edges
+	FVector LandscapeGridScale = GetLandscapeProxy()->GetRootComponent()->GetRelativeScale3D();
+
+	// first update or install the heightmap texture edge fixup
+	ULandscapeHeightmapTextureEdgeFixup* EdgeFixup = nullptr;
+	if (bUseEdgeFixup && UE::Landscape::ShouldInstallEdgeFixup())
+	{
+		// find or create edge fixup
+		EdgeFixup = ULandscapeHeightmapTextureEdgeFixup::FindOrCreateFor(HeightmapTexture);
+#if WITH_EDITOR
+		if (bUpdateSnapshotNow)
+		{
+			EdgeFixup->UpdateEdgeSnapshotFromHeightmapSource(LandscapeGridScale);
+		}
+#endif // WITH_EDITOR
+	}
+	else
+	{
+		// remove any existing edge fixup (we will remove the factory references to the EdgeFixup below)
+		HeightmapTexture->RemoveUserDataOfClass(ULandscapeHeightmapTextureEdgeFixup::StaticClass());
+		EdgeFixup = nullptr;
+	}
+
+	// check if the heightmap has a ULandscapeTextureMipEdgeOverrideFactory or a ULandscapeTextureStorageProviderFactory
+	ULandscapeTextureMipEdgeOverrideFactory* OverrideFactory = (ULandscapeTextureMipEdgeOverrideFactory*) HeightmapTexture->GetAssetUserDataOfClass(ULandscapeTextureMipEdgeOverrideFactory::StaticClass());
+	ULandscapeTextureStorageProviderFactory* StorageFactory = (ULandscapeTextureStorageProviderFactory*) HeightmapTexture->GetAssetUserDataOfClass(ULandscapeTextureStorageProviderFactory::StaticClass());
+
+	// we should never have both
+	check(!OverrideFactory || !StorageFactory);
+
+#if WITH_EDITOR
+	if (bUseCompression || StorageFactory)
+#else // !WITH_EDITOR
+	check(bUseCompression == false); // cannot install compression in non-editor builds
+	if (StorageFactory)
+#endif // !WITH_EDITOR
+	{
+		// Remove any existing mip edge override factory
+		if (OverrideFactory)
+		{
+			OverrideFactory->SetupEdgeFixup(nullptr);
+			HeightmapTexture->RemoveUserDataOfClass(ULandscapeTextureMipEdgeOverrideFactory::StaticClass());
+			OverrideFactory = nullptr;
+		}
+
+#if WITH_EDITOR
+		// Install or Update the Storage Factory [editor only]
+		if (StorageFactory == nullptr)
+		{
+			StorageFactory = ULandscapeTextureStorageProviderFactory::ApplyTo(HeightmapTexture, LandscapeGridScale);
+		}
+#endif // WITH_EDITOR
+		StorageFactory->SetupEdgeFixup(EdgeFixup);
+	}
+		
+	if (EdgeFixup != nullptr)
+	{
+		// an edge fixup requires at least one factory -- if storage factory doesn't exist, create an override factory
+		if (StorageFactory == nullptr)
+		{
+			if (OverrideFactory == nullptr)
+			{
+				OverrideFactory = ULandscapeTextureMipEdgeOverrideFactory::AddTo(HeightmapTexture);
+			}
+			OverrideFactory->SetupEdgeFixup(EdgeFixup);
+		}
+	}
+	else
+	{
+		// no edge fixup, override factory not needed -- remove any existing one
+		if (OverrideFactory)
+		{
+			OverrideFactory->SetupEdgeFixup(nullptr);
+			HeightmapTexture->RemoveUserDataOfClass(ULandscapeTextureMipEdgeOverrideFactory::StaticClass());
+			OverrideFactory = nullptr;
+		}
+	}
+
+	// double check we've achieved the desired relationships..
+	if (EdgeFixup)
+	{
+		// EdgeFixup must have exactly one factory
+		check(!!OverrideFactory == !StorageFactory);
+	}
+	else
+	{
+		// No EdgeFixup means No override factory
+		// (storage factory is optional, it can exist without an EdgeFixup)
+		check(!OverrideFactory);
+	}
+
+	return EdgeFixup;
+}
+
 void ULandscapeComponent::OnRegister()
 {
 	Super::OnRegister();
 
-	if (GetLandscapeProxy())
+	if (ALandscapeProxy* Proxy = GetLandscapeProxy())
 	{
 		// Generate MID representing the MIC
-		if (GetLandscapeProxy()->bUseDynamicMaterialInstance)
+		if (Proxy->bUseDynamicMaterialInstance)
 		{
 			MaterialInstancesDynamic.Reserve(MaterialInstances.Num());
 
@@ -2121,7 +2224,7 @@ void ULandscapeComponent::OnRegister()
 		}
 
 		// AActor::GetWorld checks for Unreachable and BeginDestroyed
-		UWorld* World = GetLandscapeProxy()->GetWorld();
+		UWorld* World = Proxy->GetWorld();
 		if (World)
 		{
 			if (ULandscapeInfo* Info = GetLandscapeInfo())
@@ -3743,15 +3846,13 @@ void ALandscapeProxy::PreSave(FObjectPreSaveContext ObjectSaveContext)
 			CVarLandscapeHeightmapCompressionMode->GetValue(HeightmapCompressionMode);
 		}
 
-		bool bShouldCompressHeightmap = (HeightmapCompressionMode == 0) ? bUseCompressedHeightmapStorage : (HeightmapCompressionMode > 0);
-		if (bShouldCompressHeightmap)
+		const bool bShouldCompressHeightmap = (HeightmapCompressionMode == 0) ? bUseCompressedHeightmapStorage : (HeightmapCompressionMode > 0);
+		const bool bIsStreamingProxy = this->IsA<ALandscapeStreamingProxy>();
+		const bool bUseEdgeFixups = bIsStreamingProxy;
+		const bool bUpdateSnapshotNow = true;
+		for (ULandscapeComponent* LandscapeComponent : LandscapeComponents)
 		{
-			for (ULandscapeComponent* LandscapeComponent : LandscapeComponents)
-			{
-				UTexture2D* HeightmapTex = LandscapeComponent->GetHeightmap();
-				FVector LandscapeGridScale = GetRootComponent()->GetRelativeScale3D();
-				ULandscapeTextureStorageProviderFactory::ApplyTo(HeightmapTex, LandscapeGridScale);
-			}
+			LandscapeComponent->InstallOrUpdateTextureUserDatas(bUseEdgeFixups, bShouldCompressHeightmap, bUpdateSnapshotNow);
 		}
 	}
 #endif // WITH_EDITOR
@@ -5507,6 +5608,9 @@ void ULandscapeInfo::RegisterLandscapeActorWithProxyInternal(ALandscapeProxy* Pr
 #if WITH_EDITOR
 		StreamingProxy->FixupSharedData(Landscape, bMapCheck);
 #endif // WITH_EDITOR
+
+		Proxy->bIsLandscapeActorRegisteredWithLandscapeInfo = true;
+		FLandscapeGroup::RegisterAllComponentsOnStreamingProxy(StreamingProxy);
 	}
 
 #if WITH_EDITOR
@@ -5551,6 +5655,8 @@ void ULandscapeInfo::RegisterLandscapeActorWithProxyInternal(ALandscapeProxy* Pr
 		}
 	}
 #endif // WITH_EDITOR
+
+	Proxy->bIsLandscapeActorRegisteredWithLandscapeInfo = true;
 }
 
 void ULandscapeInfo::RegisterActor(ALandscapeProxy* Proxy, bool bMapCheck, bool bUpdateAllAddCollisions)
@@ -5738,6 +5844,8 @@ void ULandscapeInfo::UnregisterActor(ALandscapeProxy* Proxy)
 		ALandscapeStreamingProxy* StreamingProxy = CastChecked<ALandscapeStreamingProxy>(Proxy);
 		TWeakObjectPtr<ALandscapeStreamingProxy> StreamingProxyPtr = StreamingProxy;
 		StreamingProxies.Remove(StreamingProxyPtr);
+
+		FLandscapeGroup::UnregisterAllComponentsOnStreamingProxy(StreamingProxy);
 	}
 
 #if WITH_EDITOR
@@ -5772,6 +5880,7 @@ void ULandscapeInfo::UnregisterActor(ALandscapeProxy* Proxy)
 #if WITH_EDITORONLY_DATA
 	Proxy->bIsRegisteredWithLandscapeInfo = false;
 #endif // WITH_EDITORONLY_DATA
+	Proxy->bIsLandscapeActorRegisteredWithLandscapeInfo = false;
 }
 
 #if WITH_EDITOR
diff --git a/Engine/Source/Runtime/Landscape/Private/LandscapeEdgeFixup.cpp b/Engine/Source/Runtime/Landscape/Private/LandscapeEdgeFixup.cpp
new file mode 100644
--- /dev/null
+++ b/Engine/Source/Runtime/Landscape/Private/LandscapeEdgeFixup.cpp
@@ -0,0 +1,1498 @@
+// Copyright Epic Games, Inc. All Rights Reserved.
+
+#include "LandscapeEdgeFixup.h"
+#include "Containers/StridedView.h"
+#include "Landscape.h"
+#include "LandscapeComponent.h"
+#include "LandscapePrivate.h"
+#include "LandscapeGroup.h"
+#include "RenderingThread.h"
+#include "TextureResource.h"
+#include "Hash/xxhash.h"
+#include "Streaming/TextureMipDataProvider.h"
+
+// enables verbose debug spew
+//#define ENABLE_LANDSCAPE_EDGE_FIXUP_DEBUG_SPEW 1
+#ifdef ENABLE_LANDSCAPE_EDGE_FIXUP_DEBUG_SPEW
+#define FIXUP_DEBUG_LOG(...) UE_LOG(LogLandscape, Warning, __VA_ARGS__)
+#define FIXUP_DEBUG_LOG_PATCH(...) UE_LOG(LogLandscape, Warning, __VA_ARGS__)
+#define FIXUP_DEBUG_LOG_RENDER(...) UE_LOG(LogLandscape, Warning, __VA_ARGS__)
+#define FIXUP_DEBUG_LOG_DETAIL(...) UE_LOG(LogLandscape, Warning, __VA_ARGS__)
+#else
+#define FIXUP_DEBUG_LOG(...) UE_LOG(LogLandscape, Verbose, __VA_ARGS__)
+#define FIXUP_DEBUG_LOG_PATCH(...) do {} while(0)
+#define FIXUP_DEBUG_LOG_RENDER(...) do {} while(0)
+#define FIXUP_DEBUG_LOG_DETAIL(...) do {} while(0)
+#endif // ENABLE_LANDSCAPE_EDGE_FIXUP_DEBUG_SPEW
+
+namespace UE::Landscape
+{
+	EDirectionFlags ToFlag(EDirectionIndex Index)
+	{
+		return static_cast<EDirectionFlags>(0x01 << static_cast<uint32>(Index));
+	}
+
+	EDirectionIndex GetOppositeIndex(EDirectionIndex Index)
+	{
+		uint32 OppositeIndex = (static_cast<uint32>(Index) + 4) & 0x07;	// add 4, mod 8
+		return static_cast<EDirectionIndex>(OppositeIndex);
+	}
+
+	FIntPoint GetNeighborRelativePosition(ENeighborIndex DirectionIndex)
+	{
+		static FIntPoint RelativePositions[static_cast<int32>(ENeighborIndex::Count)] =
+		{
+			FIntPoint(0, -1),		// Bottom
+			FIntPoint(1, -1),		// Bottom Right
+			FIntPoint(1, 0),		// Right
+			FIntPoint(1, 1),		// Top Right
+			FIntPoint(0, 1),		// Top
+			FIntPoint(-1, 1),		// Top Left
+			FIntPoint(-1, 0),		// Left
+			FIntPoint(-1, -1),		// Bottom Left
+		};
+		if ((DirectionIndex >= ENeighborIndex::First) && (DirectionIndex <= ENeighborIndex::Last))
+		{
+			return RelativePositions[static_cast<int32>(DirectionIndex)];
+		}
+		return FIntPoint(0, 0);
+	}
+
+	const FString& GetDirectionString(EDirectionIndex Index)
+	{
+		FName temp(TEXT("Hello"));
+		temp.ToString();
+
+		static const FString DirectionStrings[static_cast<int32>(EDirectionIndex::Count)] =
+		{
+			FString("B"),			// Bottom
+			FString("BR"),			// Bottom Right
+			FString("R"),			// Right
+			FString("TR"),			// Top Right
+			FString("T"),			// Top
+			FString("TL"),			// Top Left
+			FString("L"),			// Left
+			FString("BL"),			// Bottom Left
+		};
+		if ((Index >= EDirectionIndex::First) && (Index <= EDirectionIndex::Last))
+		{
+			return DirectionStrings[static_cast<int32>(Index)];
+		}
+		static const FString ErrorString("Error");
+		return ErrorString;
+	}
+
+	// returns true if the given direction is diagonal
+	bool IsDiagonalCorner(EDirectionIndex DirectionIndex)
+	{
+		// cardinal directions have first bit zero, diagonals one
+		return (static_cast<uint32>(DirectionIndex) & 0x01) == 1;
+	}
+
+	bool IsTopOrBottomEdge(EEdgeIndex EdgeIndex)
+	{
+		// top and bottom are (index mod 4 == 0)
+		return (static_cast<uint32>(EdgeIndex) & 0x03) == 0;
+	}
+
+	bool IsDoubleTriangleCorner(EEdgeIndex EdgeIndex)
+	{
+		// double corner triangles are top-right (3) and bottom-left (7)
+		return (static_cast<uint32>(EdgeIndex) & 0x02) == 0;
+	}
+
+	// this converts a set of edge flags into a set of neighbors that are blended with those edges
+	ENeighborFlags EdgesToAffectedNeighbors(EEdgeFlags EdgeFlags)
+	{
+		// changed edges only affect the corresponding neighbor
+		// changed corners affect the corner neighbor AND the adjacent edge neighbors as well
+		uint32 RawFlags = static_cast<uint32>(EdgeFlags);
+		uint32 CornerFlags = RawFlags & static_cast<uint32>(ENeighborFlags::AllCorners);
+		uint32 AdjacentA = (CornerFlags << 1) | (CornerFlags >> 7); // rotate left one (mod 8, applied below)
+		uint32 AdjacentB = (CornerFlags << 7) | (CornerFlags >> 1); // rotate right one (mod 8, applied below)
+		return static_cast<ENeighborFlags>((RawFlags | AdjacentA | AdjacentB) & static_cast<uint32>(ENeighborFlags::All));
+	}
+
+	// this converts a set of neighbors into the local edges that they blend with
+	// (i.e. one of the neighbor's corners/edges is blended with the local edge)
+	EEdgeFlags NeighborsToBlendedEdges(ENeighborFlags NeighborFlags)
+	{
+		// corner neighbors only blend with the corresponding corner 
+		// edge neighbors blend the corresponding edge AND also blend with the adjacent corners
+		uint32 RawFlags = static_cast<uint32>(NeighborFlags);
+		uint32 EdgeFlags = RawFlags & static_cast<uint32>(EEdgeFlags::AllEdges);
+		uint32 AdjacentA = (EdgeFlags << 1) | (EdgeFlags >> 7); // rotate left one (mod 8, applied below)
+		uint32 AdjacentB = (EdgeFlags << 7) | (EdgeFlags >> 1);	// rotate right one (mod 8, applied below)
+		return static_cast<EEdgeFlags>((RawFlags | AdjacentA | AdjacentB) & static_cast<uint32>(EEdgeFlags::All));
+	}
+
+	// select the direction rotated clockwise (45 degrees * ClockwiseRotationOffset)
+	// this works for positive or negative offsets
+	EDirectionIndex RotateDirection(EDirectionIndex NeighborIndex, int32 ClockwiseRotationOffset)
+	{
+		int32 Index = static_cast<int32>(NeighborIndex);
+		int32 AdjacentIndex = (Index - ClockwiseRotationOffset) & (static_cast<int32>(EDirectionIndex::Count) - 1);
+		return static_cast<EDirectionIndex>(AdjacentIndex);
+	}
+
+	EDirectionFlags RotateFlags(EDirectionFlags DirFlags, int32 ClockwiseRotationOffset)
+	{
+		uint32 FlagBits = (uint32) DirFlags;
+		FlagBits = FlagBits | (FlagBits << 8);
+		FlagBits = (FlagBits >> (ClockwiseRotationOffset & 0x07)) & 0xFF;
+		return static_cast<EDirectionFlags>(FlagBits);
+	}
+
+	void GetSourceTextureEdgeStartCoord(EEdgeIndex EdgeIndex, int32 EdgeLength, int32& OutStartX, int32& OutStartY)
+	{
+		//						StartX	StartY	Stride		Count			Order
+		// 	Bottom = 0,			0		0		1			EdgeLength		left to right
+		// 	BottomRight = 1,	EL-1	0		0			1
+		// 	Right = 2,			EL-1	0		EL			EdgeLength		bottom to top
+		// 	TopRight = 3,		EL-1	EL-1	0			1
+		// 	Top = 4,			0		EL-1	1			EdgeLength		left to right
+		// 	TopLeft = 5,		0		EL-1	0			1
+		// 	Left = 6,			0		0		EL			EdgeLength		bottom to top
+		// 	BottomLeft = 7,		0		0		0			1
+
+		check((EdgeIndex >= EEdgeIndex::First) && (EdgeIndex <= EEdgeIndex::Last));
+		int32 Index = static_cast<int32>(EdgeIndex);
+		OutStartX = (((Index + 7) & 0x07) < 3) ? (EdgeLength-1) : 0;
+		OutStartY = (((Index + 5) & 0x07) < 3) ? (EdgeLength-1) : 0;
+	}
+
+	int32 GetSourceTextureEdgeStartStrideCount(EEdgeIndex EdgeIndex, int32 EdgeLength, int32& OutStride, int32& OutCount)
+	{
+		int32 StartX, StartY;
+		GetSourceTextureEdgeStartCoord(EdgeIndex, EdgeLength, StartX, StartY);
+		int32 Offset = StartY * EdgeLength + StartX;
+
+		if (!IsDiagonalCorner(EdgeIndex))
+		{
+			int32 Count = EdgeLength;
+			int32 Stride = IsTopOrBottomEdge(EdgeIndex) ? 1 : EdgeLength;
+
+			// exclude corners from the edge
+			Offset += Stride;
+			Count -= 2;
+
+			OutStride = Stride;
+			OutCount = Count;
+			return Offset;
+		}
+		else
+		{
+			OutStride = 0;
+			OutCount = 1;
+			return Offset;
+		}
+	}
+
+	const TArrayView<const ENeighborIndex> AllEdgeNeighbors()
+	{
+		static const ENeighborIndex EdgeNeighbors[4] =
+		{
+			ENeighborIndex::Bottom,
+			ENeighborIndex::Right,
+			ENeighborIndex::Top,
+			ENeighborIndex::Left,
+		};
+		return TArrayView<const ENeighborIndex>(&EdgeNeighbors[0], 4);
+	}
+
+	const TArrayView<const ENeighborIndex> AllCornerNeighbors()
+	{
+		static const ENeighborIndex CornerNeighbors[4] =
+		{
+			ENeighborIndex::BottomRight,
+			ENeighborIndex::TopRight,
+			ENeighborIndex::TopLeft,
+			ENeighborIndex::BottomLeft,
+		};
+		return TArrayView<const ENeighborIndex>(&CornerNeighbors[0], 4);
+	}
+
+	// size 8x8 texture (with mips)
+	// 
+	// C T T T T T T C
+	// L             R
+	// L             R   C T T C
+	// L             R   L     R   C C
+	// L             R   L     R   C C
+	// L             R   C B B C
+	// L             R
+	// C B B B B B B C
+	//      mip 0         mip 1   mip 2
+	//
+	// data is stored in this order:  edges are stored left to right and bottom to top (for left and right edges)
+	//
+	// EdgeSnapshot memory layout:
+	// B B				<< mip 1 (4x4) 2 pixel edges,	left to right
+	// B B B B B B		<< mip 0 (8x8) 6 pixel edges,	left to right
+	// R R				<< mip 1						bottom to top
+	// R R R R R R		<< mip 0						bottom to top
+	// T T												left to right
+	// T T T T T T										left to right
+	// L L												bottom to top
+	// L L L L L L										bottom to top
+	// 
+	// C C C C			bottom-right, top-right, top-left, bottom-left (no mips needed)
+
+	int32 GetEdgeArrayMipOffset(int32 MipEdgeLength)
+	{
+		// mip size		edge size	end		start(offset)
+		// 2x2			0			0		0
+		// 4x4			2			2		0
+		// 8x8			6			8		2
+		// 16x16		14			22		8
+		// 32x32		30			52		22
+		// 64x64		62			114		52
+		// 128x128		126			240		114
+		// 256x256		254			494		240
+
+		if (MipEdgeLength < 4)
+		{
+			return 0;
+		}
+		check((MipEdgeLength & (MipEdgeLength - 1)) == 0); // is power of two
+		int32 Log2 = FMath::FloorLog2(MipEdgeLength);
+		return MipEdgeLength - Log2 * 2;
+	}
+
+	int32 GetEdgeDataSize(int32 EdgeLength)
+	{
+		int32 DirectionSize = GetEdgeArrayMipOffset(EdgeLength * 2);
+		return 4 * DirectionSize;
+	}
+}
+
+TArrayView<UE::Landscape::FHeightmapTexel> FHeightmapTextureEdgeSnapshot::GetEdgeData(UE::Landscape::EEdgeIndex InEdgeIndex, int32 InMipIndex)
+{
+	using namespace UE::Landscape;
+	check((InEdgeIndex >= EEdgeIndex::First) && (InEdgeIndex <= EEdgeIndex::Last));
+	check(!IsDiagonalCorner(InEdgeIndex));
+	{
+		const int32 DirectionSize = EdgeData.Num() / 4;
+		const int32 MipEdgeLength = EdgeLength >> InMipIndex;
+		const int32 MipOffset = GetEdgeArrayMipOffset(MipEdgeLength);
+		const int32 DirOffset = DirectionSize * ((int32)InEdgeIndex / 2);
+		const int32 EdgeCount = MipEdgeLength - 2;
+		check(DirOffset + MipOffset + EdgeCount <= EdgeData.Num());	// check data exists in the buffer
+		return TArrayView<FHeightmapTexel>((FHeightmapTexel*) &EdgeData[DirOffset + MipOffset], EdgeCount);
+	}
+}
+
+UE::Landscape::FHeightmapTexel FHeightmapTextureEdgeSnapshot::GetCornerData(UE::Landscape::EEdgeIndex InCornerIndex)
+{
+	using namespace UE::Landscape;
+	check((InCornerIndex >= EEdgeIndex::First) && (InCornerIndex <= EEdgeIndex::Last));
+	check(IsDiagonalCorner(InCornerIndex));
+
+	// The snapshot corner hashes ARE the corner data (since they are the same size)
+	FHeightmapTexel Texel;
+	Texel.Data32 = SnapshotEdgeHashes[(int32)InCornerIndex];
+	return Texel;
+}
+
+TSharedRef<FHeightmapTextureEdgeSnapshot> FHeightmapTextureEdgeSnapshot::CreateEdgeSnapshotFromTextureData(const TArrayView64<UE::Landscape::FHeightmapTexel>& InHeightmapTextureData, int32 InEdgeLength, const FVector& LandscapeGridScale)
+{
+	TRACE_CPUPROFILER_EVENT_SCOPE(ULandscapeHeightmapTextureEdgeSnapshot::CreateEdgeDataFromTextureData);
+	TSharedRef<FHeightmapTextureEdgeSnapshot> NewEdgeSnapshot = MakeShared<FHeightmapTextureEdgeSnapshot>();
+	NewEdgeSnapshot->CaptureEdgeDataFromTextureData_Internal(InHeightmapTextureData, InEdgeLength, LandscapeGridScale);
+	return NewEdgeSnapshot;
+}
+
+#if WITH_EDITOR
+TSharedRef<FHeightmapTextureEdgeSnapshot> FHeightmapTextureEdgeSnapshot::CreateEdgeSnapshotFromHeightmapSource(UTexture2D* InHeightmap, const FVector& LandscapeGridScale)
+{
+	TRACE_CPUPROFILER_EVENT_SCOPE(ULandscapeHeightmapTextureEdgeSnapshot::CreateEdgeSnapshotFromHeightmapSource);	
+	TSharedRef<FHeightmapTextureEdgeSnapshot> NewEdgeSnapshot = MakeShared<FHeightmapTextureEdgeSnapshot>();
+	NewEdgeSnapshot->CaptureEdgeDataFromHeightmapSource_Internal(InHeightmap, LandscapeGridScale);
+	return NewEdgeSnapshot;
+}
+#endif // WITH_EDITOR
+
+
+UE::Landscape::EEdgeFlags FHeightmapTextureEdgeSnapshot::CompareEdges(const FHeightmapTextureEdgeSnapshot& Other) const
+{
+	using namespace UE::Landscape;
+
+	EEdgeFlags Changed = EEdgeFlags::None;
+	for (EEdgeIndex EdgeIndex : TEnumRange<EEdgeIndex>())
+	{
+		EEdgeFlags EdgeFlag = ToFlag(EdgeIndex);
+
+		if (SnapshotEdgeHashes[(int32)EdgeIndex] != Other.SnapshotEdgeHashes[(int32)EdgeIndex] ||
+			InitialEdgeHashes[(int32)EdgeIndex] != Other.InitialEdgeHashes[(int32)EdgeIndex])
+		{
+			Changed |= EdgeFlag;
+		}
+	}
+
+	return Changed;
+}
+
+void FHeightmapTextureEdgeSnapshot::ResizeForEdgeLength(const int32 InEdgeLength)
+{
+	if (EdgeLength != InEdgeLength)
+	{
+		EdgeLength = InEdgeLength;
+		const int32 EdgeDataSize = UE::Landscape::GetEdgeDataSize(EdgeLength);
+		EdgeData.SetNumZeroed(EdgeDataSize);
+	}
+}
+
+#if WITH_EDITOR
+void FHeightmapTextureEdgeSnapshot::CaptureEdgeDataFromHeightmapSource_Internal(UTexture2D* InHeightmap, const FVector& LandscapeGridScale)
+{
+	using namespace UE::Landscape;
+
+	FTextureSource& Source = InHeightmap->Source;
+	check(Source.IsValid());
+	check(Source.GetSizeX() == Source.GetSizeY());
+
+	TArray64<uint8> MipData;
+	verify(Source.GetMipData(MipData, 0));
+	TextureSourceID = Source.GetId();
+	FIXUP_DEBUG_LOG_DETAIL(TEXT("    CaptureEdgeDataFromHeightmapSource_Internal %p (%s) - ID %s"), InHeightmap, *LandscapeGridScale.ToString(), *TextureSourceID.ToString());
+
+	TArrayView64<FHeightmapTexel> TexelData((FHeightmapTexel*) MipData.GetData(), MipData.Num() / sizeof(FHeightmapTexel));
+	CaptureEdgeDataFromTextureData_Internal(TexelData, Source.GetSizeY(), LandscapeGridScale);
+}
+#endif // WITH_EDITOR
+
+void FHeightmapTextureEdgeSnapshot::CaptureSingleEdgeDataAndComputeNormalsAndHashes(
+	const UE::Landscape::FHeightmapTexel* TextureData, const UE::Landscape::EEdgeIndex EdgeOrCorner, const FVector& LandscapeGridScale)
+{
+	using namespace UE::Landscape;
+
+	// compute texture source stats
+	int32 SourceStride = 0;
+	int32 SourceCount = 0;
+	const int32 SourceOffset = GetSourceTextureEdgeStartStrideCount(EdgeOrCorner, EdgeLength, SourceStride, SourceCount);
+	const FHeightmapTexel* Src = &TextureData[SourceOffset]; // first pixel of the specified edge or corner in the source data
+	int32 SrcLineStrideTexels = EdgeLength;
+
+	// The triangle topology for each quad is:
+	// 
+	// 00 ------ 10
+	// | \       |
+	// |  \      |
+	// |   \     |
+	// | NB \ NT |
+	// |     \   |
+	// |      \  |
+	// |       \ |
+	// 01 ------ 11
+	// 
+	// We compute VertexNormals by considering every quad that borders the desired edge or corner, and accumulating the NT/NB normals into the neighboring vertices
+
+	auto CalculateNormalsForQuadAt = [SrcLineStrideTexels, LandscapeGridScale](const FHeightmapTexel* QuadTopLeftSrcPixel, FVector& OutNT, FVector& OutNB)
+	{
+		auto GetHeight = [SrcLineStrideTexels](const FHeightmapTexel* SrcPixel, int32 DX, int32 DY)
+		{
+			return SrcPixel[DY * SrcLineStrideTexels + DX].GetHeight16();
+		};
+
+		uint16 Height00 = GetHeight(QuadTopLeftSrcPixel, 0, 0);
+		uint16 Height01 = GetHeight(QuadTopLeftSrcPixel, 0, 1);
+		uint16 Height10 = GetHeight(QuadTopLeftSrcPixel, 1, 0);
+		uint16 Height11 = GetHeight(QuadTopLeftSrcPixel, 1, 1);
+
+		const FVector Vert00 = FVector(0.0f, 0.0f, LandscapeDataAccess::GetLocalHeight(Height00)) * LandscapeGridScale;
+		const FVector Vert01 = FVector(0.0f, 1.0f, LandscapeDataAccess::GetLocalHeight(Height01)) * LandscapeGridScale;
+		const FVector Vert10 = FVector(1.0f, 0.0f, LandscapeDataAccess::GetLocalHeight(Height10)) * LandscapeGridScale;
+		const FVector Vert11 = FVector(1.0f, 1.0f, LandscapeDataAccess::GetLocalHeight(Height11)) * LandscapeGridScale;
+
+		// top and bottom triangle normals
+		OutNT = ((Vert00 - Vert10) ^ (Vert10 - Vert11)).GetSafeNormal();
+		OutNB = ((Vert11 - Vert01) ^ (Vert01 - Vert00)).GetSafeNormal();
+	};
+
+	//						StartX	StartY	Stride		Count			Order
+	// 	Bottom = 0,			0		0		1			EdgeLength		left to right
+	// 	BottomRight = 1,	EL-1	0		0			1
+	// 	Right = 2,			EL-1	0		EL			EdgeLength		bottom to top
+	// 	TopRight = 3,		EL-1	EL-1	0			1
+	// 	Top = 4,			0		EL-1	1			EdgeLength		left to right
+	// 	TopLeft = 5,		0		EL-1	0			1
+	// 	Left = 6,			0		0		EL			EdgeLength		bottom to top
+	// 	BottomLeft = 7,		0		0		0			1
+
+	//the edges do not include the corner vertices, but the edge normals still need to consider the corner quads:
+	// 
+	//    BL    B0  B1   BR
+	//      v - v - v - v
+	//      | \ | \ | \ |
+	//   L0 v - v - v - v R0
+	//      | \ | \ | \ |
+	//   L1 v - v - v - v R1
+	//      | \ | \ | \ |
+	//      v - v - v - v
+	//    TL    T0  T1   TR
+
+	uint32 InitialHash;
+	uint32 SnapshotHash;
+	if (IsDiagonalCorner(EdgeOrCorner))
+	{
+		check(SourceCount == 1);
+		const FHeightmapTexel* QuadTopLeftSrcPixel = Src;
+		FVector NT, NB, VertexNormal;
+		switch (EdgeOrCorner)
+		{
+			case EEdgeIndex::BottomRight:
+				QuadTopLeftSrcPixel -= 1;
+				CalculateNormalsForQuadAt(QuadTopLeftSrcPixel, NT, NB);
+				VertexNormal = NT;
+				break;
+			case EEdgeIndex::TopLeft:
+				QuadTopLeftSrcPixel -= SrcLineStrideTexels;
+				CalculateNormalsForQuadAt(QuadTopLeftSrcPixel, NT, NB);
+				VertexNormal = NB;
+				break;
+			case EEdgeIndex::TopRight:
+				QuadTopLeftSrcPixel -= (SrcLineStrideTexels + 1);
+				// fallthrough
+			case EEdgeIndex::BottomLeft:
+				CalculateNormalsForQuadAt(QuadTopLeftSrcPixel, NT, NB);
+				VertexNormal = NT + NB;
+				break;
+		}
+
+		// Note: Snapshot corner data is only stored in the SnapshotHash (since they are the same)
+		FHeightmapTexel Dest;
+
+		// setup normal data
+		Dest.SetNormal(VertexNormal);
+
+		// copy height data
+		Dest.HeightL = Src->HeightL;
+		Dest.HeightH = Src->HeightH;
+
+		// hashes for corners are just a copy of the data (we take advantage of this in the hash comparison)
+		InitialHash = Src->Data32;
+		SnapshotHash = Dest.Data32;
+	}
+	else
+	{
+		const int32 MipZeroIndex = 0;
+		TArrayView<FHeightmapTexel> LocalEdgeData = GetEdgeData(EdgeOrCorner, MipZeroIndex);
+		FHeightmapTexel* Dst = LocalEdgeData.GetData();
+		check(SourceCount == LocalEdgeData.Num());
+
+		const FHeightmapTexel* QuadTopLeftSrcPixel = Src;
+		FVector NT, NB;
+		TArray<FVector, TFixedAllocator<512>> VertexNormals;	// if this ~12288 byte allocation blows the stack, we can convert it to a static array
+		VertexNormals.SetNumZeroed(SourceCount, EAllowShrinking::No);
+		switch (EdgeOrCorner)
+		{
+			case EEdgeIndex::Top:
+				QuadTopLeftSrcPixel -= (SrcLineStrideTexels + 1);
+				CalculateNormalsForQuadAt(QuadTopLeftSrcPixel, NT, NB);
+				VertexNormals[0] += NT + NB;
+				for (int32 i = 0; i < SourceCount - 1; ++i)
+				{
+					QuadTopLeftSrcPixel += 1;
+					CalculateNormalsForQuadAt(QuadTopLeftSrcPixel, NT, NB);
+					VertexNormals[i] += NB;
+					VertexNormals[i + 1] += NT + NB;
+				}
+				QuadTopLeftSrcPixel += 1;
+				CalculateNormalsForQuadAt(QuadTopLeftSrcPixel, NT, NB);
+				VertexNormals[SourceCount - 1] += NB;
+				break;
+			case EEdgeIndex::Left:
+				QuadTopLeftSrcPixel -= SrcLineStrideTexels;
+				CalculateNormalsForQuadAt(QuadTopLeftSrcPixel, NT, NB);
+				VertexNormals[0] += NB;
+				for (int32 i = 0; i < SourceCount - 1; ++i)
+				{
+					QuadTopLeftSrcPixel += SrcLineStrideTexels;
+					CalculateNormalsForQuadAt(QuadTopLeftSrcPixel, NT, NB);
+					VertexNormals[i] += NT + NB;
+					VertexNormals[i + 1] += NB;
+				}
+				QuadTopLeftSrcPixel += SrcLineStrideTexels;
+				CalculateNormalsForQuadAt(QuadTopLeftSrcPixel, NT, NB);
+				VertexNormals[SourceCount - 1] += NT + NB;
+				break;
+			case EEdgeIndex::Bottom:
+				QuadTopLeftSrcPixel -= 1;
+				CalculateNormalsForQuadAt(QuadTopLeftSrcPixel, NT, NB);
+				VertexNormals[0] += NT;
+				for (int32 i = 0; i < SourceCount - 1; ++i)
+				{
+					QuadTopLeftSrcPixel += 1;
+					CalculateNormalsForQuadAt(QuadTopLeftSrcPixel, NT, NB);
+					VertexNormals[i] += NT + NB;
+					VertexNormals[i + 1] += NT;
+				}
+				QuadTopLeftSrcPixel += 1;
+				CalculateNormalsForQuadAt(QuadTopLeftSrcPixel, NT, NB);
+				VertexNormals[SourceCount - 1] += NT + NB;
+				break;
+			case EEdgeIndex::Right:
+				QuadTopLeftSrcPixel -= (SrcLineStrideTexels + 1);
+				CalculateNormalsForQuadAt(QuadTopLeftSrcPixel, NT, NB);
+				VertexNormals[0] += NT + NB;
+				for (int32 i = 0; i < SourceCount - 1; ++i)
+				{
+					QuadTopLeftSrcPixel += SrcLineStrideTexels;
+					CalculateNormalsForQuadAt(QuadTopLeftSrcPixel, NT, NB);
+					VertexNormals[i] += NT;
+					VertexNormals[i + 1] += NT + NB;
+				}
+				QuadTopLeftSrcPixel += SrcLineStrideTexels;
+				CalculateNormalsForQuadAt(QuadTopLeftSrcPixel, NT, NB);
+				VertexNormals[SourceCount - 1] += NT;
+				break;
+		}
+
+		int32 SrcReadStrideTexels = IsTopOrBottomEdge(EdgeOrCorner) ? 1 : SrcLineStrideTexels;
+
+		// edge hashes are the hash of each pixel in the edge
+		FXxHash64Builder InitHashBuilder;
+		FXxHash64Builder SnapHashBuilder;
+		for (int32 i = 0; i < SourceCount; ++i)
+		{
+			// setup normal data
+			Dst->SetNormal(VertexNormals[i]);
+
+			// copy height data
+			Dst->HeightH = Src->HeightH;
+			Dst->HeightL = Src->HeightL;
+
+			// update hashes
+			InitHashBuilder.Update(Src, sizeof(FHeightmapTexel));
+			SnapHashBuilder.Update(Dst, sizeof(FHeightmapTexel));
+
+			Dst++;
+			Src += SrcReadStrideTexels;
+		}
+
+		InitialHash = GetTypeHash(InitHashBuilder.Finalize());
+		SnapshotHash = GetTypeHash(SnapHashBuilder.Finalize());
+
+		// downsample to fill out the mip edge data
+		{
+			int32 MipCount = FMath::CeilLogTwo(EdgeLength) - 1;	// we don't need to fill out the 1x1 or 2x2 mips as those have zero edge data..
+			TArrayView<FHeightmapTexel> PrevMipEdge = LocalEdgeData;
+			for (int32 MipIndex = 1; MipIndex < MipCount; MipIndex++)
+			{
+				TArrayView<FHeightmapTexel> MipEdgeData = GetEdgeData(EdgeOrCorner, MipIndex);
+
+				// the downsample pattern between mips: first and last elements get dropped, the rest have a 2 --> 1 averaging downsample applied
+				// NOTE that edge downsampling is specifically designed to not pull in non-edge data (does not average with the "middle" of the heightmap)
+				int32 MipEdgeLength = (PrevMipEdge.Num() - 2) / 2;
+				check(MipEdgeLength == MipEdgeData.Num());
+
+				FHeightmapTexel* MDst = MipEdgeData.GetData();
+				FHeightmapTexel* MSrc = PrevMipEdge.GetData() + 1;	// drop the first element
+				for (int32 x = 0; x < MipEdgeLength; x++)
+				{
+					// blend normal
+					MDst->NormalX = (MSrc[0].NormalX + (int32) MSrc[1].NormalX) / 2;
+					MDst->NormalY = (MSrc[0].NormalY + (int32) MSrc[1].NormalY) / 2;
+
+					// blend heights
+					uint32 Height0 = MSrc[0].GetHeight16();
+					uint32 Height1 = MSrc[1].GetHeight16();
+					uint32 HeightBlended = (Height0 + Height1) / 2;
+					MDst->SetHeight16(HeightBlended);
+
+					MDst += 1;
+					MSrc += 2;
+				}
+
+				PrevMipEdge = MipEdgeData;
+			}
+		}
+	}
+
+	uint32 OldInitialHash = InitialEdgeHashes[static_cast<int32>(EdgeOrCorner)];
+	uint32 OldSnapshotHash = SnapshotEdgeHashes[static_cast<int32>(EdgeOrCorner)];
+	FIXUP_DEBUG_LOG_DETAIL(TEXT("    - %s Hash i:%x s:%x --> i:%x s:%x"), *GetDirectionString(EdgeOrCorner), OldInitialHash, OldSnapshotHash, InitialHash, SnapshotHash);
+
+	InitialEdgeHashes[static_cast<int32>(EdgeOrCorner)] = InitialHash;
+	SnapshotEdgeHashes[static_cast<int32>(EdgeOrCorner)] = SnapshotHash;
+}
+
+void FHeightmapTextureEdgeSnapshot::CaptureEdgeDataFromTextureData_Internal(const TArrayView64<UE::Landscape::FHeightmapTexel>& InHeightmapTextureData, int32 InEdgeLength, const FVector& LandscapeGridScale)
+{
+	TRACE_CPUPROFILER_EVENT_SCOPE(FHeightmapTextureEdgeSnapshot::CaptureEdgeDataFromTextureData_Internal);
+	using namespace UE::Landscape;
+
+	check(InHeightmapTextureData.Num() == InEdgeLength * InEdgeLength);
+	ResizeForEdgeLength(InEdgeLength);
+
+	const FHeightmapTexel* TextureData = InHeightmapTextureData.GetData();
+	for (EEdgeIndex EdgeIndex : TEnumRange<EEdgeIndex>())
+	{
+		CaptureSingleEdgeDataAndComputeNormalsAndHashes(
+			TextureData,
+			EdgeIndex,
+			LandscapeGridScale);
+	}
+
+	FIXUP_DEBUG_LOG_DETAIL(TEXT("     CaptureEdgeDataFromTextureData_Internal --[%x %x %x %x %x %x %x %x | %x %x %x %x %x %x %x %x]--"),
+		SnapshotEdgeHashes[0], SnapshotEdgeHashes[1], SnapshotEdgeHashes[2], SnapshotEdgeHashes[3],
+		SnapshotEdgeHashes[4], SnapshotEdgeHashes[5], SnapshotEdgeHashes[6], SnapshotEdgeHashes[7],
+		InitialEdgeHashes[0], InitialEdgeHashes[1], InitialEdgeHashes[2], InitialEdgeHashes[3],
+		InitialEdgeHashes[4], InitialEdgeHashes[5], InitialEdgeHashes[6], InitialEdgeHashes[7]);
+
+	return;
+}
+
+struct FHeightmapTextureEdgeSnapshotCustomVersion
+{
+	enum Type
+	{
+		BeforeCustomVersionWasAdded = 0,
+		BeforeInitialHashWasAdded = 1,
+		BeforeCornerDataWasRemoved = 2,
+		BeforeChangedCornerHash = 3,
+		BeforeChangedCookedFormat = 4,
+		LatestVersion = 5
+	};
+
+	static const FGuid GUID;
+};
+
+const FGuid FHeightmapTextureEdgeSnapshotCustomVersion::GUID(0x12345678, 0x12345678, 0x12345678, 0x12345678);
+FCustomVersionRegistration GRegisterFHeightmapTextureEdgeSnapshotCustomVersion(
+	FHeightmapTextureEdgeSnapshotCustomVersion::GUID,
+	FHeightmapTextureEdgeSnapshotCustomVersion::LatestVersion,
+	TEXT("FHeightmapTextureEdgeSnapshotCustomVersion"));
+
+FArchive& operator<<(FArchive& Ar, FHeightmapTextureEdgeSnapshot& EdgeSnapshot)
+{
+	Ar.UsingCustomVersion(FHeightmapTextureEdgeSnapshotCustomVersion::GUID);
+
+	int32 CustomVersion = Ar.CustomVer(FHeightmapTextureEdgeSnapshotCustomVersion::GUID);
+
+	Ar << EdgeSnapshot.EdgeLength;
+	Ar << EdgeSnapshot.EdgeData;
+	if (CustomVersion <= FHeightmapTextureEdgeSnapshotCustomVersion::BeforeCornerDataWasRemoved)
+	{
+		TStaticArray<uint32, 4> CornerData = {};
+		Ar << CornerData;
+	}
+	Ar << EdgeSnapshot.SnapshotEdgeHashes;
+
+	if (CustomVersion > FHeightmapTextureEdgeSnapshotCustomVersion::BeforeCustomVersionWasAdded)
+	{
+#if WITH_EDITOR
+		if (!Ar.IsCooking())
+		{
+			EdgeSnapshot.TextureSourceID.Serialize(Ar);
+		}
+#endif // WITH_EDITOR
+	}
+
+	if (CustomVersion > FHeightmapTextureEdgeSnapshotCustomVersion::BeforeInitialHashWasAdded)
+	{
+		Ar << EdgeSnapshot.InitialEdgeHashes;
+	}
+
+#if WITH_EDITOR
+	if (Ar.IsLoading() && CustomVersion < FHeightmapTextureEdgeSnapshotCustomVersion::LatestVersion)
+	{
+		// invalidate the guid so that we trigger re-capture of this snapshot with the latest snapshot code
+		EdgeSnapshot.TextureSourceID.Invalidate();
+	}
+#endif // WITH_EDITOR
+
+	return Ar;
+}
+
+
+ULandscapeHeightmapTextureEdgeFixup::ULandscapeHeightmapTextureEdgeFixup(const FObjectInitializer& ObjectInitializer)
+	: Super(ObjectInitializer)
+	, EdgeSnapshot(MakeShared<FHeightmapTextureEdgeSnapshot>())
+{
+}
+
+void ULandscapeHeightmapTextureEdgeFixup::SetHeightmapTexture(UTexture2D* InHeightmapTexture)
+{
+	if (this->HeightmapTexture == nullptr)
+	{
+		// first time setup
+		this->HeightmapTexture = InHeightmapTexture;
+
+		// we assume the texture is in a pristine state the very first time
+		// so we initialize the GPU edge hash to our initial edge hashes
+		this->GPUEdgeHashes = EdgeSnapshot->InitialEdgeHashes;
+		this->GPUEdgeModifiedFlags = UE::Landscape::EEdgeFlags::None;
+
+		FIXUP_DEBUG_LOG_DETAIL(TEXT("    Set Initial GPU Edge Hashes: [B:%x BR:%x R:%x TR:%x T:%x TL:%x L:%x BL:%x]"),
+			GPUEdgeHashes[0], GPUEdgeHashes[1], GPUEdgeHashes[2], GPUEdgeHashes[3], GPUEdgeHashes[4], GPUEdgeHashes[5], GPUEdgeHashes[6], GPUEdgeHashes[7]);
+	}
+	else
+	{
+		// not the first time, tracker may have existing state
+		check(this->HeightmapTexture == InHeightmapTexture);
+
+		// do not modify GPUEdgeHashes here, it records the existing state of the GPU texture (it may have already been modified)
+
+		FIXUP_DEBUG_LOG_DETAIL(TEXT("    Repeat EdgeFixup Setup : no change GPU: [B:%x BR:%x R:%x TR:%x T:%x TL:%x L:%x BL:%x] CPU: [B:%x BR:%x R:%x TR:%x T:%x TL:%x L:%x BL:%x]"),
+			GPUEdgeHashes[0], GPUEdgeHashes[1], GPUEdgeHashes[2], GPUEdgeHashes[3], GPUEdgeHashes[4], GPUEdgeHashes[5], GPUEdgeHashes[6], GPUEdgeHashes[7],
+			EdgeSnapshot->SnapshotEdgeHashes[0], EdgeSnapshot->SnapshotEdgeHashes[1], EdgeSnapshot->SnapshotEdgeHashes[2], EdgeSnapshot->SnapshotEdgeHashes[3],
+			EdgeSnapshot->SnapshotEdgeHashes[4], EdgeSnapshot->SnapshotEdgeHashes[5], EdgeSnapshot->SnapshotEdgeHashes[6], EdgeSnapshot->SnapshotEdgeHashes[7]);
+	}
+}
+
+// set the active landscape component
+void ULandscapeHeightmapTextureEdgeFixup::SetActiveComponent(ULandscapeComponent* InComponent, FLandscapeGroup* InGroup, const bool bDisableCurrentActive)
+{
+	check(HeightmapTexture != nullptr);
+
+	if (ActiveComponent != nullptr)
+	{
+		check(ActiveGroup != nullptr);
+
+		if (bDisableCurrentActive)
+		{
+			// deactivate the active component/group, add it to the disabled list
+			ActiveGroup->DisableAndUnmap(this);
+
+			FIXUP_DEBUG_LOG_DETAIL(TEXT("    DisableAndUnmap %p (%d,%d) <Group %d Key %d> edgefixup:%p MAPPED:%d TOTAL:%d"),
+				ActiveComponent, GroupCoord.X, GroupCoord.Y, ActiveComponent->GetLandscapeProxy()->LODGroupKey, ActiveGroup->LandscapeGroupKey, this,
+				ActiveGroup->XYToEdgeFixupMap.Num(), ActiveGroup->AllRegisteredFixups.Num());
+		}
+		else
+		{
+			// unmap, don't flag as disabled
+			ActiveGroup->Unmap(this);
+
+			FIXUP_DEBUG_LOG_DETAIL(TEXT("    Unmap %p (%d,%d) <Group %d Key %d> edgefixup:%p MAPPED:%d TOTAL:%d"),
+				ActiveComponent, GroupCoord.X, GroupCoord.Y, ActiveComponent->GetLandscapeProxy()->LODGroupKey, ActiveGroup->LandscapeGroupKey, this, 
+				ActiveGroup->XYToEdgeFixupMap.Num(), ActiveGroup->AllRegisteredFixups.Num());
+		}
+	}
+	else
+	{
+		check(ActiveGroup == nullptr);
+	}
+	
+	ActiveComponent = InComponent;
+	ActiveGroup = InGroup;
+
+	InGroup->Map(this, InComponent);
+
+	FIXUP_DEBUG_LOG_DETAIL(TEXT("    Map %p (%d,%d) <Group %d Key %d> edgefixup:%p MAPPED:%d TOTAL:%d "),
+		InComponent, GroupCoord.X, GroupCoord.Y, InComponent->GetLandscapeProxy()->LODGroupKey, InGroup->LandscapeGroupKey, this,
+		InGroup->XYToEdgeFixupMap.Num(), InGroup->AllRegisteredFixups.Num());
+}
+
+void ULandscapeHeightmapTextureEdgeFixup::RequestEdgeTexturePatchingForNeighbors(UE::Landscape::ENeighborFlags NeighborsNeedingPatching)
+{
+	using namespace UE::Landscape;
+	check(bMapped);
+	for (ENeighborIndex NeighborIndex : TEnumRange<ENeighborIndex>())
+	{
+		if (EnumHasAnyFlags(NeighborsNeedingPatching, ToFlag(NeighborIndex)))
+		{
+			// a change to an edge or corner always affects the neighbor in the given direction
+			if (ULandscapeHeightmapTextureEdgeFixup* Neighbor = ActiveGroup->GetNeighborEdgeFixup(GroupCoord, NeighborIndex))
+			{
+				FIXUP_DEBUG_LOG_DETAIL(TEXT("  %s Neighbor %p (%d,%d) flagged for edge texture patching"), 
+					*GetDirectionString(NeighborIndex), Neighbor->ActiveComponent, Neighbor->GroupCoord.X, Neighbor->GroupCoord.Y);
+				check(ActiveGroup->AllRegisteredFixups.Contains(Neighbor));
+				check(Neighbor->bMapped);
+				ActiveGroup->HeightmapsNeedingEdgeTexturePatching.Add(Neighbor);
+			}
+		}
+	}
+}
+
+void ULandscapeHeightmapTextureEdgeFixup::BlendCornerData(UE::Landscape::FHeightmapTexel& OutTexel, UE::Landscape::EEdgeIndex CornerIndex, const UE::Landscape::FNeighborSnapshots& NeighborSnapshots)
+{
+	using namespace UE::Landscape;
+
+	check(IsDiagonalCorner(CornerIndex));
+
+	uint32 NormalX = 0;
+	uint32 NormalY = 0;
+	uint32 Height = 0;
+	uint32 SampleCount = 0;
+
+	auto Accumulate = [&NormalX, &NormalY, &Height, &SampleCount](FHeightmapTextureEdgeSnapshot* Snapshot, EEdgeIndex EdgeIndex)
+	{
+		if (Snapshot != nullptr)
+		{
+			FHeightmapTexel CornerData = Snapshot->GetCornerData(EdgeIndex);
+
+			// to match corner blend behavior, TL and BR are weighted 1, TR and BL are weighted 2
+			// because they have twice as many triangles incidient on the corner.
+			if (IsDoubleTriangleCorner(EdgeIndex))
+			{
+				NormalX += CornerData.NormalX;
+				NormalY += CornerData.NormalY;
+				Height += CornerData.GetHeight16();
+				SampleCount++;
+			}
+			else
+			{
+				NormalX += CornerData.NormalX * 2;
+				NormalY += CornerData.NormalY * 2;
+				Height += CornerData.GetHeight16() * 2;
+				SampleCount += 2;
+			}
+		}
+	};
+
+	FHeightmapTextureEdgeSnapshot* NeighborSnapshot = NeighborSnapshots.NeighborSnapshots[(int32)CornerIndex];
+	FHeightmapTextureEdgeSnapshot* NeighborEdgeASnapshot = NeighborSnapshots.NeighborSnapshots[(int32)RotateDirection(CornerIndex, 1)];
+	FHeightmapTextureEdgeSnapshot* NeighborEdgeBSnapshot = NeighborSnapshots.NeighborSnapshots[(int32)RotateDirection(CornerIndex, -1)];
+
+	Accumulate(NeighborSnapshots.LocalSnapshot, CornerIndex);
+	Accumulate(NeighborSnapshot, GetOppositeIndex(CornerIndex));
+	Accumulate(NeighborEdgeASnapshot, RotateDirection(CornerIndex, -2));
+	Accumulate(NeighborEdgeBSnapshot, RotateDirection(CornerIndex, 2));
+	check(SampleCount >= 2);	// at least we should have the local and one neighbor before callign this function
+
+	OutTexel.NormalX = NormalX / SampleCount;
+	OutTexel.NormalY = NormalY / SampleCount;
+	OutTexel.SetHeight16(Height / SampleCount);
+}
+
+void ULandscapeHeightmapTextureEdgeFixup::BlendEdgeData(FHeightmapTextureEdgeSnapshot& EdgeSnapshot, UE::Landscape::EEdgeIndex EdgeIndex, int32 MipIndex, FHeightmapTextureEdgeSnapshot& NeighborEdgeSnapshot, TStridedView<UE::Landscape::FHeightmapTexel>& OutDestView)
+{
+	using namespace UE::Landscape;
+	check(!IsDiagonalCorner(EdgeIndex));
+
+	TArrayView<const FHeightmapTexel> ThisEdgeData = EdgeSnapshot.GetEdgeData(EdgeIndex, MipIndex);
+	TArrayView<const FHeightmapTexel> NeighborData = NeighborEdgeSnapshot.GetEdgeData(GetOppositeIndex(EdgeIndex), MipIndex);
+	int32 Count = ThisEdgeData.Num();
+	check(Count == NeighborData.Num());
+
+	// allocate the buffer if necessary
+	if (OutDestView.Num() == 0)
+	{
+		FHeightmapTexel* Buffer = new FHeightmapTexel[Count];
+		OutDestView = TStridedView<FHeightmapTexel>(sizeof(FHeightmapTexel), Buffer, Count);
+	}
+	check(Count == OutDestView.Num());
+
+	// blend data between the edge and the neighboring edge
+	const FHeightmapTexel* SrcE = &ThisEdgeData[0];
+	const FHeightmapTexel* SrcN = &NeighborData[0];
+	FHeightmapTexel* Dst = &OutDestView[0];
+	int32 DstStrideInTexels = OutDestView.GetStride() / sizeof(FHeightmapTexel);
+
+	for (int32 p = 0; p < Count; p++)
+	{
+		Dst->NormalX = (SrcE->NormalX + (int32) SrcN->NormalX) / 2;
+		Dst->NormalY = (SrcE->NormalY + (int32) SrcN->NormalY) / 2;
+
+		// blend heights
+		uint32 HeightE = SrcE->GetHeight16();
+		uint32 HeightN = SrcN->GetHeight16();
+		uint32 HeightBlended = (HeightE + HeightN) / 2;
+		Dst->SetHeight16(HeightBlended);
+
+		SrcE += 1;
+		SrcN += 1;
+		Dst += DstStrideInTexels;
+	}
+}
+
+void ULandscapeHeightmapTextureEdgeFixup::GetNeighborSnapshots(UE::Landscape::FNeighborSnapshots &OutSnapshots)
+{
+	using namespace UE::Landscape;
+
+	check(ActiveGroup != nullptr); // should be caught before calling this function
+	check(bMapped);
+
+	ENeighborFlags ExistingNeighbors = ENeighborFlags::None;
+	ENeighborFlags AnyModified = GPUEdgeModifiedFlags;
+
+	for (ENeighborIndex NeighborIndex : TEnumRange<ENeighborIndex>())
+	{
+		if (ULandscapeHeightmapTextureEdgeFixup* Neighbor = ActiveGroup->GetNeighborEdgeFixup(GroupCoord, NeighborIndex))
+		{
+			ENeighborFlags NeighborFlags = ToFlag(NeighborIndex);
+			ExistingNeighbors |= NeighborFlags;
+			OutSnapshots.NeighborSnapshots[(int32)NeighborIndex] = &Neighbor->EdgeSnapshot.Get();
+
+			EEdgeFlags NeighborEdgeModifiedFlags = Neighbor->GPUEdgeModifiedFlags;
+			if (EnumHasAnyFlags(NeighborEdgeModifiedFlags, ToFlag(GetOppositeIndex(NeighborIndex))))
+			{
+				EnumAddFlags(AnyModified, NeighborFlags);
+			}
+
+			if (!IsDiagonalCorner(NeighborIndex))
+			{
+				// this neighbor edge blends into our adjacent local corners
+				EDirectionIndex RC = RotateDirection(NeighborIndex, 1);
+				EDirectionIndex LC = RotateDirection(NeighborIndex, -1);
+
+				if (EnumHasAnyFlags(NeighborEdgeModifiedFlags, ToFlag(GetOppositeIndex(RC))))
+				{
+					EnumAddFlags(AnyModified, ToFlag(LC));
+				}
+
+				if (EnumHasAnyFlags(NeighborEdgeModifiedFlags, ToFlag(GetOppositeIndex(LC))))
+				{
+					EnumAddFlags(AnyModified, ToFlag(RC));
+				}
+			}
+		}
+		else
+		{
+			OutSnapshots.NeighborSnapshots[(int32)NeighborIndex] = nullptr;
+		}
+	}
+
+	OutSnapshots.ExistingNeighbors = ExistingNeighbors;
+	OutSnapshots.EdgesWithAnyModifiedNeighbor = AnyModified;
+	OutSnapshots.LocalSnapshot = &EdgeSnapshot.Get();
+	OutSnapshots.GPUEdgeHashes = GPUEdgeHashes;
+}
+
+int32 ULandscapeHeightmapTextureEdgeFixup::CheckAndPatchTextureEdgesFromEdgeSnapshots()
+{
+	TRACE_CPUPROFILER_EVENT_SCOPE(CheckAndPatchTextureEdgesFromEdgeSnapshots);
+
+	using namespace UE::Landscape;
+
+	FNeighborSnapshots NeighborSnapshots;
+	GetNeighborSnapshots(NeighborSnapshots);
+
+	int32 PatchedEdgeCount = 0;
+
+	// check any edges that may need to be patched, and patch them if necessary
+	if (EnumHasAnyFlags(NeighborSnapshots.ExistingNeighbors, EEdgeFlags::AllEdges))
+	{
+		for (EEdgeIndex EdgeIndex : AllEdgeNeighbors())
+		{
+			FHeightmapTextureEdgeSnapshot* NeighborSnapshot = NeighborSnapshots.NeighborSnapshots[(int32)EdgeIndex];
+			if (NeighborSnapshot == nullptr)
+			{
+				FIXUP_DEBUG_LOG_PATCH(TEXT("  nopatch  %s Edge - no neighbor"), *GetDirectionString(EdgeIndex));
+				continue;
+			}
+
+			const EEdgeFlags EdgeFlag = ToFlag(EdgeIndex);
+
+			const bool bModified = EnumHasAnyFlags(NeighborSnapshots.EdgesWithAnyModifiedNeighbor, EdgeFlag);
+			if (!bModified)
+			{
+				// with no modified edges affecting this edge, then we can compare initial states first, as it will likely be a match
+				const uint32 LocalInitialHash = NeighborSnapshots.LocalSnapshot->InitialEdgeHashes[(int32)EdgeIndex];
+				const uint32 NeighborInitialHash = NeighborSnapshot->InitialEdgeHashes[(int32)GetOppositeIndex(EdgeIndex)];
+				check(LocalInitialHash == NeighborSnapshots.GPUEdgeHashes[(int32)EdgeIndex]);
+				const bool bInitialStateMatch = (LocalInitialHash == NeighborInitialHash);
+
+				if (bInitialStateMatch)
+				{
+					FIXUP_DEBUG_LOG_PATCH(TEXT("  nopatch  %s Edge - initial state match"), *GetDirectionString(EdgeIndex));
+					continue;
+				}
+			}
+
+			// lastly, check if this edge has already been patched to the desired state, by comparing hashes
+			const uint32 LocalEdgeHash = NeighborSnapshots.LocalSnapshot->SnapshotEdgeHashes[(int32)EdgeIndex];
+			const uint32 NeighborEdgeHash = NeighborSnapshot->SnapshotEdgeHashes[(int32)GetOppositeIndex(EdgeIndex)];
+			const uint32 GPUEdgeHash = NeighborSnapshots.GPUEdgeHashes[(int32)EdgeIndex];
+			uint32 CombinedEdgeHash = HashCombineFast(LocalEdgeHash, NeighborEdgeHash);
+
+			if (CombinedEdgeHash == GPUEdgeHash)
+			{
+				FIXUP_DEBUG_LOG_PATCH(TEXT("  nopatch  %s Edge - modified match (%x,%x) (== old %x)"), *GetDirectionString(EdgeIndex), LocalEdgeHash, NeighborEdgeHash, GPUEdgeHash);
+				continue;
+			}
+
+			FIXUP_DEBUG_LOG_PATCH(TEXT("  PATCHING %s Edge (hash %x | %x, new %x != old %x)"),
+				*GetDirectionString(EdgeIndex), LocalEdgeHash, NeighborEdgeHash, CombinedEdgeHash, GPUEdgeHash);
+			PatchTextureEdge_Internal(EdgeIndex);
+			PatchedEdgeCount++;
+			GPUEdgeHashes[(int32)EdgeIndex] = CombinedEdgeHash;
+			EnumAddFlags(GPUEdgeModifiedFlags, EdgeFlag);
+		}
+	}
+
+	// check any corners that may need to be patched, and patch them if necessary
+	const EEdgeFlags CornersWithOneOrMoreNeighbors = NeighborsToBlendedEdges(NeighborSnapshots.ExistingNeighbors);
+	if (EnumHasAnyFlags(CornersWithOneOrMoreNeighbors, EEdgeFlags::AllCorners))
+	{
+		for (EEdgeIndex CornerIndex : AllCornerNeighbors())
+		{
+			const EEdgeFlags CornerFlag = ToFlag(CornerIndex);
+
+			// don't patch a corner unless at least one neighbor exists
+			if (!EnumHasAnyFlags(CornersWithOneOrMoreNeighbors, CornerFlag))
+			{
+				FIXUP_DEBUG_LOG_PATCH(TEXT("  nopatch  %s Corner - no neighbors"), *GetDirectionString(CornerIndex));
+				continue;
+			}
+
+			const bool bModified = EnumHasAnyFlags(NeighborSnapshots.EdgesWithAnyModifiedNeighbor, CornerFlag);
+			if (!bModified)
+			{
+				// if each edge's initial hash matches, no need to patch
+				// (this is slightly different than the blend matching, as the blend may only include a subset of the edges, and produce a different result,
+				// whereas the initial state may include data from edges that are not yet loaded)
+				const uint32 LocalInitialHash = NeighborSnapshots.LocalSnapshot->InitialEdgeHashes[(int32)CornerIndex];
+
+				auto NeighborHasSameInitialHash = [LocalInitialHash, CornerIndex, NeighborSnapshots](int32 NeighborRotation, int32 NeighborEdgeRotation) -> bool
+				{
+					if (FHeightmapTextureEdgeSnapshot* Snapshot = NeighborSnapshots.NeighborSnapshots[(int32)RotateDirection(CornerIndex, NeighborRotation)])
+					{
+						return Snapshot->InitialEdgeHashes[(int32)RotateDirection(CornerIndex, NeighborEdgeRotation)] == LocalInitialHash;
+					}
+					return true;
+				};
+
+				const bool bInitialStateMatch =
+					NeighborHasSameInitialHash(1, -2) &&
+					NeighborHasSameInitialHash(-1, 2) &&
+					NeighborHasSameInitialHash(0, 4);
+
+				if (bInitialStateMatch)
+				{
+					FIXUP_DEBUG_LOG_PATCH(TEXT("  nopatch  %s Corner - initial state match"), *GetDirectionString(CornerIndex));
+					continue;
+				}
+			}
+
+			// Because the corner hashes are just a copy of the corner data, we can directly compare them with the expected result
+			FHeightmapTexel Blended, GPUExisting;
+			BlendCornerData(Blended, CornerIndex, NeighborSnapshots);
+
+			GPUExisting.Data32 = GPUEdgeHashes[(int32)CornerIndex];
+
+			if (Blended.Data32 == GPUExisting.Data32)
+			{
+				FIXUP_DEBUG_LOG_PATCH(TEXT("  nopatch  %s Corner - exact match (%x == old %x)"), *GetDirectionString(CornerIndex), Blended.Data32, GPUExisting.Data32);
+				continue;
+			}
+
+			// Check if we're only off by one in the normal, and don't bother patching these cases.
+			// It might be noticeable with a very shiny material, but not in most landscape cases.
+			// These cases generally happen due to floating point and round off error and inexactness between GPU and CPU implementations.
+			if (Blended.IsSameHeight(GPUExisting) &&
+				(abs(Blended.NormalX - (int32)GPUExisting.NormalX) <= 1) &&
+				(abs(Blended.NormalY - (int32)GPUExisting.NormalY) <= 1))
+			{
+				FIXUP_DEBUG_LOG_PATCH(TEXT("  nopatch  %s Corner - off-by-one match (%x ~ old %x)"), *GetDirectionString(CornerIndex), Blended.Data32, GPUExisting.Data32);
+				continue;
+			}
+
+			FIXUP_DEBUG_LOG_PATCH(TEXT("  PATCHING %s Corner (new %x != old %x)"), *GetDirectionString(CornerIndex), Blended.Data32, GPUExisting.Data32);
+			PatchTextureCorner_Internal(CornerIndex, Blended);
+			PatchedEdgeCount++;
+			GPUEdgeHashes[(int32)CornerIndex] = Blended.Data32;
+			EnumAddFlags(GPUEdgeModifiedFlags, CornerFlag);
+		}
+	}
+
+	return PatchedEdgeCount;
+}
+
+void ULandscapeHeightmapTextureEdgeFixup::PatchTextureEdge_Internal(UE::Landscape::EEdgeIndex EdgeIndex)
+{
+	TRACE_CPUPROFILER_EVENT_SCOPE(PatchTextureEdge_Internal);
+
+	using namespace UE::Landscape;
+
+#if WITH_EDITOR
+	// this should be caught earlier -- we can't patch default textures
+	check(!HeightmapTexture->IsDefaultTexture());
+#endif // WITH_EDITOR
+	check(!IsDiagonalCorner(EdgeIndex));
+
+	FTextureResource* HeightTextureResource = HeightmapTexture->GetResource();
+
+	const int32 NumMips = HeightmapTexture->GetNumMips();
+	const int32 ResidentMips = HeightmapTexture->GetNumResidentMips();
+
+	ULandscapeHeightmapTextureEdgeFixup* EdgeNeighbor = ActiveGroup->GetNeighborEdgeFixup(GroupCoord, EdgeIndex);
+		
+	// the first (and largest) resident mip. For example, if all mips are resident, then this is mip 0.
+	const int32 FirstResidentMipIndex = NumMips - ResidentMips;
+
+	// this is the size of the first mip that has edges (2x2 and 1x1 mips don't have "edges", just "corners")
+	const int32 MinMipSizeToFix = 4;
+
+	// update this edge using data from neighbor and us		
+	int32 MipSize = EdgeSnapshot->EdgeLength >> FirstResidentMipIndex;
+	for (int32 MipIndex = FirstResidentMipIndex; MipIndex < NumMips && MipSize >= MinMipSizeToFix; MipIndex++, MipSize /= 2)
+	{
+		// passing an empty TStridedView to BlendEdgeData will make it allocate a buffer to hold the blended data
+		// (which we then delete in the render command below)
+		TStridedView<UE::Landscape::FHeightmapTexel> DestView;
+		BlendEdgeData(EdgeSnapshot.Get(), EdgeIndex, MipIndex, EdgeNeighbor->EdgeSnapshot.Get(), DestView);
+
+		ENQUEUE_RENDER_COMMAND(UpdateLandscapeHeightmapEdge)(
+			[HeightTextureResource, DestView, EdgeIndex, MipSize, MipIndex, FirstResidentMipIndex, ResidentMips](FRHICommandListImmediate& RHICmdList)
+			{
+				const uint8* Buffer = DestView[0].Data;
+
+				FTextureRHIRef RHIHeightmapTexture = HeightTextureResource->GetTextureRHI();
+
+				const FRHITextureDesc& Desc = RHIHeightmapTexture->GetDesc();
+
+				// Resource doesn't know about the unstreamed mips, so it's zero Mip starts at the first resident mip
+				check(ResidentMips == Desc.NumMips);
+				int32 ResourceMipIndex = MipIndex - FirstResidentMipIndex;
+
+				if (ResourceMipIndex >= (int32) RHIHeightmapTexture->GetNumMips())
+				{
+					UE_LOG(LogLandscape, Error, TEXT("   Can't Update Edge %d Mip %d because the texture resource does not have that mip"), EdgeIndex, ResourceMipIndex);
+				}
+				else
+				{
+					// determine the min/max range to update on the edge, based on whether we are updating the corners or not				
+					int32 DstStartX, DstStartY;
+					int32 SrcStartX, SrcStartY, SrcWidth, SrcHeight;
+					SrcStartX = 0;
+					SrcStartY = 0;
+					if (IsTopOrBottomEdge(EdgeIndex))
+					{
+						DstStartX = 1;
+						DstStartY = (EdgeIndex == EEdgeIndex::Bottom) ? 0 : (MipSize - 1);
+						SrcWidth = MipSize - 2;
+						SrcHeight = 1;
+					}
+					else
+					{
+						DstStartX = (EdgeIndex == EEdgeIndex::Left) ? 0 : (MipSize - 1);
+						DstStartY = 1;
+						SrcWidth = 1;
+						SrcHeight = MipSize - 2;
+					}
+
+					{
+						constexpr uint32 SourcePitch = 4;
+						FUpdateTextureRegion2D UpdateRegion(DstStartX, DstStartY, SrcStartX, SrcStartY, SrcWidth, SrcHeight);
+
+						FIXUP_DEBUG_LOG_RENDER(TEXT("    PatchTextureEdge_Internal [RHIUpdateTexture2D] res:%p e:%d mip:%d len:%d reg:[%d %d %d %d %d %d]"), HeightTextureResource, EdgeIndex, MipIndex, MipEdgeLength,
+							UpdateRegion.DestX, UpdateRegion.DestY, UpdateRegion.SrcX, UpdateRegion.SrcY, UpdateRegion.Width, UpdateRegion.Height);
+
+						RHIUpdateTexture2D(RHIHeightmapTexture, ResourceMipIndex, UpdateRegion, SourcePitch, Buffer);
+					}
+				}
+
+				delete Buffer;
+			}
+		);
+	}
+}
+
+void ULandscapeHeightmapTextureEdgeFixup::PatchTextureCorner_Internal(UE::Landscape::EEdgeIndex CornerIndex, UE::Landscape::FHeightmapTexel Texel)
+{
+	TRACE_CPUPROFILER_EVENT_SCOPE(PatchTextureCorner_Internal);
+
+	using namespace UE::Landscape;
+
+#if WITH_EDITOR
+	// this should be caught earlier -- we can't patch default textures
+	check(!HeightmapTexture->IsDefaultTexture());
+#endif // WITH_EDITOR
+	check(IsDiagonalCorner(CornerIndex));
+
+	FTextureResource* HeightTextureResource = HeightmapTexture->GetResource();
+
+	const int32 NumMips = HeightmapTexture->GetNumMips();
+	const int32 ResidentMips = HeightmapTexture->GetNumResidentMips();
+
+	// the first (and largest) resident mip. For example, if all mips are resident, then this is mip 0.
+	const int32 FirstResidentMipIndex = NumMips - ResidentMips;
+
+	int32 FirstResidentMipSize = (EdgeSnapshot->EdgeLength >> FirstResidentMipIndex);
+
+	ENQUEUE_RENDER_COMMAND(UpdateLandscapeHeightmapCorner)(
+		[HeightTextureResource, Texel, CornerIndex, FirstResidentMipSize](FRHICommandListImmediate& RHICmdList)
+		{
+			FTextureRHIRef RHIHeightmapTexture = HeightTextureResource->GetTextureRHI();
+
+			const FRHITextureDesc& Desc = RHIHeightmapTexture->GetDesc();
+			const int32 NumMips = Desc.NumMips;
+
+			// this is the size of the smallest mip that has 4 corners (1x1 mip only has one texel)
+			const int32 MinMipSizeToFix = 2;
+
+			int32 MipSize = FirstResidentMipSize;
+			check(FirstResidentMipSize == Desc.GetSize().X);
+
+			for (int32 ResourceMipIndex = 0; ResourceMipIndex < NumMips && MipSize >= MinMipSizeToFix; ResourceMipIndex++, MipSize /= 2)
+			{
+				// grab the corner coordinates in the texture
+				int32 DstStartX, DstStartY;
+				GetSourceTextureEdgeStartCoord(CornerIndex, MipSize, DstStartX, DstStartY);
+
+				// update one pixel
+				uint32 SourcePitch = 4;
+				const int32 SrcStartX = 0;
+				const int32 SrcStartY = 0;
+				const int32 SrcWidth = 1;
+				const int32 SrcHeight = 1;
+
+				FUpdateTextureRegion2D UpdateRegion(DstStartX, DstStartY, SrcStartX, SrcStartY, SrcWidth, SrcHeight);
+
+				FIXUP_DEBUG_LOG_RENDER(TEXT("    PatchTextureCorner_Internal [RHIUpdateTexture2D] res:%p e:%d mip:%d reg:[%d %d %d %d %d %d]"), HeightTextureResource, CornerIndex, ResourceMipIndex,
+					UpdateRegion.DestX, UpdateRegion.DestY, UpdateRegion.SrcX, UpdateRegion.SrcY, UpdateRegion.Width, UpdateRegion.Height);
+
+				RHIUpdateTexture2D(RHIHeightmapTexture, ResourceMipIndex, UpdateRegion, SourcePitch, Texel.Data);
+			}
+		}
+	);
+}
+
+int32 ULandscapeHeightmapTextureEdgeFixup::PatchTextureEdgesForStreamingMips(int32 FirstMipIndexInclusive, int32 LastMipIndexExclusive, FTextureMipInfoArray& DestMipInfos, const UE::Landscape::FNeighborSnapshots& NeighborSnapshots)
+{
+	int32 PatchedEdges = 0;
+
+	// TODO [chris.tchou] : it would be more efficient to iterate the mips internally to each edge, more similar to the non-streaming patch case
+	for (int MipIndex = FirstMipIndexInclusive; MipIndex < LastMipIndexExclusive; MipIndex++)
+	{
+		FTextureMipInfo& DestMipInfo = DestMipInfos[MipIndex];
+		PatchedEdges += PatchTextureEdgesForSingleMip(MipIndex, DestMipInfo, NeighborSnapshots);
+	}
+	return PatchedEdges;
+}
+
+int32 ULandscapeHeightmapTextureEdgeFixup::PatchTextureEdgesForSingleMip(int32 MipIndex, FTextureMipInfo& DestMipInfo, const UE::Landscape::FNeighborSnapshots& NeighborSnapshots)
+{
+	using namespace UE::Landscape;
+
+	check(DestMipInfo.Format == PF_B8G8R8A8); // format used by FHeightmapTexel
+
+	FHeightmapTexel* Dest = (FHeightmapTexel*) DestMipInfo.DestData;
+
+	// some mip allocators do not provide RowPitch -- in that case assume SizeX texels is the pitch
+	uint32 DstPitchInTexels;
+	if (DestMipInfo.RowPitch > 0)
+	{
+		DstPitchInTexels = DestMipInfo.RowPitch / sizeof(FHeightmapTexel);
+	}
+	else
+	{
+		DstPitchInTexels = DestMipInfo.SizeX;
+	}
+
+	int32 MipEdgeLength = NeighborSnapshots.LocalSnapshot->EdgeLength >> MipIndex;
+	check(DestMipInfo.SizeX >= (uint32) MipEdgeLength);
+	check(DestMipInfo.SizeY >= (uint32) MipEdgeLength);
+
+	int32 PatchedEdges = 0;
+
+	if (EnumHasAnyFlags(NeighborSnapshots.EdgesWithAnyModifiedNeighbor, EEdgeFlags::AllEdges))
+	{
+		for (EEdgeIndex EdgeIndex : AllEdgeNeighbors())
+		{
+			// only patch edges that have been modified (to match the existing state of the other mips)
+			if (!EnumHasAnyFlags(NeighborSnapshots.EdgesWithAnyModifiedNeighbor, ToFlag(EdgeIndex)))
+			{
+				continue;
+			}
+
+			// only patch edges that have an existing neighbor
+			if (FHeightmapTextureEdgeSnapshot* NeighborSnapshot = NeighborSnapshots.NeighborSnapshots[(int32)EdgeIndex])
+			{
+				// update this edge using data from neighbor and us
+				TStridedView<UE::Landscape::FHeightmapTexel> DestView;
+				int32 CopyCount = MipEdgeLength - 2;
+
+				if (IsTopOrBottomEdge(EdgeIndex))
+				{
+					// horizontal
+					int32 DstStartX = 1;
+					int32 DstStartY = (EdgeIndex == EEdgeIndex::Bottom) ? 0 : (MipEdgeLength - 1);
+					int32 DstOffset = DstStartX + DstStartY * DstPitchInTexels;
+					DestView = TStridedView<FHeightmapTexel>((int32)sizeof(FHeightmapTexel), &Dest[DstOffset], CopyCount);
+				}
+				else
+				{
+					// vertical
+					int32 DstStartX = (EdgeIndex == EEdgeIndex::Left) ? 0 : (MipEdgeLength - 1);
+					int32 DstStartY = 1;
+					int32 DstOffset = DstStartX + DstStartY * DstPitchInTexels;
+					DestView = TStridedView<FHeightmapTexel>(DstPitchInTexels * sizeof(FHeightmapTexel), &Dest[DstOffset], CopyCount);
+				}
+				BlendEdgeData(*NeighborSnapshots.LocalSnapshot, EdgeIndex, MipIndex, *NeighborSnapshot, DestView);
+				PatchedEdges++;
+
+				FIXUP_DEBUG_LOG_RENDER(TEXT("    Patch Streamed Mip Edge:%s mip:%d len:%d"), *GetDirectionString(EdgeIndex), MipIndex, MipEdgeLength);
+			}
+		}
+	}
+
+	const EEdgeFlags CornersWithOneOrMoreNeighbors = NeighborsToBlendedEdges(NeighborSnapshots.ExistingNeighbors);
+	if (EnumHasAnyFlags(CornersWithOneOrMoreNeighbors, EEdgeFlags::AllCorners))
+	{
+		for (EEdgeIndex CornerIndex : AllCornerNeighbors())
+		{
+			// only patch corners that have been modified (to match the existing state of the other mips)
+			if (!EnumHasAnyFlags(NeighborSnapshots.EdgesWithAnyModifiedNeighbor, ToFlag(CornerIndex)))
+			{
+				continue;
+			}
+
+			// and corners that have at least one existing neighbor
+			if (!EnumHasAnyFlags(CornersWithOneOrMoreNeighbors, ToFlag(CornerIndex)))
+			{
+				continue;
+			}
+
+			int32 DestStride = 0;
+			int32 DestCount = 0;
+			const int32 DestOffset = GetSourceTextureEdgeStartStrideCount(CornerIndex, MipEdgeLength, DestStride, DestCount);
+
+			FHeightmapTexel &Texel = Dest[DestOffset];
+			BlendCornerData(Texel, CornerIndex, NeighborSnapshots);
+			PatchedEdges++;
+
+			FIXUP_DEBUG_LOG_RENDER(TEXT("    Patch Streamed Mip Corner:%s mip:%d len:%d"), *GetDirectionString(CornerIndex), MipIndex, MipEdgeLength);
+		}
+	}
+
+	return PatchedEdges;
+}
+
+#if WITH_EDITOR
+void ULandscapeHeightmapTextureEdgeFixup::RequestEdgeSnapshotUpdateFromHeightmapSource(bool bInUpdateGPUEdgeHashes)
+{
+	if (bInUpdateGPUEdgeHashes)
+	{
+		bUpdateGPUEdgeHashes = true;
+	}
+	check(ActiveGroup->AllRegisteredFixups.Contains(this));
+	ActiveGroup->HeightmapsNeedingEdgeSnapshotCapture.Add(this);
+}
+
+UE::Landscape::ENeighborFlags ULandscapeHeightmapTextureEdgeFixup::UpdateEdgeSnapshotFromHeightmapSource(const FVector& LandscapeGridScale)
+{
+	TRACE_CPUPROFILER_EVENT_SCOPE(UpdateEdgeSnapshotFromHeightmapSource);
+
+	using namespace UE::Landscape;
+
+	ENeighborFlags ChangedEdges = ENeighborFlags::None;
+
+	// no need to update if source ID matches (we assume texture has not actually changed)		
+	FGuid HeightmapSourceID = HeightmapTexture->Source.GetId();
+	if (HeightmapSourceID != EdgeSnapshot->TextureSourceID)
+	{
+		FIXUP_DEBUG_LOG_DETAIL(TEXT("  UpdateEdgeSnapshotFromHeightmapSource %p HeightMap: %p -- UPDATED (ID: %s --> %s)"), ActiveComponent, HeightmapTexture,
+			*EdgeSnapshot->TextureSourceID.ToString(), *HeightmapSourceID.ToString());
+			
+		// create a new edge data
+		TSharedRef<FHeightmapTextureEdgeSnapshot> NewEdgeSnapshot = FHeightmapTextureEdgeSnapshot::CreateEdgeSnapshotFromHeightmapSource(HeightmapTexture, LandscapeGridScale);
+
+		// compare against previous edge data to see what edges have changed (and might cause neighbors to need to patch)
+		ChangedEdges = NewEdgeSnapshot->CompareEdges(EdgeSnapshot.Get());
+
+		// assign the new edge data
+		EdgeSnapshot = NewEdgeSnapshot;
+
+		if (bUpdateGPUEdgeHashes)
+		{
+			GPUEdgeHashes = EdgeSnapshot->InitialEdgeHashes;
+			GPUEdgeModifiedFlags = EEdgeFlags::None;
+
+			FIXUP_DEBUG_LOG_DETAIL(TEXT("    Reset GPU Edge Hashes: [B:%x BR:%x R:%x TR:%x T:%x TL:%x L:%x BL:%x]"),
+				GPUEdgeHashes[0], GPUEdgeHashes[1], GPUEdgeHashes[2], GPUEdgeHashes[3], GPUEdgeHashes[4], GPUEdgeHashes[5], GPUEdgeHashes[6], GPUEdgeHashes[7]);
+		}
+	}
+	else
+	{
+		FIXUP_DEBUG_LOG_DETAIL(TEXT("  UpdateEdgeSnapshotFromHeightmapSource %p HeightMap: %p -- SourceID matches (%s)"), ActiveComponent, HeightmapTexture, *EdgeSnapshot->TextureSourceID.ToString());
+	}
+
+	bUpdateGPUEdgeHashes = false;
+	bDoNotPatchUntilGPUEdgeHashesUpdated = false;
+
+	return ENeighborFlags::None;
+}
+#endif // WITH_EDITOR
+
+void ULandscapeHeightmapTextureEdgeFixup::Serialize(FArchive& Ar)
+{
+	if (Ar.IsLoading())
+	{
+		// To ensure thread-safe CopyOnWrite behavior when loading, we need to allocate a new chunk of memory to serialize into, then replace the existing reference with that
+		TSharedRef<FHeightmapTextureEdgeSnapshot> NewEdgeSnapshot = MakeShared<FHeightmapTextureEdgeSnapshot>();
+		Ar << NewEdgeSnapshot.Get();
+		EdgeSnapshot = NewEdgeSnapshot;
+
+		FIXUP_DEBUG_LOG_DETAIL(TEXT("Loaded EdgeFixup - Snapshot [B:%x BR:%x R:%x TR:%x T:%x TL:%x L:%x BL:%x] (ID: %s)"),
+			EdgeSnapshot->SnapshotEdgeHashes[0], EdgeSnapshot->SnapshotEdgeHashes[1], EdgeSnapshot->SnapshotEdgeHashes[2], EdgeSnapshot->SnapshotEdgeHashes[3],
+			EdgeSnapshot->SnapshotEdgeHashes[4], EdgeSnapshot->SnapshotEdgeHashes[5], EdgeSnapshot->SnapshotEdgeHashes[6], EdgeSnapshot->SnapshotEdgeHashes[7],
+			*EdgeSnapshot->GetTextureSourceIDAsString());
+
+		FIXUP_DEBUG_LOG_DETAIL(TEXT("                 - Initial [B:%x BR:%x R:%x TR:%x T:%x TL:%x L:%x BL:%x]"),
+			EdgeSnapshot->InitialEdgeHashes[0], EdgeSnapshot->InitialEdgeHashes[1], EdgeSnapshot->InitialEdgeHashes[2], EdgeSnapshot->InitialEdgeHashes[3],
+			EdgeSnapshot->InitialEdgeHashes[4], EdgeSnapshot->InitialEdgeHashes[5], EdgeSnapshot->InitialEdgeHashes[6], EdgeSnapshot->InitialEdgeHashes[7]);
+
+		FIXUP_DEBUG_LOG_DETAIL(TEXT("                 - GPU [B:%x BR:%x R:%x TR:%x T:%x TL:%x L:%x BL:%x]"),
+			GPUEdgeHashes[0], GPUEdgeHashes[1], GPUEdgeHashes[2], GPUEdgeHashes[3],
+			GPUEdgeHashes[4], GPUEdgeHashes[5], GPUEdgeHashes[6], GPUEdgeHashes[7]);
+	}
+	else
+	{
+		Ar << EdgeSnapshot.Get();
+
+		FIXUP_DEBUG_LOG_DETAIL(TEXT("  Saved EdgeFixup - Snapshot [B:%x BR:%x R:%x TR:%x T:%x TL:%x L:%x BL:%x] (ID: %s)"),
+			EdgeSnapshot->SnapshotEdgeHashes[0], EdgeSnapshot->SnapshotEdgeHashes[1], EdgeSnapshot->SnapshotEdgeHashes[2], EdgeSnapshot->SnapshotEdgeHashes[3],
+			EdgeSnapshot->SnapshotEdgeHashes[4], EdgeSnapshot->SnapshotEdgeHashes[5], EdgeSnapshot->SnapshotEdgeHashes[6], EdgeSnapshot->SnapshotEdgeHashes[7],
+			*EdgeSnapshot->GetTextureSourceIDAsString());
+
+		FIXUP_DEBUG_LOG_DETAIL(TEXT("                  - Initial [B:%x BR:%x R:%x TR:%x T:%x TL:%x L:%x BL:%x]"),
+			EdgeSnapshot->InitialEdgeHashes[0], EdgeSnapshot->InitialEdgeHashes[1], EdgeSnapshot->InitialEdgeHashes[2], EdgeSnapshot->InitialEdgeHashes[3],
+			EdgeSnapshot->InitialEdgeHashes[4], EdgeSnapshot->InitialEdgeHashes[5], EdgeSnapshot->InitialEdgeHashes[6], EdgeSnapshot->InitialEdgeHashes[7]);
+
+		FIXUP_DEBUG_LOG_DETAIL(TEXT("                  - GPU [B:%x BR:%x R:%x TR:%x T:%x TL:%x L:%x BL:%x]"),
+			GPUEdgeHashes[0], GPUEdgeHashes[1], GPUEdgeHashes[2], GPUEdgeHashes[3],
+			GPUEdgeHashes[4], GPUEdgeHashes[5], GPUEdgeHashes[6], GPUEdgeHashes[7]);
+	}
+}
+
+void ULandscapeHeightmapTextureEdgeFixup::AddReferencedObjects(UObject* InThis, FReferenceCollector& Collector)
+{
+	Super::AddReferencedObjects(InThis, Collector);
+	ULandscapeHeightmapTextureEdgeFixup* const TypedThis = Cast<ULandscapeHeightmapTextureEdgeFixup>(InThis);
+	Collector.AddReferencedObject(TypedThis->HeightmapTexture);
+	Collector.AddReferencedObject(TypedThis->ActiveComponent);
+}
+
+ULandscapeHeightmapTextureEdgeFixup* ULandscapeHeightmapTextureEdgeFixup::FindOrCreateFor(UTexture2D* InHeightmapTexture)
+{
+	check(InHeightmapTexture);
+
+	// try to get an existing factory
+	ULandscapeHeightmapTextureEdgeFixup* Fixup = InHeightmapTexture->GetAssetUserData<ULandscapeHeightmapTextureEdgeFixup>();
+
+#if WITH_EDITORONLY_DATA
+	if (Fixup == nullptr)
+	{
+		check(InHeightmapTexture->Source.IsValid())
+		check(InHeightmapTexture->Source.GetFormat() == TSF_BGRA8);
+
+		// create a new Fixup, and attach it to the texture via User Data (and as outer)
+		Fixup = NewObject<ULandscapeHeightmapTextureEdgeFixup>(InHeightmapTexture);	
+		InHeightmapTexture->AddAssetUserData(Fixup);
+		FIXUP_DEBUG_LOG_DETAIL(TEXT("  FindOrCreateFor texture %p -- CREATE FIXUP"), InHeightmapTexture);
+	}
+	else if (Fixup)
+	{
+		FIXUP_DEBUG_LOG_DETAIL(TEXT("  FindOrCreateFor texture %p -- FOUND FIXUP"), InHeightmapTexture);
+	}
+	else
+	{
+		FIXUP_DEBUG_LOG_DETAIL(TEXT("  FindOrCreateFor texture %p -- FAILED"), InHeightmapTexture);
+	}
+#endif // WITH_EDITORONLY_DATA
+	
+	if (Fixup)
+	{
+		Fixup->SetHeightmapTexture(InHeightmapTexture);
+	}
+
+	return Fixup;
+}
+
+#undef FIXUP_DEBUG_LOG
+#undef FIXUP_DEBUG_LOG_PATCH
+#undef FIXUP_DEBUG_LOG_RENDER
+#undef FIXUP_DEBUG_LOG_DETAIL
diff --git a/Engine/Source/Runtime/Landscape/Private/LandscapeEditLayers.cpp b/Engine/Source/Runtime/Landscape/Private/LandscapeEditLayers.cpp
--- a/Engine/Source/Runtime/Landscape/Private/LandscapeEditLayers.cpp
+++ b/Engine/Source/Runtime/Landscape/Private/LandscapeEditLayers.cpp
@@ -23,6 +23,8 @@ LandscapeEditLayers.cpp: Landscape editing layers mode
 #include "LandscapeUtilsPrivate.h"
 #include "LandscapeSubsystem.h"
 #include "LandscapeTextureStreamingManager.h"
+#include "LandscapeEdgeFixup.h"
+#include "LandscapeGroup.h"
 
 #include "Application/SlateApplicationBase.h"
 #include "Shader.h"
@@ -3226,6 +3228,17 @@ void ExecuteCopyToReadbackTexture(TArray<FLandscapeLayersCopyReadbackTexturePara
 	}
 	for (FLandscapeLayersCopyReadbackTextureParams& Params : InParams)
 	{
+		// stop any GPU texture edge patching on readback components,
+		// until we can update the GPU edge hashes after the readback results are available
+		// (otherwise we might might get incorrect GPU edge hashes in our tracking)
+		if (TObjectPtr<ULandscapeComponent>* ComponentPtr = FLandscapeGroup::HeightmapTextureToActiveComponent.Find(Params.Source))
+		{
+			if (ULandscapeHeightmapTextureEdgeFixup* Fixup = (*ComponentPtr)->RegisteredEdgeFixup)
+			{
+				Fixup->PauseTextureEdgePatchingUntilGPUEdgeHashesUpdated();
+			}
+		}
+
 		Params.Dest->Enqueue(Params.Source, MoveTemp(Params.Context));
 	}
 }
@@ -7802,7 +7815,7 @@ void ALandscape::ResolveLayersHeightmapTexture(
 	TRACE_CPUPROFILER_EVENT_SCOPE(LandscapeLayers_ResolveLayersHeightmapTexture);
 
 	ULandscapeInfo* Info = GetLandscapeInfo();
-	if (Info == nullptr)
+	if (Info == nullptr || HeightmapsToResolve.Num() == 0)
 	{
 		return;
 	}
@@ -7819,6 +7832,17 @@ void ALandscape::ResolveLayersHeightmapTexture(
 			{
 				ChangedComponents.Append(MapHelper.HeightmapToComponents[Heightmap]);
 			}
+
+			// issue the edge update request whether it was changed or not, as we need to update GPU edge hashes anyways
+			if (TObjectPtr<ULandscapeComponent>* ComponentPtr = FLandscapeGroup::HeightmapTextureToActiveComponent.Find(Heightmap))
+			{
+				if (ULandscapeHeightmapTextureEdgeFixup* Fixup = (*ComponentPtr)->RegisteredEdgeFixup)
+				{
+					// since the texture source was just updated via GPU-readback, also update the GPU edge hashes when updating edge data
+					const bool bUpdateGPUEdgeHashes = true;
+					Fixup->RequestEdgeSnapshotUpdateFromHeightmapSource(bUpdateGPUEdgeHashes);
+				}
+			}
 		}
 	}
 
diff --git a/Engine/Source/Runtime/Landscape/Private/LandscapeGroup.cpp b/Engine/Source/Runtime/Landscape/Private/LandscapeGroup.cpp
new file mode 100644
--- /dev/null
+++ b/Engine/Source/Runtime/Landscape/Private/LandscapeGroup.cpp
@@ -0,0 +1,723 @@
+// Copyright Epic Games, Inc. All Rights Reserved.
+
+#include "LandscapeGroup.h"
+#include "CoreMinimal.h"
+#include "Engine/World.h"
+#include "Landscape.h"
+#include "LandscapeStreamingProxy.h"
+#include "LandscapeEdgeFixup.h"
+#include "LandscapePrivate.h"
+#include "LandscapeSubsystem.h"
+
+// enables verbose debug spew
+//#define ENABLE_LANDSCAPE_EDGE_FIXUP_DEBUG_SPEW 1
+#ifdef ENABLE_LANDSCAPE_EDGE_FIXUP_DEBUG_SPEW
+#define GROUP_DEBUG_LOG(...) UE_LOG(LogLandscape, Warning, __VA_ARGS__)
+#define GROUP_DEBUG_LOG_DETAIL(...) UE_LOG(LogLandscape, Warning, __VA_ARGS__)
+#else
+#define GROUP_DEBUG_LOG(...) UE_LOG(LogLandscape, Verbose, __VA_ARGS__)
+#define GROUP_DEBUG_LOG_DETAIL(...) do {} while(0)
+#endif // ENABLE_LANDSCAPE_EDGE_FIXUP_DEBUG_SPEW
+
+TMap<TObjectPtr<UTexture2D>, TObjectPtr<ULandscapeComponent>> FLandscapeGroup::HeightmapTextureToActiveComponent;
+
+namespace UE::Landscape
+{
+	// TODO [chris.tchou] Re-enable once we figure out why xbox edge patching is breaking the texture
+	int32 GInstallEdgeFixup = 0;
+	static FAutoConsoleVariableRef CVarInstallEdgeFixup(
+		TEXT("landscape.InstallEdgeFixup"),
+		GInstallEdgeFixup,
+		TEXT("Controls whether edge fixup tracking is installed on landscape heightmap textures.  Default enabled (1)."));
+
+	bool ShouldInstallEdgeFixup()
+	{
+		return GInstallEdgeFixup != 0;
+	}
+
+	// TODO [chris.tchou] Re-enable once we figure out why xbox edge patching is breaking the texture
+	int32 GPatchEdges = 0;
+	static FAutoConsoleVariableRef CVarPatchEdges(
+		TEXT("landscape.PatchEdges"),
+		GPatchEdges,
+		TEXT("Controls whether landscape heightmap texture edges are patched to match neighboring heightmaps.  Default enabled (1)."));
+
+	// TODO [chris.tchou] Re-enable once we figure out why xbox edge patching is breaking the texture
+	int32 GPatchStreamingMipEdges = 0;
+	static FAutoConsoleVariableRef CVarPatchStreamingMipEdges(
+		TEXT("landscape.PatchStreamingMipEdges"),
+		GPatchStreamingMipEdges,
+		TEXT("Controls whether landscape heightmap texture MIP edges are patched when they stream in.  Default enabled (1)."));
+
+	bool ShouldPatchStreamingMipEdges()
+	{
+		return GPatchStreamingMipEdges != 0;
+	}
+
+	int32 GForcePatchAllEdges = 0;
+	static FAutoConsoleVariableRef CVarForcePatchAllEdges(
+		TEXT("landscape.ForcePatchAllEdges"),
+		GForcePatchAllEdges,
+		TEXT("Forces landscape edge patching to patch the edges of every registered landscape component on the next tick.  Default disabled (0)."));
+
+	bool ShouldPatchAllLandscapeComponentEdges(bool bResetForNext)
+	{
+		bool bPatchAll = (GForcePatchAllEdges != 0);
+		if (bResetForNext)
+		{
+			GForcePatchAllEdges = 0;
+		}
+		return bPatchAll;
+	}
+
+	int32 GAmortizedGroupValidation = 1;
+	static FAutoConsoleVariableRef CVarAmortizedGroupValidation(
+		TEXT("landscape.AmortizedGroupValidation"),
+		GAmortizedGroupValidation,
+		TEXT("Enables amortized validation of landscape group registrations, to verify that components are registered properly with their landscape group.  Default enabled (1)."));
+}
+
+void FLandscapeGroup::RegisterComponent(ULandscapeComponent* Component)
+{
+	if (Component->RegisteredEdgeFixup != nullptr)
+	{
+		// multiple registrations can happen in editor, ignore the extras
+		GROUP_DEBUG_LOG_DETAIL(TEXT("Group %p RegisterComponent %p %d (%d,%d) -- AlreadyRegistered"), this, Component, Component->RegisteredEdgeFixup->bMapped, Component->RegisteredEdgeFixup->GroupCoord.X, Component->RegisteredEdgeFixup->GroupCoord.Y);
+		return;
+	}
+
+	// TODO [chris.tchou] : Remove this once we remove shared heightmaps, so we can enable edge fixup on non-WP landscapes.
+	ALandscapeStreamingProxy* StreamingProxy = Cast<ALandscapeStreamingProxy>(Component->GetOwner());
+	if (!StreamingProxy || !StreamingProxy->bIsLandscapeActorRegisteredWithLandscapeInfo)
+	{
+		// the landscape group only registers components on streaming proxies (we don't handle non-WP cases)
+		// and we can't register until the landscape actor has been registered and shared settings are fixed
+		// (we will re-register all streaming proxy components when the landscape actor gets registered,
+		// to handle any components that were skipped here)
+		GROUP_DEBUG_LOG_DETAIL(TEXT("Group %p RegisterComponent %p -- WaitForActor"), this, Component);
+		return;
+	}
+
+	UTexture2D* Heightmap = Component->GetHeightmap();
+	check(Heightmap);
+
+	// grab an exclusive lock on the group registration data, as registration modifies the neighbor Mapping
+	FWriteScopeLock ScopeWriteLock(RWLock);
+
+	ULandscapeHeightmapTextureEdgeFixup* EdgeFixup = nullptr;
+	bool bIsDisabled = false;
+	{
+		// check if there is an active component handling this heightmap texture
+		if (TObjectPtr<ULandscapeComponent>* FoundComponent = HeightmapTextureToActiveComponent.Find(Heightmap))
+		{
+			ULandscapeComponent* ActiveComponent = FoundComponent->Get();
+			check(ActiveComponent != Component);
+			check(ActiveComponent->RegisteredLandscapeGroup != nullptr);
+
+			UWorld* NewWorld = Component->GetWorld();
+			UWorld* ActiveWorld = ActiveComponent->GetWorld();
+			check(NewWorld != ActiveWorld);
+			
+			// since there is already an active component, edge fixup is already created and registered, we can just grab it
+			EdgeFixup = ActiveComponent->RegisteredEdgeFixup;
+			check(EdgeFixup != nullptr);
+			check(EdgeFixup->HeightmapTexture == Heightmap);
+			check(EdgeFixup->ActiveComponent == ActiveComponent);
+			check(EdgeFixup->ActiveGroup == ActiveComponent->RegisteredLandscapeGroup);
+
+			bIsDisabled = EdgeFixup->DisabledComponents.Contains(Component);
+
+			FLandscapeGroup* ActiveGroup = EdgeFixup->ActiveGroup;
+			check(ActiveGroup != nullptr);
+
+			GROUP_DEBUG_LOG_DETAIL(TEXT("Group %p RegisterComponent %p -- Override Existing EdgeFixup"), this, Component);
+
+			HeightmapTextureToActiveComponent[Heightmap] = Component;
+		}
+		else
+		{
+			// ensure texture user datas are installed
+			GROUP_DEBUG_LOG_DETAIL(TEXT("Group %p RegisterComponent %p -- Activate EdgeFixup"), this, Component);
+
+			const bool bShouldCompressHeightmap = false; // compress only happens at cook time, and once compressed it remains compressed
+			const bool bUseEdgeFixup = true;
+			const bool bUpdateSnapshotNow = false; // by default we wait until Tick to update the snapshot
+			EdgeFixup = Component->InstallOrUpdateTextureUserDatas(bUseEdgeFixup, bShouldCompressHeightmap, bUpdateSnapshotNow);
+			
+			if (EdgeFixup == nullptr)
+			{
+				// failed to install, or InstallEdgeFixup is disabled
+				return;
+			}
+
+			// there should be no disabled components (that requires a collision)
+			check(EdgeFixup->DisabledComponents.Num() == 0);
+			bIsDisabled = false;
+
+			// there should be no active component yet
+			check(EdgeFixup->ActiveGroup == nullptr);
+			check(EdgeFixup->ActiveComponent == nullptr);
+
+			HeightmapTextureToActiveComponent.Add(Heightmap, Component);
+		}
+	}
+
+	check(EdgeFixup);
+	AllRegisteredFixups.Add(EdgeFixup);
+	Component->RegisteredEdgeFixup = EdgeFixup;
+	Component->RegisteredLandscapeGroup = this;
+
+	if (bIsDisabled)
+	{
+		// a disabled component is not set active on register, it must wait for the conflicting component to unregister first
+		GROUP_DEBUG_LOG_DETAIL(TEXT("  -- EdgeFixup Not Active (Disabled)"));
+		check(EdgeFixup->HeightmapTexture == Heightmap);
+	}
+	else
+	{
+		// set active will also map the fixup within the group
+		EdgeFixup->SetActiveComponent(Component, this);
+
+		GROUP_DEBUG_LOG_DETAIL(TEXT("  -- EdgeFixup Set Active At Coord (%d, %d)"), EdgeFixup->GroupCoord.X, EdgeFixup->GroupCoord.Y);
+
+#if WITH_EDITOR
+		// Request an edge snapshot update, and initialize the GPU edge hashes from that.
+		// Also disable patching until this happens.  This is because we may have loaded out of date hashes (if versions changed for instance)
+		// and we should check if we need to recapture new snapshots before doing anything else.
+		HeightmapsNeedingEdgeSnapshotCapture.Add(EdgeFixup);
+		EdgeFixup->bUpdateGPUEdgeHashes = true;
+		EdgeFixup->bDoNotPatchUntilGPUEdgeHashesUpdated = true;
+#endif // WITH_EDITOR
+
+		// Mark it, and all neighbors, as needing edge texture patching.
+		// (the appearance of this new component may cause any of them to need patching)
+		HeightmapsNeedingEdgeTexturePatching.Add(EdgeFixup);
+		for (UE::Landscape::ENeighborIndex NeighborIndex : TEnumRange<UE::Landscape::ENeighborIndex>())
+		{
+			FIntPoint NeighborCoord = EdgeFixup->GroupCoord + GetNeighborRelativePosition(NeighborIndex);
+			if (TObjectPtr<ULandscapeHeightmapTextureEdgeFixup>* NeighborPtrPtr = XYToEdgeFixupMap.Find(NeighborCoord))
+			{
+				HeightmapsNeedingEdgeTexturePatching.Add(NeighborPtrPtr->Get());
+			}
+		}
+	}
+}
+
+void FLandscapeGroup::UnregisterComponent(ULandscapeComponent* Component)
+{
+	ULandscapeHeightmapTextureEdgeFixup* EdgeFixup = Component->RegisteredEdgeFixup;
+	if (EdgeFixup == nullptr)
+	{
+		GROUP_DEBUG_LOG_DETAIL(TEXT("Group %p UnregisterComponent %p -- Already Unregistered"), this, Component);
+		return;
+	}
+	check(Component->RegisteredLandscapeGroup == this);
+
+	// grab an exclusive lock on the group registration data, as registration modifies the neighbor Mapping
+	FWriteScopeLock ScopeWriteLock(RWLock);
+
+	Component->RegisteredEdgeFixup = nullptr;
+	Component->RegisteredLandscapeGroup = nullptr;
+
+	int32 RemovedFixups = AllRegisteredFixups.Remove(EdgeFixup);
+	check(RemovedFixups == 1);
+
+	bool bIsDisabled = EdgeFixup->DisabledComponents.Contains(Component);
+
+	if (bIsDisabled)
+	{
+		// unregistering a disabled component does not change the active component
+		// the disabled component is just removed from the AllRegisteredFixups list.
+		// the disabled component remains on the DisabledComponents list,
+		// so that if it is re-registered it stay disabled if there is another active component
+		GROUP_DEBUG_LOG_DETAIL(TEXT("Group %p UnregisterComponent %p %d (%d,%d) -- Was Disabled"), this, Component, EdgeFixup->bMapped, EdgeFixup->GroupCoord.X, EdgeFixup->GroupCoord.Y);
+		check(EdgeFixup->ActiveComponent != Component);
+	}
+	else
+	{
+		check(EdgeFixup->ActiveComponent == Component);
+
+		// first let's see if there is a replacement disabled component we can activate in place of the one being unregistered
+		bool bReactivatedDisabled = false;
+		while (EdgeFixup->DisabledComponents.Num() > 0)
+		{
+			// if a disabled duplicate exists (and is still registered), then reactivate it!
+			// any disabled components that are no longer registered we can just drop
+			if (ULandscapeComponent* DisabledComponent = EdgeFixup->DisabledComponents.Pop().Get())
+			{
+				if (DisabledComponent->RegisteredEdgeFixup != nullptr)
+				{
+					check(DisabledComponent->RegisteredEdgeFixup == EdgeFixup);	// sanity check
+					FLandscapeGroup* DisabledGroup = DisabledComponent->RegisteredLandscapeGroup;
+					check(DisabledGroup != nullptr);
+
+					GROUP_DEBUG_LOG_DETAIL(TEXT("Group %p UnregisterComponent %p %d (%d,%d) -- Reactivate Other Disabled"), this, Component, EdgeFixup->bMapped, EdgeFixup->GroupCoord.X, EdgeFixup->GroupCoord.Y);
+
+					const bool bDisableCurrentActive = false; // just unmap it
+					EdgeFixup->SetActiveComponent(DisabledComponent, DisabledGroup, bDisableCurrentActive);
+
+					bReactivatedDisabled = true;
+					break;
+				}
+			}
+		}
+
+		if (!bReactivatedDisabled)
+		{
+			GROUP_DEBUG_LOG_DETAIL(TEXT("Group %p UnregisterComponent %p %d (%d,%d) -- Deactivating"), this, Component, EdgeFixup->bMapped, EdgeFixup->GroupCoord.X, EdgeFixup->GroupCoord.Y);
+
+			// no disabled available to replace the active component -- just set NO active component
+			Unmap(EdgeFixup);
+
+			GROUP_DEBUG_LOG_DETAIL(TEXT("  Unmap %p (%d,%d) <Group %d Key %d> edgefixup:%p MAPPED:%d TOTAL:%d"),
+				EdgeFixup->ActiveComponent, EdgeFixup->GroupCoord.X, EdgeFixup->GroupCoord.Y, EdgeFixup->ActiveComponent->GetLandscapeProxy()->LODGroupKey, EdgeFixup->ActiveGroup->LandscapeGroupKey, this,
+				EdgeFixup->ActiveGroup->XYToEdgeFixupMap.Num(), EdgeFixup->ActiveGroup->AllRegisteredFixups.Num());
+
+			EdgeFixup->ActiveComponent = nullptr;
+			EdgeFixup->ActiveGroup = nullptr;
+
+			HeightmapTextureToActiveComponent.Remove(EdgeFixup->HeightmapTexture);
+		}
+	}
+}
+
+void FLandscapeGroup::DisableAndUnmap(ULandscapeHeightmapTextureEdgeFixup* Fixup)
+{
+	// should be registered before it can be disabled
+	check(AllRegisteredFixups.Contains(Fixup));
+
+	// you can't disable the same component twice
+	check(!Fixup->DisabledComponents.Contains(Fixup->ActiveComponent));
+	Fixup->DisabledComponents.Add(Fixup->ActiveComponent);
+
+	Unmap(Fixup);
+}
+
+FIntPoint FLandscapeGroup::Map(ULandscapeHeightmapTextureEdgeFixup* Fixup, ULandscapeComponent* Component)
+{
+	check(Fixup->bMapped == false);
+
+	const FTransform& ComponentLocalToWorld = Component->GetComponentTransform();
+	FBoxSphereBounds ComponentLocalBounds = Component->CalcBounds(FTransform::Identity);
+	FVector ComponentCenterWorldSpace = ComponentLocalToWorld.TransformPosition(ComponentLocalBounds.Origin);
+	FVector ComponentXVectorWorldSpace = ComponentLocalToWorld.TransformVector(FVector::XAxisVector);
+	FVector ComponentYVectorWorldSpace = ComponentLocalToWorld.TransformVector(FVector::YAxisVector);
+	FVector ComponentLandscapeGridScale = Component->GetLandscapeProxy()->GetRootComponent()->GetRelativeScale3D();
+
+	if (XYToEdgeFixupMap.Num() == 0)
+	{
+		// the first registered section gets to set up our RenderCoord grid so it is located at the origin
+		ComponentResolution = Component->ComponentSizeQuads;
+		GroupCoordOrigin = ComponentCenterWorldSpace;
+		GroupCoordXVector = ComponentXVectorWorldSpace;
+		GroupCoordYVector = ComponentYVectorWorldSpace;
+		LandscapeGridScale = ComponentLandscapeGridScale;
+	}
+	else
+	{
+		// validate each additional section has a matching resolution, scale and orientation
+		bool bResolutionMatches = (ComponentResolution == Component->ComponentSizeQuads);
+		bool bXVectorMatches = (ComponentXVectorWorldSpace - GroupCoordXVector).IsNearlyZero();
+		bool bYVectorMatches = (ComponentYVectorWorldSpace - GroupCoordYVector).IsNearlyZero();
+		if (!(bResolutionMatches && bXVectorMatches && bYVectorMatches))
+		{
+			UE_LOG(LogLandscape, Warning, TEXT("Landscapes in LOD Group with Key %d do not have matching resolution (%d == %d), scale (%f == %f, %f == %f) and/or rotation; seam artifacts may appear."),
+				LandscapeGroupKey,
+				ComponentResolution, Component->ComponentSizeQuads,
+				GroupCoordXVector.Length(), ComponentXVectorWorldSpace.Length(),
+				GroupCoordYVector.Length(), ComponentYVectorWorldSpace.Length());
+		}
+
+		if (!(LandscapeGridScale - ComponentLandscapeGridScale).IsNearlyZero())
+		{
+			UE_LOG(LogLandscape, Warning, TEXT("Landscapes in LOD Group with Key %d do not have matching grid scale (%s == %s); seam artifacts may appear."),
+				LandscapeGroupKey,
+				*LandscapeGridScale.ToString(), *ComponentLandscapeGridScale.ToString());
+		}
+	}
+
+	// project onto the Component X/Y plane to calculate the group coordinates
+	FVector Delta = ComponentCenterWorldSpace - GroupCoordOrigin;
+	double GroupCoordX = Delta.Dot(GroupCoordXVector) / (GroupCoordXVector.SquaredLength() * ComponentResolution);
+	double GroupCoordY = Delta.Dot(GroupCoordYVector) / (GroupCoordYVector.SquaredLength() * ComponentResolution);
+
+	FIntPoint GroupCoord;
+	GroupCoord.X = FMath::RoundToInt32(GroupCoordX);
+	GroupCoord.Y = FMath::RoundToInt32(GroupCoordY);
+
+	if ((FMath::Abs(GroupCoordX - GroupCoord.X) > 0.01) || (FMath::Abs(GroupCoordY - GroupCoord.Y) > 0.01))
+	{
+		UE_LOG(LogLandscape, Warning, TEXT("Landscape component %s is not spatially aligned with the group grid (%d), seam artifacts may appear."), *Component->GetPathName(), LandscapeGroupKey);
+	}
+
+	// Add to the Map
+	while (true)
+	{
+		TObjectPtr<ULandscapeHeightmapTextureEdgeFixup>* MapEntryPtrPtr = XYToEdgeFixupMap.Find(GroupCoord);
+
+		if (MapEntryPtrPtr == nullptr)
+		{
+			XYToEdgeFixupMap.Add(GroupCoord) = Fixup;
+			break;
+		}
+		else if (*MapEntryPtrPtr == Fixup)
+		{
+			break;
+		}
+		else
+		{
+			UE_LOG(LogLandscape, Warning, TEXT("Two landscape components in group (%d) occupy the same group grid cell: %s and %s, artifacts may appear. Please move one."),
+				LandscapeGroupKey, *Component->GetPathName(), *(*MapEntryPtrPtr)->ActiveComponent->GetPathName());
+
+			// attempt to give a temporary location to the overlapped component
+			// We could potentially use a MultiMap instead so we don't have to hack the GroupCoord.. 
+			// TODO [chris.tchou] : when overlapped components move, we should also re-register the non-moved one.
+			GroupCoord.X += 100000;
+		}
+	}
+
+	Fixup->bMapped = true;
+	Fixup->GroupCoord = GroupCoord;
+
+	return GroupCoord;
+}
+
+void FLandscapeGroup::Unmap(ULandscapeHeightmapTextureEdgeFixup* Fixup)
+{
+	check(Fixup->bMapped == true);
+	check(Fixup->ActiveGroup == this);
+
+	FIntPoint GroupCoord = Fixup->GroupCoord;
+	TObjectPtr<ULandscapeHeightmapTextureEdgeFixup> RemovedEdgeFixup = nullptr;
+	XYToEdgeFixupMap.RemoveAndCopyValue(GroupCoord, RemovedEdgeFixup);
+	check(Fixup == RemovedEdgeFixup);
+
+#if WITH_EDITOR
+	HeightmapsNeedingEdgeSnapshotCapture.Remove(Fixup);
+#endif // WITH_EDITOR
+	HeightmapsNeedingEdgeTexturePatching.Remove(Fixup);
+
+	Fixup->bMapped = false;
+}
+
+void FLandscapeGroup::RegisterAllComponentsOnStreamingProxy(ALandscapeStreamingProxy* StreamingProxy)
+{
+	check(StreamingProxy->bIsLandscapeActorRegisteredWithLandscapeInfo);
+
+	UWorld* World = StreamingProxy->GetWorld();
+	check(World);
+
+	GROUP_DEBUG_LOG_DETAIL(TEXT("Registering components for World %p %s (Scene: %p WorldType: %d NetMode: %d)"),
+		World, *World->GetName(), World->Scene, World->WorldType, World->GetNetMode());
+
+	// if world is not renderable, we don't need the landscape groups / edge patching
+	if (World->Scene == nullptr)
+	{
+		return;
+	}
+
+	if (World->GetNetMode() == ENetMode::NM_DedicatedServer)
+	{
+		return;
+	}
+
+	if (ULandscapeSubsystem* LandscapeSubsystem = World->GetSubsystem<ULandscapeSubsystem>())
+	{
+		if (FLandscapeGroup* Group = LandscapeSubsystem->GetLandscapeGroupForProxy(StreamingProxy))
+		{
+			for (ULandscapeComponent* Component : StreamingProxy->LandscapeComponents)
+			{
+				if (Component != nullptr)
+				{
+					Group->RegisterComponent(Component);
+				}
+			}
+		}
+	}
+}
+
+void FLandscapeGroup::UnregisterAllComponentsOnStreamingProxy(ALandscapeStreamingProxy* StreamingProxy)
+{
+	UWorld* World = StreamingProxy->GetWorld();
+	check(World);
+
+	if (ULandscapeSubsystem* LandscapeSubsystem = World->GetSubsystem<ULandscapeSubsystem>())
+	{
+		for (ULandscapeComponent* Component : StreamingProxy->LandscapeComponents)
+		{
+			if (Component && Component->RegisteredLandscapeGroup)
+			{
+				Component->RegisteredLandscapeGroup->UnregisterComponent(Component);
+			}
+		}
+	}
+}
+
+void FLandscapeGroup::AddReferencedObjects(FLandscapeGroup* InThis, FReferenceCollector& Collector)
+{
+	Collector.AddReferencedObjects(InThis->AllRegisteredFixups);
+#if WITH_EDITOR
+	check(InThis->AllRegisteredFixups.Includes(InThis->HeightmapsNeedingEdgeSnapshotCapture));
+	check(InThis->AllRegisteredFixups.Includes(InThis->HeightmapsNeedingEdgeTexturePatching));
+#endif // WITH_EDITOR
+	Collector.AddReferencedObjects(InThis->XYToEdgeFixupMap);
+}
+
+
+FString DirectionFlagsToString(UE::Landscape::EDirectionFlags Flags)
+{
+	using namespace UE::Landscape;
+	if (Flags == EDirectionFlags::None)
+	{
+		return TEXT("None");
+	}
+	else
+	{
+		TStringBuilder<64> Builder;
+		bool bFirst = true;
+		for (EDirectionIndex Index : TEnumRange<EDirectionIndex>())
+		{
+			if (EnumHasAnyFlags(Flags, ToFlag(Index)))
+			{
+				if (bFirst)
+				{
+					bFirst = false;
+				}
+				else
+				{
+					Builder.Append(",", 1);
+				}
+				Builder.Append(GetDirectionString(Index));
+			}
+		}
+		return Builder.ToString();
+	}
+}
+
+void FLandscapeGroup::TickEdgeFixup(ULandscapeSubsystem* LandscapeSubsystem, bool bForcePatchAll)
+{
+	TRACE_CPUPROFILER_EVENT_SCOPE(FLandscapeGroup::TickEdgeFixup);
+
+	using namespace UE::Landscape;
+
+	if (XYToEdgeFixupMap.Num() == 0)
+	{
+		return;
+	}
+
+#if !UE_BUILD_SHIPPING
+	if (GAmortizedGroupValidation)
+	{
+		TRACE_CPUPROFILER_EVENT_SCOPE(AmortizedGroupValidation);
+
+		// perform an amortized check on all MAPPED fixups to catch any discrepancies
+		// this shouldn't be necessary, but is a good double check that nothing
+		// fell through the cracks.
+
+		// Grab the Nth entry in the fixup map, if it exists, and check it
+		AmortizeIndex = FSetElementId::FromInteger(AmortizeIndex.AsInteger() + 1);
+		if (AmortizeIndex.AsInteger() > XYToEdgeFixupMap.GetMaxIndex())
+		{
+			AmortizeIndex = FSetElementId::FromInteger(0);
+		}
+		if (XYToEdgeFixupMap.IsValidId(AmortizeIndex))
+		{
+			auto Pair = XYToEdgeFixupMap.Get(AmortizeIndex);
+			ULandscapeHeightmapTextureEdgeFixup* Fixup = Pair.Value;
+
+			// any entry in the group's map should be mapped, and this group should be the active one
+			check(Fixup->bMapped);
+			check(Fixup->ActiveGroup == this);
+			ULandscapeComponent* Component = Fixup->ActiveComponent;
+			check(Component);
+			check(Component->RegisteredLandscapeGroup == this);
+			check(Component->RegisteredEdgeFixup == Fixup);
+
+			FVector CurrentLandscapeGridScale = Component->GetLandscapeProxy()->GetRootComponent()->GetRelativeScale3D();
+			if (!ensure((CurrentLandscapeGridScale - LandscapeGridScale).IsNearlyZero()))
+			{
+				UE_LOG(LogLandscape, Warning, TEXT("Landscape component %p (%s) has a different scale than other components in the Landscape Group, normal seam artifacts may occur"), Component, *Component->GetName());
+			}
+
+			FLandscapeGroup* NewGroup = LandscapeSubsystem->GetLandscapeGroupForComponent(Component);
+			if (!ensure(NewGroup == this))
+			{
+				UE_LOG(LogLandscape, Warning, TEXT("Landscape component %p (%s) changed groups unexpectedly (%x ==> %x) - attempting to fix automatically by re-registering it"), Component, *Component->GetName(), this->LandscapeGroupKey, NewGroup->LandscapeGroupKey);
+				this->UnregisterComponent(Component);
+				NewGroup->RegisterComponent(Component);
+			}
+		}
+	}
+#endif // !UE_BUILD_SHIPPING
+
+	if (bForcePatchAll)
+	{
+		for (ULandscapeHeightmapTextureEdgeFixup* Fixup : AllRegisteredFixups)
+		{
+			if (Fixup->ActiveGroup == this)
+			{
+				// mark all edges modified, but with the wrong hash, so they will all patch
+				check(Fixup->bMapped);
+				Fixup->GPUEdgeModifiedFlags = EEdgeFlags::All;
+				memset(Fixup->GPUEdgeHashes.GetData(), 0, sizeof(uint32) * 8);
+				HeightmapsNeedingEdgeTexturePatching.Add(Fixup);
+			}
+		}
+	}
+
+	int32 SnapshotsCaptured = 0;
+	int32 SnapshotEdgesChanged = 0;
+	int32 EdgesPatched = 0;
+	int32 TexturesPatched = 0;
+
+	// grab an exclusive lock on the group registration data, as we may modify snapshots and/or patch hash tracking
+	FWriteScopeLock ScopeWriteLock(RWLock);
+
+#if WITH_EDITOR
+	// first we check if any edge data needs to be updated
+	if (HeightmapsNeedingEdgeSnapshotCapture.Num() > 0)
+	{
+		TRACE_CPUPROFILER_EVENT_SCOPE(HeightmapsNeedingEdgeSnapshotCapture);
+
+		for (auto It = HeightmapsNeedingEdgeSnapshotCapture.CreateIterator(); It; ++It)
+		{
+			ULandscapeHeightmapTextureEdgeFixup* EdgeFixup = *It;
+			check(EdgeFixup->ActiveGroup == this); // should never be in HeightmapsNeedingEdgeSnapshotCapture unless active in this group
+			check(AllRegisteredFixups.Contains(EdgeFixup));
+			check(EdgeFixup->bMapped);
+
+#if WITH_EDITORONLY_DATA
+			ALandscape* Landscape = EdgeFixup->ActiveComponent->GetLandscapeActor();
+			check(Landscape);
+			if (!Landscape->bGrassUpdateEnabled)
+			{
+				GROUP_DEBUG_LOG_DETAIL(TEXT("-- Snap Component %p (%d,%d) -------- PAUSED (editing)"), EdgeFixup->ActiveComponent, EdgeFixup->GroupCoord.X, EdgeFixup->GroupCoord.Y);
+				continue;
+			}
+#endif // WITH_EDITORONLY_DATA
+
+			EEdgeFlags ChangedEdges = EdgeFixup->UpdateEdgeSnapshotFromHeightmapSource(LandscapeGridScale);
+			SnapshotsCaptured++;
+			SnapshotEdgesChanged += FGenericPlatformMath::CountBits((uint64)ChangedEdges);
+
+			GROUP_DEBUG_LOG_DETAIL(TEXT("-- Snap Component %p (%d,%d) -------- UPDATED, Changed edges: %s"), EdgeFixup->ActiveComponent, EdgeFixup->GroupCoord.X, EdgeFixup->GroupCoord.Y, *DirectionFlagsToString(ChangedEdges));
+			if (ChangedEdges != EEdgeFlags::None)
+			{
+				// if any edges were changed -- queue it for texture edge patching
+				HeightmapsNeedingEdgeTexturePatching.Add(EdgeFixup);
+
+				ENeighborFlags AffectedNeighbors = EdgesToAffectedNeighbors(ChangedEdges);
+
+				// and also mark the affected neighbors for potential texture edge patching
+				EdgeFixup->RequestEdgeTexturePatchingForNeighbors(AffectedNeighbors);
+			}
+			It.RemoveCurrent();
+		}
+	}
+#endif // WITH_EDITOR
+
+	// now we try to fix up the heightmap GPU texture to address any edges that need to be patched
+	if ((bForcePatchAll || GPatchEdges) && (HeightmapsNeedingEdgeTexturePatching.Num() > 0))
+	{
+		TRACE_CPUPROFILER_EVENT_SCOPE(HeightmapsNeedingEdgeTexturePatching);
+
+		for (auto It = HeightmapsNeedingEdgeTexturePatching.CreateIterator(); It; ++It)
+		{
+			TRACE_CPUPROFILER_EVENT_SCOPE(CheckHeightmap);
+
+			ULandscapeHeightmapTextureEdgeFixup* EdgeFixup = *It;
+			check(EdgeFixup->ActiveGroup == this); // should never be in HeightmapsNeedingEdgeTexturePatching unless active in this group
+			check(AllRegisteredFixups.Contains(EdgeFixup));
+			check(EdgeFixup->bMapped);
+
+#if WITH_EDITORONLY_DATA
+			ALandscape* Landscape = EdgeFixup->ActiveComponent->GetLandscapeActor();
+			check(Landscape);
+			if (!Landscape->bGrassUpdateEnabled)
+			{
+				GROUP_DEBUG_LOG_DETAIL(TEXT("-- Patch Component %p (%d,%d) -------- PAUSED (editing)"), EdgeFixup->ActiveComponent, EdgeFixup->GroupCoord.X, EdgeFixup->GroupCoord.Y);
+				continue;
+			}
+#endif // WITH_EDITORONLY_DATA
+
+			if (EdgeFixup->HeightmapTexture->HasPendingInitOrStreaming())
+			{
+				GROUP_DEBUG_LOG_DETAIL(TEXT("-- Patch Component %p (%d,%d) -------- PAUSED (pending)"), EdgeFixup->ActiveComponent, EdgeFixup->GroupCoord.X, EdgeFixup->GroupCoord.Y);
+				continue;
+			}
+
+#if WITH_EDITOR
+			if (EdgeFixup->IsTextureEdgePatchingPaused())
+			{
+				GROUP_DEBUG_LOG_DETAIL(TEXT("-- Patch Component %p (%d,%d) -------- PAUSED (readback)"), EdgeFixup->ActiveComponent, EdgeFixup->GroupCoord.X, EdgeFixup->GroupCoord.Y);
+				continue;
+			}
+
+			// wait til we are not using the default texture to apply edge patching
+			if (EdgeFixup->GetHeightmapTexture()->IsDefaultTexture())
+			{
+				GROUP_DEBUG_LOG_DETAIL(TEXT("-- Patch Component %p (%d,%d) -------- WAITING (texture load)"), EdgeFixup->ActiveComponent, EdgeFixup->GroupCoord.X, EdgeFixup->GroupCoord.Y);
+				continue;
+			}
+#endif // WITH_EDITOR
+
+			check(EdgeFixup->ActiveGroup == this);
+			int32 PatchedEdgeCount = EdgeFixup->CheckAndPatchTextureEdgesFromEdgeSnapshots();
+			if (PatchedEdgeCount > 0)
+			{
+				TexturesPatched++;
+			}
+			EdgesPatched += PatchedEdgeCount;
+
+			It.RemoveCurrent();
+			GROUP_DEBUG_LOG_DETAIL(TEXT("-- Patch Component %p (%d,%d) -------- PATCHED %d edges"), EdgeFixup->ActiveComponent, EdgeFixup->GroupCoord.X, EdgeFixup->GroupCoord.Y, PatchedEdgeCount);
+		}
+	}
+
+	if ((SnapshotsCaptured > 0) || (EdgesPatched > 0))
+	{
+		GROUP_DEBUG_LOG(TEXT("LandscapeGroup Tick [%d Snapshots Captured, %d edges changed) (%d edges patched on %d textures)"), SnapshotsCaptured, SnapshotEdgesChanged, EdgesPatched, TexturesPatched);
+	}
+}
+
+
+FLandscapeGroup::~FLandscapeGroup()
+{
+	if (AllRegisteredFixups.Num() > 0)
+	{
+		UE_LOG(LogLandscape, Warning, TEXT("Landscape Group (%d) had registered components at destruction, this indicates some components were not properly unregistered"), this->LandscapeGroupKey);
+		for (ULandscapeHeightmapTextureEdgeFixup* Fixup : AllRegisteredFixups)
+		{
+			if (Fixup->ActiveGroup == this)
+			{
+				// if we're the active group, the corresponding component is the active one
+				UnregisterComponent(Fixup->ActiveComponent);
+			}
+			else
+			{
+				// find our group's component by searching through the disabled list
+				for (TWeakObjectPtr<ULandscapeComponent> DisabledComponentPtr : Fixup->DisabledComponents)
+				{
+					if (ULandscapeComponent* DisabledComponent = DisabledComponentPtr.Get())
+					{
+						if (DisabledComponent->RegisteredLandscapeGroup == this)
+						{
+							// found it!						
+							UnregisterComponent(DisabledComponent);
+							break;
+						}
+					}
+				}
+			}
+		}
+		check(AllRegisteredFixups.Num() == 0);
+	}
+
+#if WITH_EDITOR
+	check(HeightmapsNeedingEdgeSnapshotCapture.IsEmpty());
+#endif // WITH_EDITOR
+	check(HeightmapsNeedingEdgeTexturePatching.IsEmpty());
+}
+
+#undef GROUP_DEBUG_LOG
+#undef GROUP_DEBUG_LOG_DETAIL
diff --git a/Engine/Source/Runtime/Landscape/Private/LandscapeGroup.h b/Engine/Source/Runtime/Landscape/Private/LandscapeGroup.h
new file mode 100644
--- /dev/null
+++ b/Engine/Source/Runtime/Landscape/Private/LandscapeGroup.h
@@ -0,0 +1,96 @@
+// Copyright Epic Games, Inc. All Rights Reserved.
+
+#pragma once
+
+#include "CoreMinimal.h"
+#include "LandscapeEdgeFixup.h"
+
+class ULandscapeHeightmapTextureEdgeFixup;
+class ULandscapeComponent;
+class ULandscapeSubsystem;
+class UTexture2D;
+class ALandscapeStreamingProxy;
+
+namespace UE::Landscape
+{
+	bool ShouldInstallEdgeFixup();
+	bool ShouldPatchStreamingMipEdges();
+	bool ShouldPatchAllLandscapeComponentEdges(bool bResetForNext = false);
+}
+
+struct FLandscapeGroup
+{
+	friend class ULandscapeHeightmapTextureEdgeFixup;
+
+	// read-write lock that protect access to XYToEdgeFixupMap, and registered ULandscapeHeightmapTextureEdgeFixup's Snapshots, EdgeModifiedFlags and GPUEdgeHashes
+	FRWLock RWLock;
+
+	uint32 LandscapeGroupKey = 0;
+
+	// Resolution, Origin and Size, defines the group coordinate system
+	// (the first component registered gets to be the origin)
+	int32 ComponentResolution = -1;
+	FVector GroupCoordOrigin = FVector::ZeroVector;		// world space position of the center of the origin component (render coord 0,0)
+	FVector GroupCoordXVector = FVector::ZeroVector;	// world space vector in the direction of component local X
+	FVector GroupCoordYVector = FVector::ZeroVector;	// world space vector in the direction of component local Y	
+	FVector LandscapeGridScale = FVector::ZeroVector;	// scale used to calculate normals, value set at the same time as above
+	
+	// map of fixups (key is group coordinate)
+	TMap<FIntPoint, TObjectPtr<ULandscapeHeightmapTextureEdgeFixup>> XYToEdgeFixupMap;
+
+	// global map of heightmap textures to the landscape active component
+	// (this is used to detect and handle shared heightmaps)
+	static TMap<TObjectPtr<UTexture2D>, TObjectPtr<ULandscapeComponent>> HeightmapTextureToActiveComponent;
+
+	// set of fixups that are registered with this group (some of which may not be mapped)
+	// this is only used to double check correct behavior
+	TSet<TObjectPtr<ULandscapeHeightmapTextureEdgeFixup>> AllRegisteredFixups;
+
+#if WITH_EDITOR
+	// fixups that need to capture new edge snapshots
+	TSet<TObjectPtr<ULandscapeHeightmapTextureEdgeFixup>> HeightmapsNeedingEdgeSnapshotCapture;
+#endif // WITH_EDITOR
+
+	// fixups that may need to GPU edge patch their heightmap textures
+	TSet<TObjectPtr<ULandscapeHeightmapTextureEdgeFixup>> HeightmapsNeedingEdgeTexturePatching;
+
+	// amortization state for amortized slow checks
+	FSetElementId AmortizeIndex;
+
+	FLandscapeGroup(uint32 InLandscapeGroupKey)
+		: LandscapeGroupKey(InLandscapeGroupKey)
+	{
+	}
+
+	~FLandscapeGroup();
+
+	void RegisterComponent(ULandscapeComponent* Component);
+	void UnregisterComponent(ULandscapeComponent* Component);
+
+	static void RegisterAllComponentsOnStreamingProxy(ALandscapeStreamingProxy* StreamingProxy);
+	static void UnregisterAllComponentsOnStreamingProxy(ALandscapeStreamingProxy* StreamingProxy);
+	static void AddReferencedObjects(FLandscapeGroup* InThis, FReferenceCollector& Collector);
+
+	ULandscapeHeightmapTextureEdgeFixup* GetEdgeFixupAtCoord(FIntPoint Coord)
+	{
+		TObjectPtr<ULandscapeHeightmapTextureEdgeFixup>* Result = XYToEdgeFixupMap.Find(Coord);
+		if (Result != nullptr)
+		{
+			return Result->Get();
+		}
+		return nullptr;
+	}
+
+	ULandscapeHeightmapTextureEdgeFixup* GetNeighborEdgeFixup(FIntPoint Coord, UE::Landscape::ENeighborIndex NeighborIndex)
+	{
+		FIntPoint NeighborCoord = Coord + UE::Landscape::GetNeighborRelativePosition(NeighborIndex);
+		return GetEdgeFixupAtCoord(NeighborCoord);
+	}
+
+	void TickEdgeFixup(ULandscapeSubsystem* LandscapeSubsystem, bool bForcePatchAll);
+
+private:
+	void DisableAndUnmap(ULandscapeHeightmapTextureEdgeFixup* Fixup);
+	FIntPoint Map(ULandscapeHeightmapTextureEdgeFixup* Fixup, ULandscapeComponent* Component);
+	void Unmap(ULandscapeHeightmapTextureEdgeFixup* Fixup);
+};
diff --git a/Engine/Source/Runtime/Landscape/Private/LandscapeRender.cpp b/Engine/Source/Runtime/Landscape/Private/LandscapeRender.cpp
--- a/Engine/Source/Runtime/Landscape/Private/LandscapeRender.cpp
+++ b/Engine/Source/Runtime/Landscape/Private/LandscapeRender.cpp
@@ -694,7 +694,7 @@ void FLandscapeRenderSystem::CreateResources_Internal(FRHICommandListBase& RHICm
 		// project onto the Component X/Y plane to calculate the render coordinates
 		FVector Delta = SectionCenterWorldSpace - ComponentOrigin;
 		SectionInfo->RenderCoord.X = FMath::RoundToInt32(Delta.Dot(ComponentXVector) / ComponentXVector.SquaredLength());
-		SectionInfo->RenderCoord.Y = FMath::RoundToInt32(Delta.Dot(ComponentYVector) / ComponentXVector.SquaredLength());
+		SectionInfo->RenderCoord.Y = FMath::RoundToInt32(Delta.Dot(ComponentYVector) / ComponentYVector.SquaredLength());
 	}
 
 	// we changed the RenderCoord, need to update the uniform buffer
diff --git a/Engine/Source/Runtime/Landscape/Private/LandscapeSubsystem.cpp b/Engine/Source/Runtime/Landscape/Private/LandscapeSubsystem.cpp
--- a/Engine/Source/Runtime/Landscape/Private/LandscapeSubsystem.cpp
+++ b/Engine/Source/Runtime/Landscape/Private/LandscapeSubsystem.cpp
@@ -19,6 +19,9 @@
 #include "LandscapePrivate.h"
 #include "LandscapeSettings.h"
 #include "LandscapeGrassMapsBuilder.h"
+#include "LandscapeTextureStorageProvider.h"
+#include "LandscapeComponent.h"
+#include "LandscapeGroup.h"
 #include "ProfilingDebugging/CsvProfiler.h"
 #include "WorldPartition/WorldPartitionSubsystem.h"
 #include "ActorPartition/ActorPartitionSubsystem.h"
@@ -142,29 +145,67 @@ void ULandscapeSubsystem::RegisterActor(ALandscapeProxy* Proxy)
 {
 	check(Proxy != nullptr);
 
- 	TObjectPtr<ALandscapeProxy> ProxyPtr(Proxy);
+	TObjectPtr<ALandscapeProxy> ProxyPtr(Proxy);
 
-	// editor can get multiple registration calls, ensure we don't register more than once
- 	Proxies.AddUnique(ProxyPtr);
+	// in editor we can get multiple registration calls, so ensure we don't register more than once
+	if (Proxy->bIsRegisteredWithSubsystem)
+	{
+		checkSlow(Proxies.Contains(ProxyPtr));
+		return;
+	}
+	checkSlow(!Proxies.Contains(ProxyPtr));
+	Proxies.Add(ProxyPtr);
 	
 	if (ALandscape* LandscapeActor = Cast<ALandscape>(Proxy))
 	{
 		TObjectPtr<ALandscape> LandscapeActorPtr(LandscapeActor);
 		LandscapeActors.AddUnique(LandscapeActorPtr);
 	}
+	else if (ALandscapeStreamingProxy* StreamingProxy = Cast<ALandscapeStreamingProxy>(Proxy))
+	{
+		// We want to know when a streaming proxy is moved, so we can adjust the map
+		StreamingProxy->GetRootComponent()->TransformUpdated.AddUObject(this, &ULandscapeSubsystem::OnProxyMoved);
+	}
+
+	Proxy->bIsRegisteredWithSubsystem = true;
 }
 
 void ULandscapeSubsystem::UnregisterActor(ALandscapeProxy* Proxy)
 {
 	check(Proxy != nullptr);
 	TObjectPtr<ALandscapeProxy> ProxyPtr(Proxy);
-	Proxies.Remove(ProxyPtr);
+
+	// in editor we can get multiple unregistration calls, so ensure we don't register more than once
+	if (!Proxy->bIsRegisteredWithSubsystem)
+	{
+		check(!Proxies.Contains(ProxyPtr));
+		return;
+	}
+	int32 RemovedCount = Proxies.Remove(ProxyPtr);
+	check(RemovedCount == 1);
 
 	if (ALandscape* LandscapeActor = Cast<ALandscape>(Proxy))
 	{
 		TObjectPtr<ALandscape> LandscapeActorPtr(LandscapeActor);
 		LandscapeActors.Remove(LandscapeActorPtr);
 	}
+	else if (ALandscapeStreamingProxy* StreamingProxy = Cast<ALandscapeStreamingProxy>(Proxy))
+	{
+		StreamingProxiesNeedingReregister.Remove(StreamingProxy);
+	}
+
+	Proxy->bIsRegisteredWithSubsystem = false;
+}
+
+void ULandscapeSubsystem::AddReferencedObjects(UObject* InThis, FReferenceCollector& Collector)
+{
+	Super::AddReferencedObjects(InThis, Collector);
+	ULandscapeSubsystem* const TypedThis = Cast<ULandscapeSubsystem>(InThis);
+
+	for (auto Pair : TypedThis->Groups)
+	{
+		FLandscapeGroup::AddReferencedObjects(Pair.Value, Collector);
+	}
 }
 
 void ULandscapeSubsystem::Initialize(FSubsystemCollectionBase& Collection)
@@ -248,7 +289,7 @@ void ULandscapeSubsystem::Deinitialize()
 	Scalability::OnScalabilitySettingsChanged.Remove(OnScalabilityChangedHandle);
 	
 #if WITH_EDITOR
-	
+
 	while (NaniteBuildsInFlight != 0)
 	{
 		ENamedThreads::Type CurrentThread = FTaskGraphInterface::Get().GetCurrentThreadIfKnown();
@@ -256,12 +297,23 @@ void ULandscapeSubsystem::Deinitialize()
 		FAssetCompilingManager::Get().ProcessAsyncTasks();
 	}
 
-	delete GrassMapsBuilder;
 	delete PhysicalMaterialBuilder;
 	delete NotificationManager;
+#endif
+
+	delete GrassMapsBuilder;
 	delete TextureStreamingManager;
 
-#endif
+	// cleanup landscape groups
+	for (auto Pair : Groups)
+	{
+		FLandscapeGroup* Group = Pair.Value;
+		delete Group;
+	}
+	Groups.Empty();
+
+	StreamingProxiesNeedingReregister.Empty();
+
 	Proxies.Empty();
 	LandscapeActors.Empty();
 
@@ -274,21 +326,59 @@ void ULandscapeSubsystem::HandlePostGarbageCollect()
 	GetTextureStreamingManager()->CleanupPostGarbageCollect();
 }
 
+void ULandscapeSubsystem::OnProxyMoved(USceneComponent* MovedComponent, EUpdateTransformFlags UpdateTransformFlags, ETeleportType Teleport)
+{
+	AActor* Owner = MovedComponent->GetOwner();
+	if (ALandscapeStreamingProxy* StreamingProxy = Cast<ALandscapeStreamingProxy>(Owner))
+	{
+		StreamingProxiesNeedingReregister.Add(StreamingProxy);
+	}
+}
+
 TStatId ULandscapeSubsystem::GetStatId() const
 {
 	RETURN_QUICK_DECLARE_CYCLE_STAT(ULandscapeSubsystem, STATGROUP_Tickables);
 }
 
+FLandscapeGroup* ULandscapeSubsystem::GetLandscapeGroupForProxy(ALandscapeProxy* Proxy)
+{
+	check(Proxy);
+
+	// use LODGroupKey instead of LandscapeGUID when LODGroupKey is non-zero
+	uint32 LandscapeGroupKey = (Proxy->LODGroupKey != 0) ? Proxy->LODGroupKey : GetTypeHash(Proxy->GetLandscapeGuid());
+
+	FLandscapeGroup*& GroupPtr = Groups.FindOrAdd(LandscapeGroupKey);
+	if (GroupPtr == nullptr)
+	{
+		GroupPtr = new FLandscapeGroup(LandscapeGroupKey);
+	}
+	return GroupPtr;
+}
+
+FLandscapeGroup* ULandscapeSubsystem::GetLandscapeGroupForComponent(ULandscapeComponent* Component)
+{
+	return GetLandscapeGroupForProxy(Component->GetLandscapeProxy());
+}
+
 void ULandscapeSubsystem::RegisterComponent(ULandscapeComponent* Component)
 {
+ 	if (FLandscapeGroup* Group = GetLandscapeGroupForComponent(Component))
+ 	{
+ 		Group->RegisterComponent(Component);
+ 	}
+
 	GetGrassMapBuilder()->RegisterComponent(Component);
 }
 
 void ULandscapeSubsystem::UnregisterComponent(ULandscapeComponent* Component)
 {
 	GetGrassMapBuilder()->UnregisterComponent(Component);
-}
 
+	if (FLandscapeGroup* Group = Component->RegisteredLandscapeGroup)
+	{
+		Group->UnregisterComponent(Component);
+	}
+}
 
 void ULandscapeSubsystem::RemoveGrassInstances(const TSet<ULandscapeComponent*>* ComponentsToRemoveGrassInstances)
 {
@@ -389,7 +479,7 @@ void ULandscapeSubsystem::Tick(float DeltaTime)
 		const bool bNewShowCollisions = NumViewsWithShowCollision > 0;
 		const bool bShowCollisionChanged = (bNewShowCollisions != bAnyViewShowCollisions);
 		bAnyViewShowCollisions = bNewShowCollisions;
-        
+
 		if (bShowCollisionChanged)
 		{
 			for (ULandscapeHeightfieldCollisionComponent* LandscapeHeightfieldCollisionComponent : TObjectRange<ULandscapeHeightfieldCollisionComponent>(RF_ClassDefaultObject | RF_ArchetypeObject, true, EInternalObjectFlags::Garbage))
@@ -433,9 +523,7 @@ void ULandscapeSubsystem::Tick(float DeltaTime)
 		}
 	}
 
-	// dereference and filter TWeakObjectPtr proxies once, up front, to get the list of active proxies
-	// also run early update on Proxies, and determine if all of the proxies are ready for grass generation to start
-	// TODO [chris.tchou] We should drop the usage of TWeakObjectPtr to reduce the dereference cost, can rely on register/unregister instead
+	// run early update on Proxies, and determine if all of the proxies are ready for grass generation to start
 	bool bAllProxiesReadyForGrassMapGeneration = true;
 	bool bAllProxiesRuntimeGrassMapsDisabled = true;
 
@@ -512,6 +600,14 @@ void ULandscapeSubsystem::Tick(float DeltaTime)
 
 	GrassMapsBuilder->AmortizedUpdateGrassMaps(Cameras ? *Cameras : TArray<FVector>(), bIsGrassCreationPrioritized, bAllowStartGrassMapGeneration);
 
+	// run edge fixup for streaming proxies, which only exist in partitioned worlds
+	// TODO [chris.tchou] : remove this check once we support edge fixup on non-partitioned worlds (requires removing heightmap sharing between components)
+	bool bRunEdgeFixup = UWorld::IsPartitionedWorld(GetWorld());
+	if (bRunEdgeFixup)
+	{
+		TickEdgeFixup();
+	}
+
 #if WITH_EDITOR
 	int32 NumProxiesUpdated = 0;
 	int32 NumMeshesToUpdate = 0;
@@ -640,6 +736,52 @@ void ULandscapeSubsystem::OnNaniteEnabledChanged(IConsoleVariable*)
 		}
 	}
 
+void ULandscapeSubsystem::TickEdgeFixup()
+{
+	TRACE_CPUPROFILER_EVENT_SCOPE(ULandscapeSubsystem::TickEdgeFixup);
+
+	// apply any requested re-registrations first
+	if (StreamingProxiesNeedingReregister.Num() > 0)
+	{
+		int32 ComponentsUnregistered = 0;
+		int32 ComponentsRegistered = 0;
+
+		// unregister all first, from whatever group they are registered to
+		for (ALandscapeProxy* StreamingProxy : StreamingProxiesNeedingReregister)
+		{
+			for (ULandscapeComponent* Component : StreamingProxy->LandscapeComponents)
+			{
+				if (FLandscapeGroup* Group = Component->RegisteredLandscapeGroup)
+				{
+					Group->UnregisterComponent(Component);
+					ComponentsUnregistered++;
+				}
+			}
+		}
+
+		// then re-register them to the group they should be in
+		for (ALandscapeProxy* StreamingProxy : StreamingProxiesNeedingReregister)
+		{
+			FLandscapeGroup* NewGroup = GetLandscapeGroupForProxy(StreamingProxy);
+			for (ULandscapeComponent* Component : StreamingProxy->LandscapeComponents)
+			{
+				NewGroup->RegisterComponent(Component);
+				ComponentsRegistered++;
+			}
+		}
+
+		StreamingProxiesNeedingReregister.Empty();
+	}
+
+	// tick each group
+	const bool bForcePatchAll = UE::Landscape::ShouldPatchAllLandscapeComponentEdges(/* bResetForNext = */ true);
+	for (auto Pair : Groups)
+	{
+		FLandscapeGroup* Group = Pair.Value;
+		Group->TickEdgeFixup(this, bForcePatchAll);
+	}
+}
+
 #if WITH_EDITOR
 void ULandscapeSubsystem::BuildAll()
 {
diff --git a/Engine/Source/Runtime/Landscape/Private/LandscapeTextureStorageProvider.cpp b/Engine/Source/Runtime/Landscape/Private/LandscapeTextureStorageProvider.cpp
--- a/Engine/Source/Runtime/Landscape/Private/LandscapeTextureStorageProvider.cpp
+++ b/Engine/Source/Runtime/Landscape/Private/LandscapeTextureStorageProvider.cpp
@@ -8,19 +8,157 @@ LandscapeTextureStorageProvider.cpp: Alternative Texture Storage for Landscape T
 #include "LandscapeDataAccess.h"
 #include "RHIGlobals.h"
 #include "ContentStreaming.h"
+#include "LandscapeGroup.h"
 #include "LandscapePrivate.h"
 
 #include UE_INLINE_GENERATED_CPP_BY_NAME(LandscapeTextureStorageProvider)
 
-FLandscapeTextureStorageMipProvider::FLandscapeTextureStorageMipProvider(ULandscapeTextureStorageProviderFactory* InFactory)
-	: FTextureMipDataProvider(InFactory->Texture, ETickState::Init, ETickThread::Async)
+// enables verbose debug spew
+//#define ENABLE_LANDSCAPE_PROVIDER_DEBUG_SPEW 1
+#ifdef ENABLE_LANDSCAPE_PROVIDER_DEBUG_SPEW
+#define PROVIDER_DEBUG_LOG(...) UE_LOG(LogLandscape, Warning, __VA_ARGS__)
+#define PROVIDER_DEBUG_LOG_DETAIL(...) UE_LOG(LogLandscape, Warning, __VA_ARGS__)
+#else
+#define PROVIDER_DEBUG_LOG(...) UE_LOG(LogLandscape, Verbose, __VA_ARGS__)
+#define PROVIDER_DEBUG_LOG_DETAIL(...) do {} while(0)
+#endif // ENABLE_LANDSCAPE_PROVIDER_DEBUG_SPEW
+
+FLandscapeTextureMipEdgeOverrideProvider::FLandscapeTextureMipEdgeOverrideProvider(ULandscapeHeightmapTextureEdgeFixup* InEdgeFixup, UTexture2D* InTexture)
+	: FTextureMipDataProvider(InTexture, ETickState::GetMips, ETickThread::Async)
 {
-	this->Factory = InFactory;
+	this->EdgeFixup = InEdgeFixup;
+	TextureName = InTexture->GetFName();
+}
+
+FLandscapeTextureMipEdgeOverrideProvider::~FLandscapeTextureMipEdgeOverrideProvider()
+{
+}
+
+void FLandscapeTextureMipEdgeOverrideProvider::Init(const FTextureUpdateContext& Context, const FTextureUpdateSyncOptions& SyncOptions)
+{
+	AdvanceTo(ETickState::GetMips, ETickThread::Async);
+}
+
+int32 FLandscapeTextureMipEdgeOverrideProvider::GetMips(const FTextureUpdateContext& Context, int32 StartingMipIndex, const FTextureMipInfoArray& MipInfos, const FTextureUpdateSyncOptions& SyncOptions)
+{
+	// make a copy of the dest mip infos, for reference in PollMips
+	DestMipInfos = MipInfos;
+
+	AdvanceTo(ETickState::PollMips, ETickThread::Async);
+	return StartingMipIndex;	// we don't directly handle any mips -- return the same starting mip index
+}
+
+bool FLandscapeTextureMipEdgeOverrideProvider::PollMips(const FTextureUpdateSyncOptions& SyncOptions)
+{
+	using namespace UE::Landscape;
+
+	// poll mips will run once all io requests are complete (or cancelled)
+	// here we are relying on the behavior of the default providers, whose PollMips run _after_ PollMips on custom providers like this one.
+	// We rely on the fact that they do not modify the MipData in PollMips.
+	// THIS IS NOT TRUE OF ALL PROVIDERS -- for example the FLandscapeTextureStorageMipProvider will write to mip data in PollMips
+	// however, we handle that case by merging the override functionality into FLandscapeTextureStorageMipProvider, so we don't need a separate override provider.
+	bool bSuccess= true;
+
+	if (!ShouldPatchStreamingMipEdges())
+	{
+		AdvanceTo(ETickState::Done, ETickThread::None);
+		return bSuccess;
+	}
+
+	if ((EdgeFixup == nullptr) || !EdgeFixup->IsActive())
+	{
+		// this heightmap is not yet registered and active -- we can't patch yet.
+		// not to worry though!  When it DOES register, it will fix all existing mips.
+		// (so this mip will be handled at that point)
+		PROVIDER_DEBUG_LOG_DETAIL(TEXT("---- PollMips Coord Mips (%d ... %d) -- NOT READY"), PendingFirstLODIdx, CurrentFirstLODIdx - 1);
+		AdvanceTo(ETickState::Done, ETickThread::None);
+		return bSuccess;
+	}
+
+	int32 PatchedEdges = 0;
+
+	// ensure no one modifies neighbor mapping or snapshots while we are reading them
+	FReadScopeLock ScopeReadLock(EdgeFixup->ActiveGroup->RWLock);
+
+	// Grab neighbor snapshots (null if they don't exist) -- IN A THREAD SAFE MANNER
+	FNeighborSnapshots NeighborSnapshots;
+	EdgeFixup->GetNeighborSnapshots(NeighborSnapshots);
+
+	// patch edges for ALL mips that are requested
+	if (NeighborSnapshots.ExistingNeighbors != ENeighborFlags::None)
+	{
+		PatchedEdges += EdgeFixup->PatchTextureEdgesForStreamingMips(PendingFirstLODIdx, CurrentFirstLODIdx, DestMipInfos, NeighborSnapshots);
+	}
+
+	PROVIDER_DEBUG_LOG(TEXT("---- PollMips Coord (%d,%d) Mips (%d ... %d) -- PATCHED %d edges"), EdgeFixup->GetGroupCoord().X, EdgeFixup->GetGroupCoord().Y, PendingFirstLODIdx, CurrentFirstLODIdx - 1, PatchedEdges);
+
+	AdvanceTo(ETickState::Done, ETickThread::None);
+	return bSuccess;
+}
+
+void FLandscapeTextureMipEdgeOverrideProvider::CleanUp(const FTextureUpdateSyncOptions& SyncOptions)
+{
+	AdvanceTo(ETickState::Done, ETickThread::None);
+}
+
+void FLandscapeTextureMipEdgeOverrideProvider::Cancel(const FTextureUpdateSyncOptions& SyncOptions)
+{
+}
+
+FTextureMipDataProvider::ETickThread FLandscapeTextureMipEdgeOverrideProvider::GetCancelThread() const
+{
+	return ETickThread::None;
+}
 
-	if (InFactory->Texture)
+ULandscapeTextureMipEdgeOverrideFactory::ULandscapeTextureMipEdgeOverrideFactory(const FObjectInitializer& ObjectInitializer)
+	: Super(ObjectInitializer)
+{
+}
+
+ULandscapeTextureMipEdgeOverrideFactory* ULandscapeTextureMipEdgeOverrideFactory::AddTo(UTexture2D* TargetTexture)
+{
+	check(TargetTexture);
+
+	// try to get an existing factory
+	ULandscapeTextureMipEdgeOverrideFactory* Factory = TargetTexture->GetAssetUserData<ULandscapeTextureMipEdgeOverrideFactory>();
+	if (Factory == nullptr)
 	{
-		TextureName = InFactory->Texture->GetFName();
+		// create a new one (with TargetTexture as outer)
+		Factory = NewObject<ULandscapeTextureMipEdgeOverrideFactory>(TargetTexture);
+		Factory->Texture = TargetTexture;
+		TargetTexture->AddAssetUserData(Factory);
 	}
+
+	check(Factory->Texture == TargetTexture);
+	check(Factory->GetOuter() == TargetTexture);
+
+	return Factory;
+}
+
+void ULandscapeTextureMipEdgeOverrideFactory::SetupEdgeFixup(ULandscapeHeightmapTextureEdgeFixup* InEdgeFixup)
+{
+	this->EdgeFixup = InEdgeFixup;
+}
+
+void ULandscapeTextureMipEdgeOverrideFactory::Serialize(FArchive& Ar)
+{
+	Super::Serialize(Ar);
+	Ar << Texture;
+}
+
+void ULandscapeTextureMipEdgeOverrideFactory::AddReferencedObjects(UObject* InThis, FReferenceCollector& Collector)
+{
+	Super::AddReferencedObjects(InThis, Collector);
+	ULandscapeTextureMipEdgeOverrideFactory* const TypedThis = Cast<ULandscapeTextureMipEdgeOverrideFactory>(InThis);
+	Collector.AddReferencedObject(TypedThis->Texture);
+	Collector.AddReferencedObject(TypedThis->EdgeFixup);
+}
+
+FLandscapeTextureStorageMipProvider::FLandscapeTextureStorageMipProvider(ULandscapeTextureStorageProviderFactory* InFactory)
+	: FTextureMipDataProvider(InFactory->Texture, ETickState::Init, ETickThread::Async)
+{
+	this->Factory = InFactory;
+	TextureName = InFactory->Texture->GetFName();
 }
 
 FLandscapeTextureStorageMipProvider::~FLandscapeTextureStorageMipProvider()
@@ -116,6 +254,14 @@ void ULandscapeTextureStorageProviderFactory::Serialize(FArchive& Ar)
 	Ar << Texture;
 }
 
+void ULandscapeTextureStorageProviderFactory::AddReferencedObjects(UObject* InThis, FReferenceCollector& Collector)
+{
+	Super::AddReferencedObjects(InThis, Collector);
+	ULandscapeTextureStorageProviderFactory* const TypedThis = Cast<ULandscapeTextureStorageProviderFactory>(InThis);
+	Collector.AddReferencedObject(TypedThis->Texture);
+	Collector.AddReferencedObject(TypedThis->EdgeFixup);
+}
+
 FStreamableRenderResourceState ULandscapeTextureStorageProviderFactory::GetResourcePostInitState(const UTexture* Owner, bool bAllowStreaming)
 {
 	// We are using the non-offline mode to upload these textures currently, so we don't need to consider mip tails.
@@ -315,6 +461,7 @@ ULandscapeTextureStorageProviderFactory* ULandscapeTextureStorageProviderFactory
 		TargetTexture->AddAssetUserData(Factory);
 	}
 
+	// need this to properly calculate normals
 	Factory->LandscapeGridScale = InLandsapeGridScale;
 
 	// calculate number of non-streaming mips
@@ -374,6 +521,11 @@ ULandscapeTextureStorageProviderFactory* ULandscapeTextureStorageProviderFactory
 #endif // WITH_EDITORONLY_DATA
 
 
+void ULandscapeTextureStorageProviderFactory::SetupEdgeFixup(ULandscapeHeightmapTextureEdgeFixup* InEdgeFixup)
+{
+	this->EdgeFixup = InEdgeFixup;
+}
+
 // Helper to configure the AsyncFileCallBack.
 void FLandscapeTextureStorageMipProvider::CreateAsyncFileCallback(const FTextureUpdateSyncOptions& SyncOptions)
 {
@@ -452,11 +604,13 @@ void FLandscapeTextureStorageMipProvider::Init(const FTextureUpdateContext& Cont
 	AdvanceTo(ETickState::GetMips, ETickThread::Async);
 }
 
-int32 FLandscapeTextureStorageMipProvider::GetMips(const FTextureUpdateContext& Context, int32 StartingMipIndex, const FTextureMipInfoArray& MipInfos, const FTextureUpdateSyncOptions& SyncOptions)\
+int32 FLandscapeTextureStorageMipProvider::GetMips(const FTextureUpdateContext& Context, int32 StartingMipIndex, const FTextureMipInfoArray& MipInfos, const FTextureUpdateSyncOptions& SyncOptions)
 {
 	CreateAsyncFileCallback(SyncOptions);	// this just creates it... callback has to be passed to the IO request completion to actually get called...
 	check(SyncOptions.Counter != nullptr);
 
+	DestMipInfos = MipInfos;
+
 	FirstRequestedMipIndex = StartingMipIndex;
 	while (StartingMipIndex < CurrentFirstLODIdx && MipInfos.IsValidIndex(StartingMipIndex))
 	{
@@ -507,6 +661,8 @@ int32 FLandscapeTextureStorageMipProvider::GetMips(const FTextureUpdateContext&
 
 bool FLandscapeTextureStorageMipProvider::PollMips(const FTextureUpdateSyncOptions& SyncOptions)
 {
+	using namespace UE::Landscape;
+
 	// poll mips will run once all io requests are complete (or cancelled)
 
 	// Notify that some files have possibly been unmounted / missing.
@@ -531,6 +687,31 @@ bool FLandscapeTextureStorageMipProvider::PollMips(const FTextureUpdateSyncOptio
 			uint8* DestData = IORequests[MipIndex].DestMipData;
 			Factory->DecompressMip(SourceData, SourceMip->BulkData.GetBulkDataSize(), DestData, DestDataBytes, MipIndex);
 		}
+		
+		if (ShouldPatchStreamingMipEdges())
+		{
+			// run mip patching if EdgeFixup is valid
+			ULandscapeHeightmapTextureEdgeFixup* EdgeFixup = Factory->EdgeFixup;
+			if ((EdgeFixup != nullptr) && EdgeFixup->IsActive())
+			{
+				int32 PatchedEdges = 0;
+			
+				// ensure no one modifies neighbor mapping or snapshots while we are reading them
+				FReadScopeLock ScopeReadLock(EdgeFixup->ActiveGroup->RWLock);
+
+				// Grab neighbor snapshots (null if they don't exist)
+				FNeighborSnapshots NeighborSnapshots;
+				EdgeFixup->GetNeighborSnapshots(NeighborSnapshots);
+
+				// patch edges for ALL mips that are requested
+				if (NeighborSnapshots.ExistingNeighbors != ENeighborFlags::None)
+				{
+					PatchedEdges += EdgeFixup->PatchTextureEdgesForStreamingMips(PendingFirstLODIdx, CurrentFirstLODIdx, DestMipInfos, NeighborSnapshots);
+				}
+
+				PROVIDER_DEBUG_LOG(TEXT("---- PollMips Coord (%d,%d) Mips (%d ... %d) -- PATCHED COMPRESSED %d edges"), EdgeFixup->GetGroupCoord().X, EdgeFixup->GetGroupCoord().Y, PendingFirstLODIdx, CurrentFirstLODIdx - 1, PatchedEdges);
+			}
+		}
 	}
 
 	ClearIORequests();
@@ -791,3 +972,7 @@ void ULandscapeTextureStorageProviderFactory::DecompressMip(uint8* SourceData, i
 
 	check(SourceOffset == SourceDataBytes);
 }
+
+#undef PROVIDER_DEBUG_LOG
+#undef PROVIDER_DEBUG_LOG_DETAIL
+
diff --git a/Engine/Source/Runtime/Landscape/Public/LandscapeSubsystem.h b/Engine/Source/Runtime/Landscape/Public/LandscapeSubsystem.h
--- a/Engine/Source/Runtime/Landscape/Public/LandscapeSubsystem.h
+++ b/Engine/Source/Runtime/Landscape/Public/LandscapeSubsystem.h
@@ -9,6 +9,7 @@
 
 class ALandscapeProxy;
 class ALandscape;
+class ALandscapeStreamingProxy;
 class AWorldSettings;
 class IConsoleVariable;
 class ULandscapeInfo;
@@ -19,6 +20,9 @@ class FLandscapeTextureStreamingManager;
 struct FActionableMessage;
 struct FDateTime;
 struct FScopedSlowTask;
+class ULandscapeHeightmapTextureEdgeFixup;
+struct FLandscapeGroup;
+enum class EUpdateTransformFlags : int32;
 
 namespace UE::Landscape
 {
@@ -72,6 +76,8 @@ public:
 	virtual TStatId GetStatId() const override;
 	// End FTickableGameObject overrides
 
+	static void AddReferencedObjects(UObject* InThis, FReferenceCollector& Collector);
+
 	// setting this to true causes grass instance generation to go wider (multiplies the limits by GGrassCreationPrioritizedMultipler)
 	void PrioritizeGrassCreation(bool bPrioritizeGrassCreation) { bIsGrassCreationPrioritized = bPrioritizeGrassCreation; }
 	bool IsGrassCreationPrioritized() const { return bIsGrassCreationPrioritized; }
@@ -143,6 +149,12 @@ public:
 	void WaitLaunchNaniteBuild(); 
 #endif // WITH_EDITOR
 
+	// runs per-tick edge fixup on ALL landscape groups in the subsystem
+	void TickEdgeFixup();
+
+	FLandscapeGroup* GetLandscapeGroupForProxy(ALandscapeProxy* Proxy);
+	FLandscapeGroup* GetLandscapeGroupForComponent(ULandscapeComponent* Component);
+
 private:
 	LANDSCAPE_API void ForEachLandscapeInfo(TFunctionRef<bool(ULandscapeInfo*)> ForEachLandscapeInfoFunc) const;
 
@@ -155,15 +167,25 @@ private:
 	void OnNaniteEnabledChanged(IConsoleVariable*);
 
 	void HandlePostGarbageCollect();
+	
+	// When proxies move, we may need to update their position in the landscape group
+	void OnProxyMoved(USceneComponent*, EUpdateTransformFlags, ETeleportType);
+
+	// LODGroupKey --> Landscape Group
+	TMap<uint32, FLandscapeGroup*> Groups;
+
+	// list of streaming proxies that need to re-register with their group because they moved, or changed their LODGroupKey
+	UPROPERTY(Transient, DuplicateTransient, NonTransactional)
+	TSet<TObjectPtr<ALandscapeStreamingProxy>> StreamingProxiesNeedingReregister;
 
 	bool bIsGrassCreationPrioritized = false;
 
-	// UPROPERTY ensures these objects are not deleted before being unregistered
-	// (technically not necessary, as actors should always unregister prior to deletion)
-	UPROPERTY()
+	UPROPERTY(Transient, DuplicateTransient, NonTransactional)
 	TArray<TObjectPtr<ALandscape>> LandscapeActors;
 
-	UPROPERTY()
+	// UPROPERTY ensures these objects are not deleted before being unregistered
+	// (technically not necessary, as actors should always unregister prior to deletion)
+	UPROPERTY(Transient, DuplicateTransient, NonTransactional)
 	TArray<TObjectPtr<ALandscapeProxy>> Proxies;
 
 	FDelegateHandle OnNaniteWorldSettingsChangedHandle;
