Fix for PVS Studio warning about unitialized variable (Issue 578175)

diff --git a/Engine/Source/Runtime/Landscape/Private/LandscapeEdgeFixup.cpp b/Engine/Source/Runtime/Landscape/Private/LandscapeEdgeFixup.cpp
--- a/Engine/Source/Runtime/Landscape/Private/LandscapeEdgeFixup.cpp
+++ b/Engine/Source/Runtime/Landscape/Private/LandscapeEdgeFixup.cpp
@@ -1057,7 +1057,7 @@ int32 ULandscapeHeightmapTextureEdgeFixup::CheckAndPatchTextureEdgesFromEdgeSnap
 			}
 
 			// Because the corner hashes are just a copy of the corner data, we can directly compare them with the expected result
-			FHeightmapTexel Blended, GPUExisting;
+			FHeightmapTexel Blended, GPUExisting = {};
 			BlendCornerData(Blended, CornerIndex, NeighborSnapshots);
 
 			GPUExisting.Data32 = GPUEdgeHashes[(int32)CornerIndex];
