Landscape runtime grass maps were sometimes not rendering for a component at the necessary distance if there were other closer landscape components with zero grass varieties.
Fixes the issue by changing the priority key we sort on to be "distance until spawn" instead of "distance to component" - so a component that should spawn closer will always be considered first.

I also changed the fast paths to handle the specific case of no grass varieties (or no grass varieties with meshes), so it doesn't bother to render their grass maps.

I fixed a TODO by adding simple camera jump detection to update all priorities during camera cuts.

I optimized the handling of new pending components by tracking the range of new components to minimize initialization time, and only adding or removing them from the pending heap when necessary.

diff --git a/Engine/Source/Runtime/GeometryCore/Public/Spatial/PointHashGrid3.h b/Engine/Source/Runtime/GeometryCore/Public/Spatial/PointHashGrid3.h
--- a/Engine/Source/Runtime/GeometryCore/Public/Spatial/PointHashGrid3.h
+++ b/Engine/Source/Runtime/GeometryCore/Public/Spatial/PointHashGrid3.h
@@ -56,6 +56,16 @@ public:
 		Hash.Reserve(Num);
 	}
 
+	/**
+	 * Remove all entries in the underlying hash map, without resizing or releasing any allocations, and reset the cell size.
+	 * @param NewCellSize The size of grid cells
+	 */
+	void Reset(RealType NewCellSize)
+	{
+		Hash.Reset();
+		Indexer.CellSize = NewCellSize;
+	}
+
 	/** Invalid grid value */
 	PointDataType GetInvalidValue() const
 	{
diff --git a/Engine/Source/Runtime/Landscape/Classes/LandscapeComponent.h b/Engine/Source/Runtime/Landscape/Classes/LandscapeComponent.h
--- a/Engine/Source/Runtime/Landscape/Classes/LandscapeComponent.h
+++ b/Engine/Source/Runtime/Landscape/Classes/LandscapeComponent.h
@@ -568,6 +568,7 @@ public:
 	struct FGrassTypeSummary
 	{
 		bool bInvalid = true;
+		bool bHasAnyGrass = true;
 		double MaxInstanceDiscardDistance = DBL_MAX;
 	};
 	FGrassTypeSummary GrassTypeSummary;
diff --git a/Engine/Source/Runtime/Landscape/Private/LandscapeGrass.cpp b/Engine/Source/Runtime/Landscape/Private/LandscapeGrass.cpp
--- a/Engine/Source/Runtime/Landscape/Private/LandscapeGrass.cpp
+++ b/Engine/Source/Runtime/Landscape/Private/LandscapeGrass.cpp
@@ -1035,27 +1035,33 @@ void ALandscapeProxy::UpdateGrassTypeSummary()
 
 	bool bGrassQualityLevelEnabled = GEngine && GEngine->UseGrassVarityPerQualityLevels;
 
-	bool bAnyGrass = false;
+	bool bProxyHasAnyGrass = false;
 	double ProxyMaxGrassCullDistance = 0.0;
 	for (ULandscapeComponent* Component : LandscapeComponents)
 	{
 		double ComponentMaxGrassCullDistance;
+		int32 ComponentGrassVarietyCount = 0;
 		if (Component->GrassTypeSummary.bInvalid)
 		{
-			// Update the component cache
+			// Update the component grass type summary
 			ComponentMaxGrassCullDistance = 0.0;
 			for (ULandscapeGrassType* GrassType : Component->GetGrassTypes())
 			{
-				if (GrassType)
+				if (GrassType != nullptr)
 				{
 					for (FGrassVariety& GrassVariety : GrassType->GrassVarieties)
 					{
-						int32 EndCullDistance = bGrassQualityLevelEnabled ? GrassVariety.EndCullDistanceQuality.GetValue(GGrassQualityLevel) : GrassVariety.EndCullDistance.GetValue();
-						ComponentMaxGrassCullDistance = FMath::Max(ComponentMaxGrassCullDistance, EndCullDistance);
+						if (GrassVariety.GrassMesh != nullptr)
+						{
+							int32 EndCullDistance = bGrassQualityLevelEnabled ? GrassVariety.EndCullDistanceQuality.GetValue(GGrassQualityLevel) : GrassVariety.EndCullDistance.GetValue();
+							ComponentMaxGrassCullDistance = FMath::Max(ComponentMaxGrassCullDistance, EndCullDistance);
+							ComponentGrassVarietyCount++;
+						}
 					}
 				}
 			}
 			Component->GrassTypeSummary.MaxInstanceDiscardDistance = ComponentMaxGrassCullDistance;
+			Component->GrassTypeSummary.bHasAnyGrass = (ComponentGrassVarietyCount > 0);
 			Component->GrassTypeSummary.bInvalid = false;
 		}
 		else
@@ -1063,11 +1069,11 @@ void ALandscapeProxy::UpdateGrassTypeSummary()
 			ComponentMaxGrassCullDistance = Component->GrassTypeSummary.MaxInstanceDiscardDistance;
 		}
 
-		bAnyGrass |= (Component->GetGrassTypes().Num() > 0);
+		bProxyHasAnyGrass = bProxyHasAnyGrass || Component->GrassTypeSummary.bHasAnyGrass;
 		ProxyMaxGrassCullDistance = FMath::Max(ProxyMaxGrassCullDistance, ComponentMaxGrassCullDistance);
 	}
 
-	GrassTypeSummary.bHasAnyGrass = bAnyGrass;
+	GrassTypeSummary.bHasAnyGrass = bProxyHasAnyGrass;
 	GrassTypeSummary.MaxInstanceDiscardDistance = ProxyMaxGrassCullDistance;
 	GrassTypeSummary.LandscapeComponentCount = LandscapeComponents.Num();
 }
diff --git a/Engine/Source/Runtime/Landscape/Private/LandscapeGrassMapsBuilder.cpp b/Engine/Source/Runtime/Landscape/Private/LandscapeGrassMapsBuilder.cpp
--- a/Engine/Source/Runtime/Landscape/Private/LandscapeGrassMapsBuilder.cpp
+++ b/Engine/Source/Runtime/Landscape/Private/LandscapeGrassMapsBuilder.cpp
@@ -102,6 +102,12 @@ static FAutoConsoleVariableRef CVarGrassMapGuardBandDiscardMultiplier(
 	GGrassMapGuardBandDiscardMultiplier,
 	TEXT("Used to control discarding in the grass map runtime generation system. Approximate range, 1-4. Multiplied by the cull distance to control when we discard grass maps."));
 
+static float GGrassMapCameraCutTranslationThreshold = 10000.0f;
+static FAutoConsoleVariableRef CVarCameraCutTranslationThreshold(
+	TEXT("grass.GrassMap.CameraCutTranslationThreshold"),
+	GGrassMapCameraCutTranslationThreshold,
+	TEXT("The maximum camera translation distance in centimeters allowed between two frames before we consider it a camera cut and grass map priorities are immediately recalculated."));
+
 DECLARE_CYCLE_STAT(TEXT("Update Component GrassMap "), STAT_UpdateComponentGrassMaps, STATGROUP_Foliage);
 DECLARE_CYCLE_STAT(TEXT("Prioritize Pending GrassMaps"), STAT_PrioritizePendingGrassMaps, STATGROUP_Foliage);
 DECLARE_CYCLE_STAT(TEXT("Render GrassMap"), STAT_RenderGrassMap, STATGROUP_Foliage);
@@ -111,6 +117,7 @@ DECLARE_CYCLE_STAT(TEXT("Remove Grass Instances"), STAT_RemoveGrassInstances, ST
 FLandscapeGrassMapsBuilder::FLandscapeGrassMapsBuilder(UWorld* InOwner, FLandscapeTextureStreamingManager& InTextureStreamingManager)
 	: World(InOwner)
 	, TextureStreamingManager(InTextureStreamingManager)
+	, PreviousCameraHashGrid(GGrassMapCameraCutTranslationThreshold, FVector(NAN, NAN, NAN))
 #if WITH_EDITOR
 	, OutdatedGrassMapCount(0)
 	, GrassMapsLastCheckTime(0)
@@ -234,7 +241,7 @@ namespace UE::Landscape
 	}
 
 	// calculates the minimum distance between any cameras and the specified Worldbounds.
-	static inline double CalculateMinDistanceToCameras(const TArray<FVector>& Cameras, const FBoxSphereBounds& WorldBounds)
+	static inline double CalculateMinDistanceToCamerasSquared(const TArray<FVector>& Cameras, const FBoxSphereBounds& WorldBounds)
 	{
 		if (!Cameras.Num())
 		{
@@ -374,13 +381,14 @@ FLandscapeGrassMapsBuilder::FComponentState::FComponentState(ULandscapeComponent
 #endif // WITH_EDITOR
 }
 
-void FLandscapeGrassMapsBuilder::FPendingComponent::UpdatePriorityDistance(const TArray<FVector>&Cameras)
+void FLandscapeGrassMapsBuilder::FPendingComponent::UpdatePriorityDistance(const TArray<FVector>&Cameras, float MustHaveDistanceScale)
 {
 	TRACE_CPUPROFILER_EVENT_SCOPE(FLandscapeGrassMapsBuilder::FPendingComponent::UpdatePriorityDistance);
 	ULandscapeComponent* Component = State->Component;
 	FBoxSphereBounds WorldBounds = Component->CalcBounds(Component->GetComponentTransform());
-	double MinSqrDistanceToComponent = UE::Landscape::CalculateMinDistanceToCameras(Cameras, WorldBounds);
-	PriorityKey = MinSqrDistanceToComponent;
+	const double MinSqrDistanceToComponent = UE::Landscape::CalculateMinDistanceToCamerasSquared(Cameras, WorldBounds);
+	const double ThresholdDistance = Component->GrassTypeSummary.MaxInstanceDiscardDistance * MustHaveDistanceScale;
+	PriorityKey = MinSqrDistanceToComponent - ThresholdDistance * ThresholdDistance;
 }
 
 void FAsyncFetchTask::DoWork()
@@ -601,40 +609,95 @@ void FLandscapeGrassMapsBuilder::StartPrioritizedGrassMapGeneration(const TArray
 	// no point in calling this if there are no cameras -- we can't calculate priority
 	check(Cameras.Num() > 0);
 
-	float MustHaveDistanceScale = GGrassMapGuardBandMultiplier * GGrassCullDistanceScale;
-
 	// update pending component priorities (distances)
 	check(PendingComponentsHeap.Num() == PendingCount);
 	if (PendingCount)
 	{
-		// We update the priority of one element in each heap level.
-		// Because the heap is ordered by distance, this approximately updates
-		// closer elements more often than distant elements.
-		// The closest element is updated every frame, the second and third every other frame, 4-7 every fourth frame, etc.
-		// This way we update at most Log2(N) elements each frame.
-		for (int32 LevelSize = 1; LevelSize <= 65536; LevelSize += LevelSize)
+		// determine if we need to recalculate all the pending component heap priorities
+		bool bRecalculateAllPriorities = (FirstNewPendingComponent == 0);
+
+		// check if scales have changed
+		const float MustHaveDistanceScale = GGrassMapGuardBandMultiplier * GGrassCullDistanceScale;
+		if (MustHaveDistanceScale != LastMustHaveDistanceScale)
+		{
+			bRecalculateAllPriorities = true;
+			LastMustHaveDistanceScale = MustHaveDistanceScale;
+		}
+
+		if (!bRecalculateAllPriorities)
 		{
-			// this should select successive elements on the given Level as the counter is incremented
-			int32 UpdateIndex = (PendingUpdateAmortizationCounter & (LevelSize-1)) + (LevelSize-1);
-			if (UpdateIndex >= PendingCount)
+			// check if cameras jumped suddenly
+			// TODO [chris.tchou] : This is calculated in FViewInfo already - See GCameraCutTranslationThreshold
+			// however, we would need to somehow pass that info through to the FStreamingViewInfos, (and ideally flag position-only cuts as we don't care about angle)
+			float ThresholdSquared = GGrassMapCameraCutTranslationThreshold * GGrassMapCameraCutTranslationThreshold;
+			for (const FVector& CameraPos : Cameras)
 			{
-				break;
+				TPair<FVector, double> Found = PreviousCameraHashGrid.FindAnyInRadius(CameraPos, GGrassMapCameraCutTranslationThreshold, [CameraPos](const FVector& OtherDistanceSqFunc) { return FVector::DistSquared(CameraPos, OtherDistanceSqFunc); });
+				if (Found.Value > ThresholdSquared)
+				{
+					// did not find any previous cameras within the threshold distance, implies the camera must have jumped farther than that distance
+					bRecalculateAllPriorities = true;
+					break;
+				}
 			}
-			PendingComponentsHeap[UpdateIndex].UpdatePriorityDistance(Cameras);
 		}
-		PendingUpdateAmortizationCounter++;
 
-		// also check if any have a negative priority, which indicates newly pending components, and immediately calculate priority for those
-		// TODO [chris.tchou] : if the camera jumps for some reason, we should recalculate all priorities immediately.
-		for (FPendingComponent& Pending : PendingComponentsHeap)
+		// record camera positions for next frame
+		PreviousCameraHashGrid.Reset(GGrassMapCameraCutTranslationThreshold);
+		for (const FVector& CameraPos : Cameras)
+		{
+			PreviousCameraHashGrid.InsertPointUnsafe(CameraPos, CameraPos);
+		}
+
+		if (bRecalculateAllPriorities)
 		{
-			if (Pending.PriorityKey < 0.0)
+			TRACE_CPUPROFILER_EVENT_SCOPE(RecalculateAllPriorities)
+			for (FPendingComponent& Pending : PendingComponentsHeap)
 			{
-				Pending.UpdatePriorityDistance(Cameras);
+				Pending.UpdatePriorityDistance(Cameras, MustHaveDistanceScale);
 			}
 		}
+		else
+		{
+			TRACE_CPUPROFILER_EVENT_SCOPE(AmortizedPriorityUpdate)
+
+			// immediately update any new pending components (at the end of the array)
+			if (FirstNewPendingComponent < PendingComponentsHeap.Num())
+			{
+				check(FirstNewPendingComponent >= 0);
+				for (int32 Index = FirstNewPendingComponent; Index < PendingComponentsHeap.Num(); Index++)
+				{
+					FPendingComponent& Pending = PendingComponentsHeap[Index];
+					Pending.UpdatePriorityDistance(Cameras, MustHaveDistanceScale);
+				}
+			}
+			else
+			{
+				// set temporarily so we can use it as the upper bound in the amortized update below
+				FirstNewPendingComponent = PendingComponentsHeap.Num();
+			}
+
+			// We update the priority of one element in each heap level.
+			// Because the heap is ordered by distance, this approximately updates
+			// closer elements more often than distant elements.
+			// The closest element is updated every frame, the second and third every other frame, 4-7 every fourth frame, etc.
+			// This way we update at most Log2(N) elements each frame.
+			for (int32 LevelSize = 1; LevelSize <= 65536; LevelSize += LevelSize)
+			{
+				// this should select successive elements on the given Level as the counter is incremented
+				int32 UpdateIndex = (PendingUpdateAmortizationCounter & (LevelSize - 1)) + (LevelSize - 1);
+				if (UpdateIndex >= FirstNewPendingComponent)
+				{
+					break;
+				}
+				PendingComponentsHeap[UpdateIndex].UpdatePriorityDistance(Cameras, MustHaveDistanceScale);
+			}
+			PendingUpdateAmortizationCounter++;
+		}
 	}
 
+	FirstNewPendingComponent = TNumericLimits<int32>::Max();
+
 	// re-heapify to reflect updated priorities
 	{
 		TRACE_CPUPROFILER_EVENT_SCOPE(Heapify)
@@ -648,12 +711,12 @@ void FLandscapeGrassMapsBuilder::StartPrioritizedGrassMapGeneration(const TArray
 		FComponentState* State = Pending.State;
 		const ULandscapeComponent* Component = State->Component;
 		check(State->Stage == EComponentStage::Pending);
+		check(Pending.PriorityKey != TNumericLimits<double>::Lowest()); // default value, indicates the priority was never calculated
 
 		// if threshold is enabled, check if the nearest pending component is close enough
 		if (bOnlyWhenCloserThanEvictionRange)
 		{
-			const double MaxSpawnDistance = Component->GrassTypeSummary.MaxInstanceDiscardDistance * MustHaveDistanceScale;
-			if (Pending.PriorityKey > MaxSpawnDistance * MaxSpawnDistance)
+			if (Pending.PriorityKey > 0.0)
 			{
 				// the closest component is not close enough to generate yet
 				break;
@@ -661,12 +724,16 @@ void FLandscapeGrassMapsBuilder::StartPrioritizedGrassMapGeneration(const TArray
 		}
 
 		PendingComponentsHeap.HeapPopDiscard(EAllowShrinking::No);
+		State->bInPendingHeap = false;
+
 		if (StartGrassMapGeneration(*State, false))
 		{
 			MaxComponentsToStart--;
 		}
 	}
 
+	check(PendingComponentsHeap.Num() == PendingCount);
+
 	TotalComponentsWaitingCount = PendingComponentsHeap.Num();
 }
 
@@ -683,16 +750,19 @@ void FLandscapeGrassMapsBuilder::RegisterComponent(ULandscapeComponent* Componen
 	else
 	{
 		FComponentState* NewState = new(StatePoolAllocator.Allocate()) FComponentState(Component);
-
+		// we defer adding NewState to the pending component heap - to save the work of removing it if a fastpath is taken
 		ComponentStates.Add(Component, NewState);
-		PendingComponentsHeap.Add(FPendingComponent(NewState));
-
 		PendingCount++;
 
 		GRASS_DEBUG_LOG(TEXT("Register %s (%d total)"), *Component->GetName(), ComponentStates.Num());
 
 		// immediately after a new registration, check if fast paths apply
-		TryFastpathsFromPending(*NewState, /* bRecalculateHashes = */ false);
+		if (!TryFastpathsFromPending(*NewState, /* bRecalculateHashes = */ false))
+		{
+			// no fast path was taken, so add it to the pending component heap
+			check(NewState->Stage == EComponentStage::Pending);
+			AddToPendingComponentHeap(NewState);
+		}
 	}
 }
 
@@ -1095,7 +1165,7 @@ bool FLandscapeGrassMapsBuilder::CancelAndEvict(FComponentState& State, bool bCa
 		if (State.Component != nullptr)
 		{
 			// don't bother to add if component is null as we will just have to remove it immediately in the deallocate
-			PendingComponentsHeap.Add(FPendingComponent(&State));
+			AddToPendingComponentHeap(&State);
 		}
 	}
 
@@ -1142,7 +1212,7 @@ bool FLandscapeGrassMapsBuilder::TryFastpathsFromPending(FComponentState& State,
 	}
 
 	// handle the easy case of empty grass types; skip directly to the populated state without running the pipeline
-	if (Component->GetGrassTypes().IsEmpty())
+	if (!Component->GrassTypeSummary.bHasAnyGrass)
 	{
 		PendingToPopulatedFastPathNoGrass(State);
 		return true;
@@ -1264,15 +1334,35 @@ void FLandscapeGrassMapsBuilder::PendingToStreaming(FComponentState& State)
 	State.TickCount = 0;
 }
 
+void FLandscapeGrassMapsBuilder::AddToPendingComponentHeap(FComponentState* State)
+{
+	check(!State->bInPendingHeap);
+	int32 NewIndex = PendingComponentsHeap.Add(FPendingComponent(State));
+	State->bInPendingHeap = true;
+	FirstNewPendingComponent = FMath::Min(FirstNewPendingComponent, NewIndex);
+}
+
 void FLandscapeGrassMapsBuilder::RemoveFromPendingComponentHeap(FComponentState* State)
 {
-	for (int32 Index = 0; Index < PendingComponentsHeap.Num(); Index++)
+	// may be removed more than once
+	if (State->bInPendingHeap)
 	{
-		if (PendingComponentsHeap[Index].State == State)
+		for (int32 Index = 0; Index < PendingComponentsHeap.Num(); Index++)
 		{
-			PendingComponentsHeap.RemoveAtSwap(Index);
-			break;
+			if (PendingComponentsHeap[Index].State == State)
+			{
+				PendingComponentsHeap.RemoveAtSwap(Index, EAllowShrinking::No);
+
+				if ((FirstNewPendingComponent <= PendingComponentsHeap.Num()) && (Index < FirstNewPendingComponent))
+				{
+					// this should be fairly rare, as we generally process new pending components before removing anything
+					// but if a new pending component does get swapped by the remove, we need to update the new range
+					FirstNewPendingComponent = Index;
+				}
+				break;
+			}
 		}
+		State->bInPendingHeap = false;
 	}
 }
 
@@ -1401,7 +1491,7 @@ bool FLandscapeGrassMapsBuilder::FComponentState::IsBeyondEvictionRange(const TA
 {
 	check(Stage == EComponentStage::GrassMapsPopulated);
 	const FBoxSphereBounds WorldBounds = Component->CalcBounds(Component->GetComponentTransform());
-	const float MinSqrDistanceToComponent = UE::Landscape::CalculateMinDistanceToCameras(Cameras, WorldBounds);
+	const float MinSqrDistanceToComponent = UE::Landscape::CalculateMinDistanceToCamerasSquared(Cameras, WorldBounds);
 	const float DiscardDistanceScale = GGrassMapGuardBandDiscardMultiplier * GGrassCullDistanceScale;
 	const float MinEvictDistance = Component->GrassTypeSummary.MaxInstanceDiscardDistance * DiscardDistanceScale;
 	return (MinSqrDistanceToComponent > MinEvictDistance * MinEvictDistance);
diff --git a/Engine/Source/Runtime/Landscape/Private/LandscapeGrassMapsBuilder.h b/Engine/Source/Runtime/Landscape/Private/LandscapeGrassMapsBuilder.h
--- a/Engine/Source/Runtime/Landscape/Private/LandscapeGrassMapsBuilder.h
+++ b/Engine/Source/Runtime/Landscape/Private/LandscapeGrassMapsBuilder.h
@@ -7,6 +7,7 @@
 #include "LandscapeComponent.h"
 #include "LandscapeTextureStreamingManager.h"
 #include "Containers/AllocatorFixedSizeFreeList.h"
+#include "Spatial/PointHashGrid3.h"
 
 class FLandscapeGrassWeightExporter;
 struct FScopedSlowTask;
@@ -88,7 +89,7 @@ private:
 	// (based on distance from the given Camera set) -- Cameras must not be empty.
 	void StartPrioritizedGrassMapGeneration(const TArray<FVector>& Cameras, int32 MaxComponentsToStart, bool bOnlyWhenCloserThanEvictionRange);
 	
-	enum class EComponentStage
+	enum class EComponentStage : uint8
 	{
 		Pending,							// initial state; the component is waiting to start the generation process.
 		NotReady,							// tried to start generation process, but either no grass types exist or the material is not ready. wait for that to change.
@@ -107,6 +108,8 @@ private:
 
 		EComponentStage Stage = EComponentStage::Pending;
 
+		bool bInPendingHeap = false;
+
 		// counts the number of ticks this component has remained in the current stage
 		int32 TickCount = 0;
 
@@ -144,10 +147,10 @@ private:
 			return PriorityKey < Other.PriorityKey;
 		}
 
-		void UpdatePriorityDistance(const TArray<FVector>& Cameras);
+		void UpdatePriorityDistance(const TArray<FVector>& Cameras, float MustHaveDistanceScale);
 
 		FComponentState* State = nullptr;
-		double PriorityKey = -1.0;
+		double PriorityKey = TNumericLimits<double>::Lowest();
 	};
 
 	// structure to do an amortized update of a set of update elements in a (conceptual) array
@@ -230,7 +233,9 @@ private:
 
 	// Pending components, in a min heap by distance to streaming cameras
 	TArray<FPendingComponent> PendingComponentsHeap;
+	int32 FirstNewPendingComponent = TNumericLimits<int32>::Max();
 	int32 PendingUpdateAmortizationCounter = 0;
+	float LastMustHaveDistanceScale = 0.0f;			// used to detect spawn distance scales changing
 
 	// state to amortize the update of components
 	FAmortizedUpdate AmortizedUpdate;
@@ -238,6 +243,9 @@ private:
 	// system to manage texture streaming requests
 	FLandscapeTextureStreamingManager& TextureStreamingManager;
 
+	// hash grid storing previously used camera locations (to determine when cameras have jumped)
+	UE::Geometry::TPointHashGrid3<FVector, double> PreviousCameraHashGrid;
+
 	// tries to cancel any in flight operations and transition back to the Pending state
 	// returns true when the state has been successfully transitioned back to Pending
 	// if bCancelImmediately is true, it will ensure the component reaches Pending state before returning (possibly blocking on async or gpu tasks)
@@ -266,6 +274,9 @@ private:
 	void PendingToPopulatedFastPathAlreadyHasData(FComponentState& State);
 	void PendingToPopulatedFastPathNoGrass(FComponentState& State);
 	void PendingToStreaming(FComponentState& State);
+
+	// pending component heap add/remove
+	void AddToPendingComponentHeap(FComponentState* State);
 	void RemoveFromPendingComponentHeap(FComponentState* State);
 
 	void CompleteAllAsyncTasksNow();
