Build fix for CL 36207520
Replace a NaN with something else (NaN not supported)

diff --git a/Engine/Source/Runtime/Landscape/Private/LandscapeGrassMapsBuilder.cpp b/Engine/Source/Runtime/Landscape/Private/LandscapeGrassMapsBuilder.cpp
--- a/Engine/Source/Runtime/Landscape/Private/LandscapeGrassMapsBuilder.cpp
+++ b/Engine/Source/Runtime/Landscape/Private/LandscapeGrassMapsBuilder.cpp
@@ -117,7 +117,7 @@ DECLARE_CYCLE_STAT(TEXT("Remove Grass Instances"), STAT_RemoveGrassInstances, ST
 FLandscapeGrassMapsBuilder::FLandscapeGrassMapsBuilder(UWorld* InOwner, FLandscapeTextureStreamingManager& InTextureStreamingManager)
 	: World(InOwner)
 	, TextureStreamingManager(InTextureStreamingManager)
-	, PreviousCameraHashGrid(GGrassMapCameraCutTranslationThreshold, FVector(NAN, NAN, NAN))
+	, PreviousCameraHashGrid(GGrassMapCameraCutTranslationThreshold, FVector(TNumericLimits<double>::Lowest()))
 #if WITH_EDITOR
 	, OutdatedGrassMapCount(0)
 	, GrassMapsLastCheckTime(0)
