Fixed ensure when analyzing the grass map count in parallel, but on the game thread

diff --git a/Engine/Source/Runtime/Landscape/Private/LandscapeSubsystem.cpp b/Engine/Source/Runtime/Landscape/Private/LandscapeSubsystem.cpp
--- a/Engine/Source/Runtime/Landscape/Private/LandscapeSubsystem.cpp
+++ b/Engine/Source/Runtime/Landscape/Private/LandscapeSubsystem.cpp
@@ -844,7 +844,7 @@ TArray<TTuple<ALandscapeProxy*, UE::Landscape::EOutdatedDataFlags>> ULandscapeSu
 	TArray<ProxyAndFlagsArray> TaskContexts;
 	ParallelForWithTaskContext(TaskContexts, Proxies.Num(), [this, InMatchingOutdatedDataFlags, bInMustMatchAllFlags](ProxyAndFlagsArray& InTaskContext, int32 InIndex)
 		{
-			FTaskTagScope Scope(ETaskTag::EParallelGameThread);
+			FOptionalTaskTagScope Scope(ETaskTag::EParallelGameThread);
 			ALandscapeProxy* ValidProxy = Proxies[InIndex].Get();
 			const UE::Landscape::EOutdatedDataFlags ProxyOutdatedDataFlags = ValidProxy->GetOutdatedDataFlags();
 
