Workaround for landscape crash because of incorrect world at registration when using sublevels

(Converts the crash into a warning)

diff --git a/Engine/Source/Runtime/Landscape/Classes/LandscapeProxy.h b/Engine/Source/Runtime/Landscape/Classes/LandscapeProxy.h
--- a/Engine/Source/Runtime/Landscape/Classes/LandscapeProxy.h
+++ b/Engine/Source/Runtime/Landscape/Classes/LandscapeProxy.h
@@ -895,6 +895,10 @@ PRAGMA_ENABLE_DEPRECATION_WARNINGS
 	/** True when this Proxy is registered with the ULandscapeSubsystem */
 	bool bIsRegisteredWithSubsystem : 1 = false;
 
+	/** Used to detect incorrect registration/unregistration behavior */
+	UPROPERTY(Transient, DuplicateTransient, NonTransactional)
+	TObjectPtr<ULandscapeSubsystem> RegisteredToSubsystem;
+
 	/** Data set at creation time */
 	UPROPERTY(meta = (LandscapeInherited))
 	int32 ComponentSizeQuads;    // Total number of quads in each component
diff --git a/Engine/Source/Runtime/Landscape/Private/LandscapeSubsystem.cpp b/Engine/Source/Runtime/Landscape/Private/LandscapeSubsystem.cpp
--- a/Engine/Source/Runtime/Landscape/Private/LandscapeSubsystem.cpp
+++ b/Engine/Source/Runtime/Landscape/Private/LandscapeSubsystem.cpp
@@ -42,6 +42,16 @@
 
 #include UE_INLINE_GENERATED_CPP_BY_NAME(LandscapeSubsystem)
 
+// enables debug spew
+//#define ENABLE_LANDSCAPE_SUBSYSTEM_DEBUG_SPEW 1
+#ifdef ENABLE_LANDSCAPE_SUBSYSTEM_DEBUG_SPEW
+#define SUBSYSTEM_DEBUG_LOG(...) UE_LOG(LogLandscape, Warning, __VA_ARGS__)
+#define SUBSYSTEM_DEBUG_LOG_REGISTER(...) UE_LOG(LogLandscape, Warning, __VA_ARGS__)
+#else
+#define SUBSYSTEM_DEBUG_LOG(...) UE_LOG(LogLandscape, Verbose, __VA_ARGS__)
+#define SUBSYSTEM_DEBUG_LOG_REGISTER(...) do {} while(0)
+#endif // ENABLE_LANDSCAPE_EDGE_FIXUP_DEBUG_SPEW
+
 #if WITH_EDITOR
 #include "ActionableMessageSubsystem.h"
 #include "Async/ParallelFor.h"
@@ -55,7 +65,6 @@ static FAutoConsoleVariableRef CVarUseStreamingManagerForCameras(
 	GUseStreamingManagerForCameras,
 	TEXT("1: Use Streaming Manager; 0: Use ViewLocationsRenderedLastFrame"));
 
-
 static FAutoConsoleVariable CVarMaxAsyncNaniteProxiesPerSecond(
 	TEXT("landscape.Nanite.MaxAsyncProxyBuildsPerSecond"),
 	6.0f,
@@ -144,16 +153,20 @@ ULandscapeSubsystem::~ULandscapeSubsystem()
 void ULandscapeSubsystem::RegisterActor(ALandscapeProxy* Proxy)
 {
 	check(Proxy != nullptr);
+	SUBSYSTEM_DEBUG_LOG_REGISTER(TEXT("ULandscapeSubsystem::RegisterActor %p %s (%s) to subsystem %p in world %p (%s)"),
+		Proxy, *Proxy->GetFullName(), *Proxy->GetClass()->GetName(), this, GetWorld(), *GetWorld()->GetName());
 
 	TObjectPtr<ALandscapeProxy> ProxyPtr(Proxy);
 
 	// in editor we can get multiple registration calls, so ensure we don't register more than once
 	if (Proxy->bIsRegisteredWithSubsystem)
 	{
-		checkSlow(Proxies.Contains(ProxyPtr));
+		check(Proxy->RegisteredToSubsystem == this);
+		check(Proxies.Contains(ProxyPtr));
 		return;
 	}
-	checkSlow(!Proxies.Contains(ProxyPtr));
+	check(Proxy->RegisteredToSubsystem == nullptr);
+	check(!Proxies.Contains(ProxyPtr));
 	Proxies.Add(ProxyPtr);
 	
 	if (ALandscape* LandscapeActor = Cast<ALandscape>(Proxy))
@@ -168,19 +181,33 @@ void ULandscapeSubsystem::RegisterActor(ALandscapeProxy* Proxy)
 	}
 
 	Proxy->bIsRegisteredWithSubsystem = true;
+	Proxy->RegisteredToSubsystem = this;
 }
 
 void ULandscapeSubsystem::UnregisterActor(ALandscapeProxy* Proxy)
 {
 	check(Proxy != nullptr);
+	SUBSYSTEM_DEBUG_LOG_REGISTER(TEXT("ULandscapeSubsystem::UnregisterActor %p %s (%s) to subsystem %p in world %p (%s)"),
+		Proxy, *Proxy->GetFullName(), *Proxy->GetClass()->GetName(), this, GetWorld(), *GetWorld()->GetName());
+
 	TObjectPtr<ALandscapeProxy> ProxyPtr(Proxy);
 
 	// in editor we can get multiple unregistration calls, so ensure we don't register more than once
 	if (!Proxy->bIsRegisteredWithSubsystem)
 	{
+		check(Proxy->RegisteredToSubsystem == nullptr);
 		check(!Proxies.Contains(ProxyPtr));
 		return;
 	}
+	if (Proxy->RegisteredToSubsystem != this)
+	{
+		UE_LOG(LogLandscape, Warning, TEXT("Landscape Proxy %s was registered to world '%s' but is being unregistered from world '%s', this may indicate that worlds were changed without re-registering actors, which may cause issues in the landscape system.  We will assume it should be unregistered from the original world."),
+			*Proxy->GetFullName(),
+			*Proxy->RegisteredToSubsystem->GetWorld()->GetName(),
+			*GetWorld()->GetName());
+		Proxy->RegisteredToSubsystem->UnregisterActor(Proxy);
+		return;
+	}
 	int32 RemovedCount = Proxies.Remove(ProxyPtr);
 	check(RemovedCount == 1);
 
@@ -195,6 +222,7 @@ void ULandscapeSubsystem::UnregisterActor(ALandscapeProxy* Proxy)
 	}
 
 	Proxy->bIsRegisteredWithSubsystem = false;
+	Proxy->RegisteredToSubsystem = nullptr;
 }
 
 void ULandscapeSubsystem::AddReferencedObjects(UObject* InThis, FReferenceCollector& Collector)
