Moved landscape visibility layer info asset to from EditorLandscapeResources to EngineResources so that it doesn't get stripped on cook when "Exclude editor content when cooking" is used

#virtualized

diff --git a/Engine/Config/BaseEngine.ini b/Engine/Config/BaseEngine.ini
--- a/Engine/Config/BaseEngine.ini
+++ b/Engine/Config/BaseEngine.ini
@@ -1286,6 +1286,7 @@ DefaultWorldPartitionHLODLayer=/Engine/Maps/Templates/HLODs/HLODLayer_Instanced.
 +PackageRedirects=(OldName="/Script/StructUtils", NewName="/Script/CoreUObject")
 +PackageRedirects=(OldName="/Script/StructUtilsEngine", NewName="/Script/Engine")
 +PackageRedirects=(OldName="/Script/StructUtilsNodes", NewName="/Script/BlueprintGraph")
++PackageRedirects=(OldName="/Engine/EditorLandscapeResources/DataLayer.DataLayer", NewName="/Engine/EngineResources/LandscapeVisibilityLayerInfo.LandscapeVisibilityLayerInfo")
 
 +ClassRedirects=(OldName="/Script/Engine.UserDefinedStruct",NewName="/Script/CoreUObject.UserDefinedStruct")
 +ClassRedirects=(OldName="/Script/StructUtils.StructUtilsFunctionLibrary",NewName="/Script/Engine.BlueprintInstancedStructLibrary")
diff --git a/Engine/Source/Runtime/Landscape/Private/Landscape.cpp b/Engine/Source/Runtime/Landscape/Private/Landscape.cpp
--- a/Engine/Source/Runtime/Landscape/Private/Landscape.cpp
+++ b/Engine/Source/Runtime/Landscape/Private/Landscape.cpp
@@ -1624,15 +1624,15 @@ ALandscapeProxy::ALandscapeProxy(const FObjectInitializer& ObjectInitializer)
 		// Structure to hold one-time initialization
 		struct FConstructorStatics
 		{
-			ConstructorHelpers::FObjectFinderOptional<ULandscapeLayerInfoObject> DataLayer;
+			ConstructorHelpers::FObjectFinderOptional<ULandscapeLayerInfoObject> LandscapeVisibilityLayerInfoFinder;
 			FConstructorStatics()
-				: DataLayer(TEXT("LandscapeLayerInfoObject'/Engine/EditorLandscapeResources/DataLayer.DataLayer'"))
+				: LandscapeVisibilityLayerInfoFinder(TEXT("LandscapeLayerInfoObject'/Engine/EngineResources/LandscapeVisibilityLayerInfo.LandscapeVisibilityLayerInfo'"))
 			{
 			}
 		};
 		static FConstructorStatics ConstructorStatics;
 
-		VisibilityLayer = ConstructorStatics.DataLayer.Get();
+		VisibilityLayer = ConstructorStatics.LandscapeVisibilityLayerInfoFinder.Get();
 		check(VisibilityLayer);
 #if WITH_EDITORONLY_DATA
 		// This layer should be no weight blending
