Fixing LandscapeComponent not having the physics material set in editor.
Editor queries were previously not returning a valid physics material even when set, now the physics material will be correctly returned.

diff --git a/Engine/Source/Runtime/Landscape/Private/LandscapeCollision.cpp b/Engine/Source/Runtime/Landscape/Private/LandscapeCollision.cpp
--- a/Engine/Source/Runtime/Landscape/Private/LandscapeCollision.cpp
+++ b/Engine/Source/Runtime/Landscape/Private/LandscapeCollision.cpp
@@ -454,6 +454,7 @@ void ULandscapeHeightfieldCollisionComponent::OnCreatePhysicsState()
 
 					NewEditorShape->SetQueryData(QueryFilterDataEd);
 					NewEditorShape->SetSimData(SimFilterDataEd);
+					NewEditorShape->SetMaterials(HeightfieldRef->UsedChaosMaterials);
 
 					Geoms.Emplace(MoveTemp(ChaosEditorHeightFieldFromCooked));
 					ShapeArray.Emplace(MoveTemp(NewEditorShape));
