Fix for ensure hit when sculpting off the edge of landscape.

Also adding some comments, and fixing a few found issues.

diff --git a/Engine/Source/Editor/LandscapeEditor/Private/LandscapeEdModeTools.h b/Engine/Source/Editor/LandscapeEditor/Private/LandscapeEdModeTools.h
--- a/Engine/Source/Editor/LandscapeEditor/Private/LandscapeEdModeTools.h
+++ b/Engine/Source/Editor/LandscapeEditor/Private/LandscapeEdModeTools.h
@@ -312,8 +312,10 @@ public:
 				int32 ValidX2 = X2;
 				int32 ValidY2 = Y2;
 				DataAccess.GetData(ValidX1, ValidY1, ValidX2, ValidY2, CachedData);
-				if (!ensureMsgf(ValidX1 <= ValidX2 && ValidY1 <= ValidY2, TEXT("Invalid cache area: X(%d-%d), Y(%d-%d) from region X(%d-%d), Y(%d-%d)"), ValidX1, ValidX2, ValidY1, ValidY2, X1, X2, Y1, Y2))
+
+				if (ValidX1 > ValidX2 || ValidY1 > ValidY2)
 				{
+					// there was no data in the requested region
 					bIsValid = false;
 					return;
 				}
@@ -677,8 +679,11 @@ struct FXYOffsetmapAccessor
 	// accessors
 	void GetData(int32& X1, int32& Y1, int32& X2, int32& Y2, TMap<FIntPoint, FVector>& Data)
 	{
+		// NOTE: GetXYOffsetData may shrink X1/Y1/X2/Y2 to reflect the sub-range that contains valid data
 		LandscapeEdit->GetXYOffsetData(X1, Y1, X2, Y2, Data);
 
+		// NOTE: GetHeightData may shrink X1/Y1/X2/Y2 to reflect the sub-range that contains valid data
+		// In general, GetHeight should not further shrink the range, as anywhere the XYOffsets are valid, the Heights should also be valid.
 		TMap<FIntPoint, uint16> NewHeights;
 		LandscapeEdit->GetHeightData(X1, Y1, X2, Y2, NewHeights);
 		for (int32 Y = Y1; Y <= Y2; ++Y)
@@ -698,6 +703,9 @@ struct FXYOffsetmapAccessor
 	{
 		LandscapeEdit->GetXYOffsetDataFast(X1, Y1, X2, Y2, Data);
 
+		// NOTE: this can modify X1/Y1/X2/Y2 to reflect the region with valid height data, which may be a subset of the original range
+		// We rely on the fact that generally anywhere the XYOffsets are valid, the Heights should also be valid
+		// So in theory we should write the Z value for every XY Offset that was queried above.
 		TMap<FIntPoint, uint16> NewHeights;
 		LandscapeEdit->GetHeightData(X1, Y1, X2, Y2, NewHeights);
 		for (int32 Y = Y1; Y <= Y2; ++Y)
diff --git a/Engine/Source/Runtime/Landscape/Private/LandscapeEditInterface.cpp b/Engine/Source/Runtime/Landscape/Private/LandscapeEditInterface.cpp
--- a/Engine/Source/Runtime/Landscape/Private/LandscapeEditInterface.cpp
+++ b/Engine/Source/Runtime/Landscape/Private/LandscapeEditInterface.cpp
@@ -539,6 +539,10 @@ void FLandscapeEditDataInterface::RecalculateNormals()
 		int32 Y1 = Component->GetSectionBase().Y-1;
 		int32 X2 = Component->GetSectionBase().X+ComponentSizeQuads+1;
 		int32 Y2 = Component->GetSectionBase().Y+ComponentSizeQuads+1;
+		const int32 OriginalX1 = X1;
+		const int32 OriginalY1 = Y1;
+		const int32 OriginalX2 = X2;
+		const int32 OriginalY2 = Y2;
 		int32 Stride = ComponentSizeQuads+3; 
 		int32 StrideSquared = FMath::Square(Stride);
 
@@ -550,8 +554,11 @@ void FLandscapeEditDataInterface::RecalculateNormals()
 
 		// Get XY offset
 		GetXYOffsetDataFast(X1,Y1,X2,Y2,XYOffsets,0);
+		check(X1 == OriginalX1 && Y1 == OriginalY1 && X2 == OriginalX2 && Y2 == OriginalY2);	// because we're querying within an existing component, we assume the data should always exist
+
 		// Get the vertex positions for entire quad
 		GetHeightData(X1,Y1,X2,Y2,HeightData,0);
+		check(X1 == OriginalX1 && Y1 == OriginalY1 && X2 == OriginalX2 && Y2 == OriginalY2);	// because we're querying within an existing component, we assume the data should always exist
 
 		// Contribute face normals for all triangles contributing to this components' normals
 		for( int32 Y=0;Y<Stride-1;Y++ )
@@ -1466,11 +1473,15 @@ void FLandscapeEditDataInterface::GetHeightDataInternal(int32& ValidX1, int32& V
 
 	if (bHasMissingValue)
 	{
+		// generate something to fill in any missing data within the requested region
 		CalcMissingValues<typename TDataAccess::DataType, TDataAccess, float>(X1, X2, Y1, Y2,
 			ComponentIndexX1, ComponentIndexX2, ComponentIndexY1, ComponentIndexY2,
 			ComponentSizeX, ComponentSizeY, CornerValues,
 			NoBorderY1, NoBorderY2, ComponentDataExist, StoreData);
-		// Update valid region
+
+		// clamp the returned bounds to the original request bounds
+		// this means the returned bounds are allowed to shrink (if there is only a smaller part of the request that actually exists)
+		// but it is not allowed to expand beyond the original request
 		ValidX1 = FMath::Max<int32>(X1, ValidX1);
 		ValidX2 = FMath::Min<int32>(X2, ValidX2);
 		ValidY1 = FMath::Max<int32>(Y1, ValidY1);
@@ -1478,6 +1489,7 @@ void FLandscapeEditDataInterface::GetHeightDataInternal(int32& ValidX1, int32& V
 	}
 	else
 	{
+		// if there were no missing values, then the entire requested bounds is returned
 		ValidX1 = X1;
 		ValidX2 = X2;
 		ValidY1 = Y1;
@@ -4272,11 +4284,15 @@ void FLandscapeEditDataInterface::GetWeightDataTempl(ULandscapeLayerInfoObject*
 
 	if (bHasMissingValue)
 	{
+		// generate something to fill in any missing data within the requested region
 		CalcMissingValues<uint8, TStoreData, float>( X1, X2, Y1, Y2,
 			ComponentIndexX1, ComponentIndexX2, ComponentIndexY1, ComponentIndexY2, 
 			ComponentSizeX, ComponentSizeY, CornerValues,
 			NoBorderY1, NoBorderY2, ComponentDataExist, StoreData );
-		// Update valid region
+
+		// clamp the returned bounds to the original request bounds
+		// this means the returned bounds are allowed to shrink (if there is only a smaller part of the request that actually exists)
+		// but it is not allowed to expand beyond the original request
 		ValidX1 = FMath::Max<int32>(X1, ValidX1);
 		ValidX2 = FMath::Min<int32>(X2, ValidX2);
 		ValidY1 = FMath::Max<int32>(Y1, ValidY1);
@@ -4284,6 +4300,7 @@ void FLandscapeEditDataInterface::GetWeightDataTempl(ULandscapeLayerInfoObject*
 	}
 	else
 	{
+		// if there were no missing values, then the entire requested bounds is returned
 		ValidX1 = X1;
 		ValidX2 = X2;
 		ValidY1 = Y1;
@@ -5501,11 +5518,15 @@ void FLandscapeEditDataInterface::GetXYOffsetDataTempl(int32& ValidX1, int32& Va
 
 	if (bHasMissingValue)
 	{
+		// generate something to fill in any missing data within the requested region
 		CalcMissingValues<FVector2D, TStoreData, FVector2D>(X1, X2, Y1, Y2,
 			ComponentIndexX1, ComponentIndexX2, ComponentIndexY1, ComponentIndexY2,
 			ComponentSizeX, ComponentSizeY, CornerValues,
 			NoBorderY1, NoBorderY2, ComponentDataExist, StoreData);
-		// Update valid region
+
+		// clamp the returned bounds to the original request bounds
+		// this means the returned bounds are allowed to shrink (if there is only a small part of the request that actually exists)
+		// but it is not allowed to expand beyond the original request
 		ValidX1 = FMath::Max<int32>(X1, ValidX1);
 		ValidX2 = FMath::Min<int32>(X2, ValidX2);
 		ValidY1 = FMath::Max<int32>(Y1, ValidY1);
@@ -5513,6 +5534,7 @@ void FLandscapeEditDataInterface::GetXYOffsetDataTempl(int32& ValidX1, int32& Va
 	}
 	else
 	{
+		// if there were no missing values, then the entire requested bounds is returned
 		ValidX1 = X1;
 		ValidX2 = X2;
 		ValidY1 = Y1;
diff --git a/Engine/Source/Runtime/Landscape/Private/LandscapeSplineRaster.cpp b/Engine/Source/Runtime/Landscape/Private/LandscapeSplineRaster.cpp
--- a/Engine/Source/Runtime/Landscape/Private/LandscapeSplineRaster.cpp
+++ b/Engine/Source/Runtime/Landscape/Private/LandscapeSplineRaster.cpp
@@ -276,15 +276,16 @@ void RasterizeHeight(int32& MinX, int32& MinY, int32& MaxX, int32& MaxY, FLandsc
 
 	if (ValidMinX > ValidMaxX || ValidMinY > ValidMaxY)
 	{
-		// The bounds don't intersect any data, so skip it
+		// The bounds don't intersect any data, so skip it.  Return the invalid bounds.
 		MinX = ValidMinX;
 		MinY = ValidMinY;
 		MaxX = ValidMaxX;
 		MaxY = ValidMaxY;
-
+			
 		return;
 	}
 
+	// shrink our data to the valid range returned by GetHeightData
 	FLandscapeEditDataInterface::ShrinkData(HeightData, MinX, MinY, MaxX, MaxY, ValidMinX, ValidMinY, ValidMaxX, ValidMaxY);
 
 	const ALandscape* Landscape = LandscapeEdit.GetTargetLandscape();
@@ -394,7 +395,7 @@ void RasterizeControlPointAlpha(int32& MinX, int32& MinY, int32& MaxX, int32& Ma
 
 	if (ValidMinX > ValidMaxX || ValidMinY > ValidMaxY)
 	{
-		// The control point's bounds don't intersect any data, so skip it
+		// The control point's bounds don't intersect any data, so skip it.  Return the invalid bounds.
 		MinX = ValidMinX;
 		MinY = ValidMinY;
 		MaxX = ValidMaxX;
@@ -403,6 +404,7 @@ void RasterizeControlPointAlpha(int32& MinX, int32& MinY, int32& MaxX, int32& Ma
 		return;
 	}
 
+	// shrink our data to the valid range returned by GetWeightData
 	FLandscapeEditDataInterface::ShrinkData(Data, MinX, MinY, MaxX, MaxY, ValidMinX, ValidMinY, ValidMaxX, ValidMaxY);
 
 	MinX = ValidMinX;
@@ -510,7 +512,7 @@ void RasterizeSegmentAlpha(int32& MinX, int32& MinY, int32& MaxX, int32& MaxY, F
 
 	if (ValidMinX > ValidMaxX || ValidMinY > ValidMaxY)
 	{
-		// The segment's bounds don't intersect any data, so skip it
+		// The segment's bounds don't intersect any data, so skip it.  Return the invalid bounds.
 		MinX = ValidMinX;
 		MinY = ValidMinY;
 		MaxX = ValidMaxX;
@@ -519,6 +521,7 @@ void RasterizeSegmentAlpha(int32& MinX, int32& MinY, int32& MaxX, int32& MaxY, F
 		return;
 	}
 
+	// shrink our data to the valid range returned by GetWeightData
 	FLandscapeEditDataInterface::ShrinkData(Data, MinX, MinY, MaxX, MaxY, ValidMinX, ValidMinY, ValidMaxX, ValidMaxY);
 
 	MinX = ValidMinX;
@@ -765,8 +768,9 @@ bool ULandscapeInfo::ApplySplinesInternal(bool bOnlySelected, TScriptInterface<I
 
 			if (MinX > MaxX || MinY > MaxY)
 			{
-				// The segment's bounds don't intersect any data, so we skip it entirely
-				// it wouldn't intersect any weightmap data either so we don't even bother trying
+				// The segment's bounds don't intersect any height data.
+				// This means it wouldn't intersect any weightmap data either so we can skip the alpha rasterize below
+				continue;
 			}
 		}
 
