Fix for an editor crash in grass builder when the landscape component scene proxy is removed while grass generation for that component is in the pipeline

Also removing the grass data bIsDirty flag, that is no longer used.

diff --git a/Engine/Source/Runtime/Landscape/Classes/LandscapeComponent.h b/Engine/Source/Runtime/Landscape/Classes/LandscapeComponent.h
--- a/Engine/Source/Runtime/Landscape/Classes/LandscapeComponent.h
+++ b/Engine/Source/Runtime/Landscape/Classes/LandscapeComponent.h
@@ -219,9 +219,6 @@ struct FLandscapeComponentGrassData
 #if WITH_EDITORONLY_DATA
 	// Height data for LODs 1+, keyed on LOD index
 	TMap<int32, TArray<uint16>> HeightMipData;
-
-	// Grass data was updated but not saved yet
-	bool bIsDirty = false;
 #endif // WITH_EDITORONLY_DATA
 	
 	static constexpr int32 UnknownNumElements = -1;
diff --git a/Engine/Source/Runtime/Landscape/Private/Landscape.cpp b/Engine/Source/Runtime/Landscape/Private/Landscape.cpp
--- a/Engine/Source/Runtime/Landscape/Private/Landscape.cpp
+++ b/Engine/Source/Runtime/Landscape/Private/Landscape.cpp
@@ -2250,16 +2250,16 @@ void ULandscapeComponent::OnUnregister()
 	PhysicalMaterialTask.Release();
 #endif
 
-	if (GetLandscapeProxy())
+	if (ALandscapeProxy* Proxy = GetLandscapeProxy())
 	{
 		// Generate MID representing the MIC
-		if (GetLandscapeProxy()->bUseDynamicMaterialInstance)
+		if (Proxy->bUseDynamicMaterialInstance)
 		{
 			MaterialInstancesDynamic.Empty();
 		}
 
 		// AActor::GetWorld checks for Unreachable and BeginDestroyed
-		UWorld* World = GetLandscapeProxy()->GetWorld();
+		UWorld* World = Proxy->GetWorld();
 
 		if (World)
 		{
@@ -3763,8 +3763,6 @@ void ALandscapeProxy::PreSave(FObjectPreSaveContext ObjectSaveContext)
 			int32 ValidGrassCount = 0;
 			for (ULandscapeComponent* Component : LandscapeComponents)
 			{
-				// Manually reset dirty flag (for post save)
-				Component->GrassData->bIsDirty = false;
 				if (Component->GrassData->HasValidData())
 				{
 					ValidGrassCount++;
diff --git a/Engine/Source/Runtime/Landscape/Private/LandscapeGrass.cpp b/Engine/Source/Runtime/Landscape/Private/LandscapeGrass.cpp
--- a/Engine/Source/Runtime/Landscape/Private/LandscapeGrass.cpp
+++ b/Engine/Source/Runtime/Landscape/Private/LandscapeGrass.cpp
@@ -1023,10 +1023,6 @@ void ULandscapeComponent::RemoveGrassMap()
 	// this ensures if anyone else is accessing the old grassdata (while holding a shared ref to it), it won't be modified or deleted
 	// this is also important for PIE, which can share grassdatas with editor
 	GrassData = MakeShared<FLandscapeComponentGrassData>();
-
-#if WITH_EDITOR
-	GrassData->bIsDirty = true;
-#endif // WITH_EDITOR
 }
 
 void ALandscapeProxy::UpdateGrassTypeSummary()
diff --git a/Engine/Source/Runtime/Landscape/Private/LandscapeGrassMapsBuilder.cpp b/Engine/Source/Runtime/Landscape/Private/LandscapeGrassMapsBuilder.cpp
--- a/Engine/Source/Runtime/Landscape/Private/LandscapeGrassMapsBuilder.cpp
+++ b/Engine/Source/Runtime/Landscape/Private/LandscapeGrassMapsBuilder.cpp
@@ -585,6 +585,17 @@ bool FLandscapeGrassMapsBuilder::UpdateTrackedComponents(const TArray<FVector>&
 
 			if (State->AreTexturesStreamedIn())
 			{
+#if WITH_EDITOR
+				// in editor the renderability can change unexpectedly -- so check one last time just before we actually render
+				// (in game, we only check at the beginning of kicking off the generation process)
+				if (!UE::Landscape::CanRenderGrassMap(State->Component))
+				{
+					// can't render, move to NotReady state, which will monitor until it is renderable
+					StreamingToNotReady(*State);
+					bChanged= true;
+					continue;
+				}
+#endif // WITH_EDITOR
 				KickOffRenderAndReadback(*State);
 				bRenderCommandsQueuedByLastUpdate = true;
 				bChanged = true;
@@ -1254,6 +1265,18 @@ bool FLandscapeGrassMapsBuilder::StartGrassMapGeneration(FComponentState& State,
 	return true;
 }
 
+void FLandscapeGrassMapsBuilder::StreamingToNotReady(FComponentState& State)
+{
+	check(StreamingCount > 0);
+	StreamingCount--;
+	RemoveTextureStreamingRequests(State);
+	State.Stage = EComponentStage::NotReady;
+	NotReadyCount++;
+
+	DEBUG_TRANSITION(State, Streaming, NotReady);
+	State.TickCount = 0;
+}
+
 void FLandscapeGrassMapsBuilder::PendingToNotReady(FComponentState& State)
 {
 	check(PendingCount > 0);
diff --git a/Engine/Source/Runtime/Landscape/Private/LandscapeGrassMapsBuilder.h b/Engine/Source/Runtime/Landscape/Private/LandscapeGrassMapsBuilder.h
--- a/Engine/Source/Runtime/Landscape/Private/LandscapeGrassMapsBuilder.h
+++ b/Engine/Source/Runtime/Landscape/Private/LandscapeGrassMapsBuilder.h
@@ -270,6 +270,7 @@ private:
 		void RemoveTextureStreamingRequests(FComponentState& State);
 
 	// state transition helpers
+	void StreamingToNotReady(FComponentState& State);
 	void PendingToNotReady(FComponentState& State);
 	void PendingToPopulatedFastPathAlreadyHasData(FComponentState& State);
 	void PendingToPopulatedFastPathNoGrass(FComponentState& State);
diff --git a/Engine/Source/Runtime/Landscape/Private/LandscapeGrassWeightExporter.cpp b/Engine/Source/Runtime/Landscape/Private/LandscapeGrassWeightExporter.cpp
--- a/Engine/Source/Runtime/Landscape/Private/LandscapeGrassWeightExporter.cpp
+++ b/Engine/Source/Runtime/Landscape/Private/LandscapeGrassWeightExporter.cpp
@@ -826,10 +826,6 @@ void FLandscapeGrassWeightExporter::ApplyResults(TMap<ULandscapeComponent*, TUni
 		// Assign the new data (thread-safe)
 		Component->GrassData = MakeShareable(ComponentGrassData);
 
-#if WITH_EDITORONLY_DATA
-		Component->GrassData->bIsDirty = true;
-#endif // WITH_EDITORONLY_DATA
-
 #if WITH_EDITOR
 		if (Proxy->bBakeMaterialPositionOffsetIntoCollision)
 		{
