Mobile: Disable GPUScene support for a landscape VF. Atm we can handle static lighting correctly for landscape components while GPUScene is on

diff --git a/Engine/Source/Runtime/Landscape/Private/LandscapeRender.cpp b/Engine/Source/Runtime/Landscape/Private/LandscapeRender.cpp
--- a/Engine/Source/Runtime/Landscape/Private/LandscapeRender.cpp
+++ b/Engine/Source/Runtime/Landscape/Private/LandscapeRender.cpp
@@ -3601,8 +3601,8 @@ void FLandscapeVertexFactory::InitRHI(FRHICommandListBase& RHICmdList)
 	// position decls
 	Elements.Add(AccessStreamComponent(Data.PositionComponent, 0));
 
-	// Use the same attribute on mobile and non-mobile, to enable the GPUScene path on both.
-	AddPrimitiveIdStreamElement(EVertexInputStreamType::Default, Elements, /* AttributeIndex = */ 1, /* AttributeIndex_Mobile = */ 1);
+	// see FLandscapeVertexFactory::ModifyCompilationEnvironment for mobile GPUScene exception
+	AddPrimitiveIdStreamElement(EVertexInputStreamType::Default, Elements, /* AttributeIndex = */ 1, /* AttributeIndex_Mobile = */0xFF);
 	// create the actual device decls
 	InitDeclaration(Elements);
 }
@@ -3623,7 +3623,11 @@ void FLandscapeVertexFactory::ModifyCompilationEnvironment(const FVertexFactoryS
 {
 	FVertexFactory::ModifyCompilationEnvironment(Parameters, OutEnvironment);
 
-	OutEnvironment.SetDefine(TEXT("VF_SUPPORTS_PRIMITIVE_SCENE_DATA"), Parameters.VertexFactoryType->SupportsPrimitiveIdStream() && UseGPUScene(Parameters.Platform, GetMaxSupportedFeatureLevel(Parameters.Platform)));
+	const FStaticFeatureLevel MaxSupportedFeatureLevel = GetMaxSupportedFeatureLevel(Parameters.Platform);
+	// TODO: support GPUScene on mobile. We need to pass a correct LODLightmapDataIndex to a batching CS, which we do only for StaticMesh atm only
+	const bool bUseGPUScene = UseGPUScene(Parameters.Platform, MaxSupportedFeatureLevel) && (MaxSupportedFeatureLevel > ERHIFeatureLevel::ES3_1);
+
+	OutEnvironment.SetDefine(TEXT("VF_SUPPORTS_PRIMITIVE_SCENE_DATA"), Parameters.VertexFactoryType->SupportsPrimitiveIdStream() && bUseGPUScene);
 
 	// Make sure landscape vertices go back to local space so that we have consistency between the transform on normals and geometry
 	OutEnvironment.SetDefine(TEXT("RAY_TRACING_DYNAMIC_MESH_IN_LOCAL_SPACE"), TEXT("1"));
@@ -3634,7 +3638,7 @@ void FLandscapeVertexFactory::GetPSOPrecacheVertexFetchElements(EVertexInputStre
 	Elements.Add(FVertexElement(0, 0, VET_UByte4, 0, sizeof(FLandscapeVertex), false));
 	
 	if (UseGPUScene(GMaxRHIShaderPlatform, GMaxRHIFeatureLevel)
-		&& !PlatformGPUSceneUsesUniformBufferView(GMaxRHIShaderPlatform))
+		&& GMaxRHIFeatureLevel > ERHIFeatureLevel::ES3_1) // see FLandscapeVertexFactory::ModifyCompilationEnvironment
 	{
 		Elements.Add(FVertexElement(1, 0, VET_UInt, 1, sizeof(uint32), true));
 	}
