Added missing file for previous CL
#buildfix
#rnx

diff --git a/Engine/Source/Runtime/Landscape/Classes/LandscapeSplineActor.h b/Engine/Source/Runtime/Landscape/Classes/LandscapeSplineActor.h
--- a/Engine/Source/Runtime/Landscape/Classes/LandscapeSplineActor.h
+++ b/Engine/Source/Runtime/Landscape/Classes/LandscapeSplineActor.h
@@ -34,6 +34,7 @@ public:
 	virtual void UpdateSharedProperties(ULandscapeInfo* InLandscapeInfo) override;
 		
 #if WITH_EDITOR
+	LANDSCAPE_API bool ShouldPartitionSpline() const;
 	LANDSCAPE_API bool HasGeneratedLandscapeSplineMeshesActors() const;
 	void SetLandscapeGuid(const FGuid& InGuid);
 	virtual void GetActorDescProperties(FPropertyPairsMap& PropertyPairsMap) const override;
