[Landscape]

Fix lingering hard reference to deleted landscapes being held by landscape info objects.

diff --git a/Engine/Source/Editor/LandscapeEditor/Private/LandscapeEdMode.cpp b/Engine/Source/Editor/LandscapeEditor/Private/LandscapeEdMode.cpp
--- a/Engine/Source/Editor/LandscapeEditor/Private/LandscapeEdMode.cpp
+++ b/Engine/Source/Editor/LandscapeEditor/Private/LandscapeEdMode.cpp
@@ -2589,8 +2589,8 @@ void FEdModeLandscape::UpdateTargetList()
 					}
 
 					// Construct Thumbnail MIC
-					UMaterialInterface* LandscapeMaterial = LayerSettings.Owner ? LayerSettings.Owner->GetLandscapeMaterial() : UMaterial::GetDefaultMaterial(MD_Surface);
-					LayerSettings.ThumbnailMIC = ALandscapeProxy::GetLayerThumbnailMIC(LandscapeMaterial, LayerName, ThumbnailWeightmap, ThumbnailHeightmap, LayerSettings.Owner);
+					UMaterialInterface* LandscapeMaterial = LayerSettings.Owner.IsValid() ? LayerSettings.Owner.Get()->GetLandscapeMaterial() : UMaterial::GetDefaultMaterial(MD_Surface);
+					LayerSettings.ThumbnailMIC = ALandscapeProxy::GetLayerThumbnailMIC(LandscapeMaterial, LayerName, ThumbnailWeightmap, ThumbnailHeightmap, LayerSettings.Owner.Get());
 				}
 
 				// Add the layer
diff --git a/Engine/Source/Runtime/Landscape/Classes/LandscapeInfo.h b/Engine/Source/Runtime/Landscape/Classes/LandscapeInfo.h
--- a/Engine/Source/Runtime/Landscape/Classes/LandscapeInfo.h
+++ b/Engine/Source/Runtime/Landscape/Classes/LandscapeInfo.h
@@ -77,7 +77,7 @@ struct FLandscapeInfoLayerSettings
 	TObjectPtr<UMaterialInstanceConstant> ThumbnailMIC;
 
 	UPROPERTY()
-	TObjectPtr<ALandscapeProxy> Owner;
+	TWeakObjectPtr<ALandscapeProxy> Owner;
 
 	UPROPERTY(transient)
 	int32 DebugColorChannel;
@@ -100,17 +100,7 @@ struct FLandscapeInfoLayerSettings
 
 	LANDSCAPE_API FLandscapeInfoLayerSettings(ULandscapeLayerInfoObject* InLayerInfo, ALandscapeProxy* InProxy);
 
-	FLandscapeInfoLayerSettings(FName InPlaceholderLayerName, ALandscapeProxy* InProxy)
-		: LayerInfoObj(nullptr)
-		, LayerName(InPlaceholderLayerName)
-#if WITH_EDITORONLY_DATA
-		, ThumbnailMIC(nullptr)
-		, Owner(InProxy)
-		, DebugColorChannel(0)
-		, bValid(false)
-#endif
-	{
-	}
+	LANDSCAPE_API FLandscapeInfoLayerSettings(FName InPlaceholderLayerName, ALandscapeProxy* InProxy);
 
 	LANDSCAPE_API FName GetLayerName() const;
 
diff --git a/Engine/Source/Runtime/Landscape/Private/LandscapeInfo.cpp b/Engine/Source/Runtime/Landscape/Private/LandscapeInfo.cpp
--- a/Engine/Source/Runtime/Landscape/Private/LandscapeInfo.cpp
+++ b/Engine/Source/Runtime/Landscape/Private/LandscapeInfo.cpp
@@ -2,6 +2,7 @@
 
 #include "LandscapeInfo.h"
 #include "LandscapeLayerInfoObject.h"
+#include "LandscapeProxy.h"
 
 #include UE_INLINE_GENERATED_CPP_BY_NAME(LandscapeInfo)
 
@@ -15,3 +16,15 @@ LANDSCAPE_API FLandscapeInfoLayerSettings::FLandscapeInfoLayerSettings(ULandscap
 	, bValid(false)
 #endif
 {}
+
+LANDSCAPE_API FLandscapeInfoLayerSettings::FLandscapeInfoLayerSettings(FName InPlaceholderLayerName, ALandscapeProxy* InProxy)
+		: LayerInfoObj(nullptr)
+		, LayerName(InPlaceholderLayerName)
+#if WITH_EDITORONLY_DATA
+		, ThumbnailMIC(nullptr)
+		, Owner(InProxy)
+		, DebugColorChannel(0)
+		, bValid(false)
+#endif
+	{
+	}
\ No newline at end of file
