Landscape error reporting improvements :
* Removed erroneous "unsynchronized landscape properties" map check error when fixing up unused edit layer data only
* Unified post-load fixup map check messages to display the landscape proxy's actor label (rather than its package path) and let the the user fix the problem by clicking on a "Save Modified Landscapes" token

diff --git a/Engine/Source/Editor/LandscapeEditor/Private/LandscapeEditorModule.cpp b/Engine/Source/Editor/LandscapeEditor/Private/LandscapeEditorModule.cpp
--- a/Engine/Source/Editor/LandscapeEditor/Private/LandscapeEditorModule.cpp
+++ b/Engine/Source/Editor/LandscapeEditor/Private/LandscapeEditorModule.cpp
@@ -14,6 +14,7 @@
 #include "LandscapeProxy.h"
 #include "LandscapeEdMode.h"
 #include "Landscape.h"
+#include "LandscapeUtils.h"
 #include "LandscapeEditorCommands.h"
 #include "Classes/ActorFactoryLandscape.h"
 #include "LandscapeFileFormatPng.h"
@@ -143,7 +144,7 @@ public:
 			FUIAction ActionBuildNanite(FExecuteAction::CreateStatic(&UE::Landscape::BuildNanite), FCanExecuteAction());
 			Section.AddMenuEntry(NAME_None, LOCTEXT("BuildNaniteOnly", "Build Nanite Only"), LOCTEXT("BuildLandscapeNanite", "Build Nanite representation"), TAttribute<FSlateIcon>(), ActionBuildNanite, EUserInterfaceActionType::Button);
 
-			FUIAction ActionSaveModifiedLandscapes(FExecuteAction::CreateStatic(&SaveModifiedLandscapes), FCanExecuteAction::CreateStatic(&HasModifiedLandscapes));
+			FUIAction ActionSaveModifiedLandscapes(FExecuteAction::CreateStatic(&UE::Landscape::SaveModifiedLandscapes), FCanExecuteAction::CreateStatic(&UE::Landscape::HasModifiedLandscapes));
 			Section.AddMenuEntry(NAME_None,
 				LOCTEXT("SaveModifiedLandscapes", "Save Modified Landscapes"), LOCTEXT("SaveModifiedLandscapesToolTip", "Save landscapes that were modified outside of the editor mode"),
 				TAttribute<FSlateIcon>(), ActionSaveModifiedLandscapes, EUserInterfaceActionType::Button);
@@ -231,30 +232,6 @@ public:
 		);
 	}
 
-	static bool HasModifiedLandscapes()
-	{
-		if (UWorld* World = GEditor->GetEditorWorldContext().World())
-		{
-			if (ULandscapeSubsystem* LandscapeSubsystem = World->GetSubsystem<ULandscapeSubsystem>())
-			{
-				return  LandscapeSubsystem->HasModifiedLandscapes();
-			}
-		}
-
-		return false;
-	}
-
-	static void SaveModifiedLandscapes()
-	{
-		if (UWorld* World = GEditor->GetEditorWorldContext().World())
-		{
-			if (ULandscapeSubsystem* LandscapeSubsystem = World->GetSubsystem<ULandscapeSubsystem>())
-			{
-				LandscapeSubsystem->SaveModifiedLandscapes();
-			}
-		}
-	}
-
 	static void ChangeLandscapeViewMode(ELandscapeViewMode::Type ViewMode)
 	{
 		if (ViewMode != GLandscapeViewMode)
diff --git a/Engine/Source/Runtime/Landscape/Private/Landscape.cpp b/Engine/Source/Runtime/Landscape/Private/Landscape.cpp
--- a/Engine/Source/Runtime/Landscape/Private/Landscape.cpp
+++ b/Engine/Source/Runtime/Landscape/Private/Landscape.cpp
@@ -4524,36 +4524,6 @@ void ALandscapeProxy::Destroyed()
 
 namespace UE::Landscape::Private
 {
-	static bool HasModifiedLandscapes()
-	{
-		if (GEditor)
-		{
-			if (UWorld* World = GEditor->GetEditorWorldContext().World())
-			{
-				if (ULandscapeSubsystem* LandscapeSubsystem = World->GetSubsystem<ULandscapeSubsystem>())
-				{
-					return  LandscapeSubsystem->HasModifiedLandscapes();
-				}
-			}
-		}
-
-		return false;
-	}
-
-	static void SaveModifiedLandscapes()
-	{
-		if (GEditor)
-		{
-			if (UWorld* World = GEditor->GetEditorWorldContext().World())
-			{
-				if (ULandscapeSubsystem* LandscapeSubsystem = World->GetSubsystem<ULandscapeSubsystem>())
-				{
-					LandscapeSubsystem->SaveModifiedLandscapes();
-				}
-			}
-		}
-	}
-
 	static bool CopyProperty(FProperty* InProperty, UObject* InSourceObject, UObject* InDestinationObject)
 	{
 		void* SrcValuePtr = InProperty->ContainerPtrToValuePtr<void>(InSourceObject);
@@ -4603,6 +4573,8 @@ namespace UE::Landscape::Private
 
 	static void DisplaySynchronizedPropertiesMapcheckWarning(const TArray<FName>& InSynchronizedProperties, const ALandscapeProxy& InSynchronizedProxy, const ALandscapeProxy& InParentLandscape, const bool bAddSilencingMessage = false)
 	{
+		check(!InSynchronizedProperties.IsEmpty());
+
 		TStringBuilder<1024> SynchronizedPropertiesStringBuilder;
 		ULandscapeSubsystem* LandscapeSubsystem = InSynchronizedProxy.GetWorld() ? InSynchronizedProxy.GetWorld()->GetSubsystem<ULandscapeSubsystem>() : nullptr;
 		checkf(LandscapeSubsystem != nullptr, TEXT("DisplaySynchronizedPropertiesMapcheckWarning can only be called when a subsystem is available"));
@@ -4624,8 +4596,8 @@ namespace UE::Landscape::Private
 			->AddToken(FUObjectToken::Create(&InSynchronizedProxy, FText::FromString(InSynchronizedProxy.GetActorNameOrLabel())))
 			->AddToken(FTextToken::Create(LOCTEXT("MapCheck_Message_LandscapeProxy_FixupSharedData", "had some shared properties not in sync with its parent landscape actor. This has been fixed but the proxy needs to be saved in order to ensure cooking behaves as expected. ")))
 			->AddToken(FActionToken::Create(LOCTEXT("MapCheck_SaveFixedUpData", "Save Modified Landscapes"), LOCTEXT("MapCheck_SaveFixedUpData_Desc", "Saves the modified landscape proxy actors"),
-				FOnActionTokenExecuted::CreateStatic(&SaveModifiedLandscapes),
-				FCanExecuteActionToken::CreateStatic(&HasModifiedLandscapes),
+				FOnActionTokenExecuted::CreateStatic(&UE::Landscape::SaveModifiedLandscapes),
+				FCanExecuteActionToken::CreateStatic(&UE::Landscape::HasModifiedLandscapes),
 				/*bInSingleUse = */false))
 			->AddToken(FTextToken::Create(FText::Format(LOCTEXT("MapCheck_Message_LandscapeProxy_FixupSharedData_SharedProperties", "The following properties were synchronized: {0}."), FText::FromString(SynchronizedPropertiesStringBuilder.ToString()))));
 
@@ -4907,7 +4879,7 @@ void ALandscapeProxy::UpgradeSharedProperties(ALandscape* InParentLandscape)
 
 		if (!CVarSilenceSharedPropertyDeprecationFixup->GetBool())
 		{
-			UE::Landscape::Private::DisplaySynchronizedPropertiesMapcheckWarning(SynchronizedProperties, *this, *InParentLandscape, true);
+			UE::Landscape::Private::DisplaySynchronizedPropertiesMapcheckWarning(SynchronizedProperties, /*InSynchronizedProxy = */*this, *InParentLandscape, /*bAddSilencingMessage = */true);
 		}
 	}
 
@@ -4958,9 +4930,9 @@ void ALandscapeProxy::FixupSharedData(ALandscape* Landscape, const bool bMapChec
 			// Force resave the proxy through the modified landscape system, so that the user can then use the Save Modified Landscapes menu and therefore manually trigger the re-save of all modified proxies. * /
 			bool bNeedsManualResave = LandscapeInfo->MarkObjectDirty(/*InObject = */this, /*bInForceResave = */true, LandscapeActor);
 
-			if (bMapCheck && bNeedsManualResave)
+			if (bMapCheck && bNeedsManualResave && !SynchronizedProperties.IsEmpty())
 			{
-				UE::Landscape::Private::DisplaySynchronizedPropertiesMapcheckWarning(SynchronizedProperties, *this, *Landscape);
+				UE::Landscape::Private::DisplaySynchronizedPropertiesMapcheckWarning(SynchronizedProperties, /*InSynchronizedProxy = */*this, *Landscape);
 			}
 		}
 	}
diff --git a/Engine/Source/Runtime/Landscape/Private/LandscapeEditLayers.cpp b/Engine/Source/Runtime/Landscape/Private/LandscapeEditLayers.cpp
--- a/Engine/Source/Runtime/Landscape/Private/LandscapeEditLayers.cpp
+++ b/Engine/Source/Runtime/Landscape/Private/LandscapeEditLayers.cpp
@@ -12011,11 +12011,14 @@ bool ALandscapeProxy::RemoveObsoleteLayers(const TSet<FGuid>& InExistingLayers)
 			FFormatNamedArguments Arguments;
 			Arguments.Add(TEXT("LayerName"), FText::FromString(LayerGuidAndName.Value.ToString()));
 			Arguments.Add(TEXT("LayerGuid"), FText::FromString(LayerGuidAndName.Key.ToString(EGuidFormats::HexValuesInBraces)));
-			Arguments.Add(TEXT("ProxyPackage"), FText::FromString(GetOutermost()->GetName()));
 
 			FMessageLog("MapCheck").Info()
-				->AddToken(FUObjectToken::Create(this))
-				->AddToken(FTextToken::Create(FText::Format(LOCTEXT("MapCheck_Message_LandscapeProxyObsoleteLayer","Layer '{LayerName}' ('{LayerGuid}') was removed from LandscapeProxy because it doesn't match any of the LandscapeActor Layers. Please resave '{ProxyPackage}'."), Arguments)))
+				->AddToken(FUObjectToken::Create(this, FText::FromString(GetActorNameOrLabel())))
+				->AddToken(FTextToken::Create(FText::Format(LOCTEXT("MapCheck_Message_LandscapeProxyObsoleteLayer","Layer '{LayerName}' ('{LayerGuid}') was removed from LandscapeProxy because it doesn't match any of the Target Layers. Saving will remove this data for good."), Arguments)))
+				->AddToken(FActionToken::Create(LOCTEXT("MapCheck_RemoveObsoleteLayers", "Save Modified Landscapes"), LOCTEXT("MapCheck_RemoveObsoleteLayers_Desc", "Saves the modified landscape proxy actors"),
+					FOnActionTokenExecuted::CreateStatic(&UE::Landscape::SaveModifiedLandscapes),
+					FCanExecuteActionToken::CreateStatic(&UE::Landscape::HasModifiedLandscapes),
+					/*bInSingleUse = */false))
 				->AddToken(FMapErrorToken::Create(FMapErrors::LandscapeComponentPostLoad_Warning));
 
 			DeleteLayer(LayerGuidAndName.Key);
diff --git a/Engine/Source/Runtime/Landscape/Private/LandscapeUtils.cpp b/Engine/Source/Runtime/Landscape/Private/LandscapeUtils.cpp
--- a/Engine/Source/Runtime/Landscape/Private/LandscapeUtils.cpp
+++ b/Engine/Source/Runtime/Landscape/Private/LandscapeUtils.cpp
@@ -7,18 +7,23 @@
 #include "LandscapeLayerInfoObject.h"
 #include "LandscapeProxy.h"
 #include "LandscapeEditTypes.h"
+#include "LandscapeSubsystem.h"
 #include "Algo/Transform.h"
 
 #if WITH_EDITOR
+#include "AssetRegistry/AssetRegistryModule.h"
+#include "Editor/EditorEngine.h"
 #include "EditorDirectories.h"
 #include "Engine/Texture2D.h"
 #include "LandscapeComponent.h"
 #include "ObjectTools.h"
-#include "AssetRegistry/AssetRegistryModule.h"
-#endif
+#endif // WITH_EDITOR
 
+#if WITH_EDITOR
 // Channel remapping
 extern const size_t ChannelOffsets[4];
+extern UNREALED_API UEditorEngine* GEditor;
+#endif // WITH_EDITOR
 
 namespace UE::Landscape
 {
@@ -284,6 +289,36 @@ ULandscapeLayerInfoObject* FLayerInfoFinder::Find(const FName& LayerName) const
 	}
 	return nullptr;
 }
-#endif //!WITH_EDITOR
+
+bool HasModifiedLandscapes()
+{
+	if (GEditor)
+	{
+		if (UWorld* World = GEditor->GetEditorWorldContext().World())
+		{
+			if (ULandscapeSubsystem* LandscapeSubsystem = World->GetSubsystem<ULandscapeSubsystem>())
+			{
+				return LandscapeSubsystem->HasModifiedLandscapes();
+			}
+		}
+	}
+
+	return false;
+}
+
+void SaveModifiedLandscapes()
+{
+	if (GEditor)
+	{
+		if (UWorld* World = GEditor->GetEditorWorldContext().World())
+		{
+			if (ULandscapeSubsystem* LandscapeSubsystem = World->GetSubsystem<ULandscapeSubsystem>())
+			{
+				LandscapeSubsystem->SaveModifiedLandscapes();
+			}
+		}
+	}
+}
+#endif // WITH_EDITOR
 
 } // end namespace UE::Landscape
\ No newline at end of file
diff --git a/Engine/Source/Runtime/Landscape/Public/LandscapeUtils.h b/Engine/Source/Runtime/Landscape/Public/LandscapeUtils.h
--- a/Engine/Source/Runtime/Landscape/Public/LandscapeUtils.h
+++ b/Engine/Source/Runtime/Landscape/Public/LandscapeUtils.h
@@ -113,6 +113,13 @@ struct FLayerInfoFinder
 	TArray<FAssetData> LayerInfoAssets;
 };
 
-#endif //!WITH_EDITOR
+
+/** Returns true if there are some landscapes in the editor world that have been automatically modified and are in need of being saved (see LandscapeDirtyingMode) */
+LANDSCAPE_API bool HasModifiedLandscapes();
+
+/** Dirties and saves the landscapes in the editor world that have been automatically modified and are in need of being saved (see LandscapeDirtyingMode) */
+LANDSCAPE_API void SaveModifiedLandscapes();
+
+#endif // WITH_EDITOR
 
 } // end namespace UE::Landscape
