Remove call to UpdateRenderingMethod from ULandscapeSubsystem::Tick().

This checked whether each landscape proxy should render with nanite or non-nanite code.  Previously disabled in runtime via cvar to avoid the performance cost.  Now fully removed.  Fix a few issues this causes.

Change type of cvar class managing landscape.RenderNanite, allowing it to be safely read from the game thread.  Fixes inconsistent behavior failing to change in/out of nanite rendering when the cvar is changed in the editor.

Fix changing out of nanite rendering when the visibility layer is cleared.  If no other weightmaps remained, then the UpdateLayersContent() code was failing to call UpdateRenderingMethod.

diff --git a/Engine/Source/Runtime/Landscape/Private/Landscape.cpp b/Engine/Source/Runtime/Landscape/Private/Landscape.cpp
--- a/Engine/Source/Runtime/Landscape/Private/Landscape.cpp
+++ b/Engine/Source/Runtime/Landscape/Private/Landscape.cpp
@@ -232,10 +232,8 @@ static FAutoConsoleVariableRef CVarLandscapeHeightmapCompressionMode(
 
 #endif // WITH_EDITOR
 
-int32 GRenderNaniteLandscape = 1;
-FAutoConsoleVariableRef CVarRenderNaniteLandscape(
-	TEXT("landscape.RenderNanite"),
-	GRenderNaniteLandscape,
+TAutoConsoleVariable<int32> CVarRenderNaniteLandscape(
+	TEXT("landscape.RenderNanite"), 1,
 	TEXT("Render Landscape using Nanite."),
 	ECVF_Scalability | ECVF_RenderThreadSafe
 );
@@ -6862,7 +6860,7 @@ void ALandscapeProxy::UpdateRenderingMethod()
 	}
 
 	bool bNaniteActive = false;
-	if ((GRenderNaniteLandscape != 0) && HasNaniteComponents())
+	if ((CVarRenderNaniteLandscape.GetValueOnGameThread() != 0) && HasNaniteComponents())
 	{
 		bNaniteActive = UseNanite(GShaderPlatformForFeatureLevel[GEngine->GetDefaultWorldFeatureLevel()]);
 #if WITH_EDITOR
diff --git a/Engine/Source/Runtime/Landscape/Private/LandscapeEditLayers.cpp b/Engine/Source/Runtime/Landscape/Private/LandscapeEditLayers.cpp
--- a/Engine/Source/Runtime/Landscape/Private/LandscapeEditLayers.cpp
+++ b/Engine/Source/Runtime/Landscape/Private/LandscapeEditLayers.cpp
@@ -8751,13 +8751,22 @@ void ALandscape::ReallocateLayersWeightmaps(FUpdateLayersContentContext& InUpdat
 	}
 	}
 
-	// Trim the components that don't need weightmaps anymore (e.g. all edit layers are made invisible : there were some components LandscapeComponentsWeightmapsToResolve but there aren't anymore now):
-	InUpdateLayersContentContext.LandscapeComponentsWeightmapsToResolve.RemoveAll([&LayerAllocsPerComponent](ULandscapeComponent* Component)
-	{
-		TArray<ULandscapeLayerInfoObject*>* ComponentLayerAlloc = LayerAllocsPerComponent.Find(Component);
-		check(ComponentLayerAlloc != nullptr);
-		return ComponentLayerAlloc->IsEmpty();
-	});
+	// Trim the components that don't need weightmaps anymore (e.g. all edit layers are made invisible : there were some
+	// components in LandscapeComponentsWeightmapsToResolve but there aren't anymore now).
+	// Remember which components are removed.  These might need InvalidateGeneratedComponentData to fix up nanite dependencies.
+	TArray<ULandscapeComponent*> ComponentsWithNoAllocs;
+	InUpdateLayersContentContext.LandscapeComponentsWeightmapsToResolve.RemoveAllSwap([&LayerAllocsPerComponent, &ComponentsWithNoAllocs](ULandscapeComponent* Component) -> bool
+		{
+			TArray<ULandscapeLayerInfoObject*>* ComponentLayerAlloc = LayerAllocsPerComponent.Find(Component);
+			check(ComponentLayerAlloc != nullptr);
+			bool bNoAllocs = ComponentLayerAlloc->IsEmpty();
+			if (bNoAllocs)
+			{
+				ComponentsWithNoAllocs.Add(Component);
+			}
+			return bNoAllocs;
+		},
+		EAllowShrinking::No);
 
 	// Determine if the Final layer need to add/remove some alloc
 	for (auto& ItPair : LayerAllocsPerComponent)
@@ -8856,6 +8865,7 @@ void ALandscape::ReallocateLayersWeightmaps(FUpdateLayersContentContext& InUpdat
 		return true;
 	});
 
+	InvalidateGeneratedComponentData(ComponentsWithNoAllocs, false);
 	ValidateProxyLayersWeightmapUsage();
 
 	InUpdateLayersContentContext.Refresh(FUpdateLayersContentContext::ERefreshFlags::RefreshWeightmapInfos | FUpdateLayersContentContext::ERefreshFlags::RefreshMapHelper);
diff --git a/Engine/Source/Runtime/Landscape/Private/LandscapeSubsystem.cpp b/Engine/Source/Runtime/Landscape/Private/LandscapeSubsystem.cpp
--- a/Engine/Source/Runtime/Landscape/Private/LandscapeSubsystem.cpp
+++ b/Engine/Source/Runtime/Landscape/Private/LandscapeSubsystem.cpp
@@ -70,12 +70,6 @@ static FAutoConsoleVariable CVarMaxAsyncNaniteProxiesPerSecond(
 	6.0f,
 	TEXT("Number of Async nanite proxies to dispatch per second"));
 
-static bool GUpdateProxyActorRenderMethodOnTickAtRuntime = false;
-static FAutoConsoleVariableRef CVarUpdateProxyActorRenderMethodUpdateOnTickAtRuntime(
-	TEXT("landscape.UpdateProxyActorRenderMethodOnTickAtRuntime"),
-	GUpdateProxyActorRenderMethodOnTickAtRuntime,
-	TEXT("Update landscape proxy's rendering method (nanite enabled) when ticked. Always enabled in editor."));
-
 extern int32 GGrassMapUseRuntimeGeneration;
 
 DECLARE_CYCLE_STAT(TEXT("LandscapeSubsystem Tick"), STAT_LandscapeSubsystemTick, STATGROUP_Landscape);
@@ -687,13 +681,6 @@ void ULandscapeSubsystem::Tick(float DeltaTime)
 			int32 InOutNumComponentsCreated = 0;
 			Proxy->UpdateGrass(*Cameras, InOutNumComponentsCreated);
 		}
-
-#if !WITH_EDITOR
-		if (GUpdateProxyActorRenderMethodOnTickAtRuntime)
-#endif // WITH_EDITOR
-		{
-			Proxy->UpdateRenderingMethod();
-		}
 	}
 
 	ActiveProxies.Reset();
