Fix crash when calling ULandscapeHeightfieldCollisionComponent::CreateCollisionObject in an editor level.

diff --git a/Engine/Source/Runtime/Landscape/Private/LandscapeCollision.cpp b/Engine/Source/Runtime/Landscape/Private/LandscapeCollision.cpp
--- a/Engine/Source/Runtime/Landscape/Private/LandscapeCollision.cpp
+++ b/Engine/Source/Runtime/Landscape/Private/LandscapeCollision.cpp
@@ -1048,6 +1048,15 @@ void ULandscapeHeightfieldCollisionComponent::CreateCollisionObject(
 	HeightfieldRef = GSharedHeightfieldRefs.Add(HeightfieldGuid, new FHeightfieldGeometryRef(HeightfieldGuid));
 	HeightfieldRef->HeightfieldGeometry = Chaos::FHeightFieldPtr(new Chaos::FHeightField(Heights, PhysicalMaterialIds, CollisionSizeVerts, CollisionSizeVerts, Chaos::FVec3(1)));
 	
+#if WITH_EDITOR
+	UWorld* World = GetWorld();
+	const bool bNeedsEditorHeightField = World && !World->IsGameWorld() && !GetOutermost()->bIsCookedForEditor;
+	if (bNeedsEditorHeightField)
+	{
+		HeightfieldRef->EditorHeightfieldGeometry = Chaos::FHeightFieldPtr(new Chaos::FHeightField(Heights, PhysicalMaterialIds, CollisionSizeVerts, CollisionSizeVerts, Chaos::FVec3(1)));
+	}
+#endif // WITH_EDITOR
+
 	if (bGenerateSimpleCollision)
 	{
 		HeightfieldRef->HeightfieldSimpleGeometry = Chaos::FHeightFieldPtr(new Chaos::FHeightField(SimpleHeights, SimplePhysicalMaterialIds, SimpleCollisionSizeVerts, SimpleCollisionSizeVerts, Chaos::FVec3(1)));
