Enable landscape heightmap normal edge fixup by default

diff --git a/Engine/Source/Runtime/Landscape/Private/LandscapeGroup.cpp b/Engine/Source/Runtime/Landscape/Private/LandscapeGroup.cpp
--- a/Engine/Source/Runtime/Landscape/Private/LandscapeGroup.cpp
+++ b/Engine/Source/Runtime/Landscape/Private/LandscapeGroup.cpp
@@ -23,8 +23,7 @@ TMap<TObjectPtr<UTexture2D>, TObjectPtr<ULandscapeComponent>> FLandscapeGroup::H
 
 namespace UE::Landscape
 {
-	// TODO [chris.tchou] Re-enable once we figure out why xbox edge patching is breaking the texture
-	int32 GInstallEdgeFixup = 0;
+	int32 GInstallEdgeFixup = 1;
 	static FAutoConsoleVariableRef CVarInstallEdgeFixup(
 		TEXT("landscape.InstallEdgeFixup"),
 		GInstallEdgeFixup,
@@ -35,15 +34,13 @@ namespace UE::Landscape
 		return GInstallEdgeFixup != 0;
 	}
 
-	// TODO [chris.tchou] Re-enable once we figure out why xbox edge patching is breaking the texture
-	int32 GPatchEdges = 0;
+	int32 GPatchEdges = 1;
 	static FAutoConsoleVariableRef CVarPatchEdges(
 		TEXT("landscape.PatchEdges"),
 		GPatchEdges,
 		TEXT("Controls whether landscape heightmap texture edges are patched to match neighboring heightmaps.  Default enabled (1)."));
 
-	// TODO [chris.tchou] Re-enable once we figure out why xbox edge patching is breaking the texture
-	int32 GPatchStreamingMipEdges = 0;
+	int32 GPatchStreamingMipEdges = 1;
 	static FAutoConsoleVariableRef CVarPatchStreamingMipEdges(
 		TEXT("landscape.PatchStreamingMipEdges"),
 		GPatchStreamingMipEdges,
