[Backout] - CL36902937
Original CL Desc
-----------------------------------------------------------------
Enable landscape heightmap normal edge fixup by default

diff --git a/Engine/Source/Runtime/Landscape/Private/LandscapeGroup.cpp b/Engine/Source/Runtime/Landscape/Private/LandscapeGroup.cpp
--- a/Engine/Source/Runtime/Landscape/Private/LandscapeGroup.cpp
+++ b/Engine/Source/Runtime/Landscape/Private/LandscapeGroup.cpp
@@ -23,7 +23,8 @@ TMap<TObjectPtr<UTexture2D>, TObjectPtr<ULandscapeComponent>> FLandscapeGroup::H
 
 namespace UE::Landscape
 {
-	int32 GInstallEdgeFixup = 1;
+	// TODO [chris.tchou] Re-enable once we figure out why xbox edge patching is breaking the texture
+	int32 GInstallEdgeFixup = 0;
 	static FAutoConsoleVariableRef CVarInstallEdgeFixup(
 		TEXT("landscape.InstallEdgeFixup"),
 		GInstallEdgeFixup,
@@ -34,13 +35,15 @@ namespace UE::Landscape
 		return GInstallEdgeFixup != 0;
 	}
 
-	int32 GPatchEdges = 1;
+	// TODO [chris.tchou] Re-enable once we figure out why xbox edge patching is breaking the texture
+	int32 GPatchEdges = 0;
 	static FAutoConsoleVariableRef CVarPatchEdges(
 		TEXT("landscape.PatchEdges"),
 		GPatchEdges,
 		TEXT("Controls whether landscape heightmap texture edges are patched to match neighboring heightmaps.  Default enabled (1)."));
 
-	int32 GPatchStreamingMipEdges = 1;
+	// TODO [chris.tchou] Re-enable once we figure out why xbox edge patching is breaking the texture
+	int32 GPatchStreamingMipEdges = 0;
 	static FAutoConsoleVariableRef CVarPatchStreamingMipEdges(
 		TEXT("landscape.PatchStreamingMipEdges"),
 		GPatchStreamingMipEdges,
