Fix issue with PIE and Landscape Groups not unregistering properly, causing crashes on the second time entering PIE, when edge patching is active.

The issue was that HeightmapTextureToActiveComponent was not updated properly everywhere - particularly when re-enabling a disabled component, it was not updating HeightmapTextureToActiveComponent to point at the re-enabled, newly active component.

The check() that was getting hit was catching the issue, but only on the second time running PIE.

Here I make sure it is updated every time we call SetActiveComponent(), which should be the only way to change the active component now.

diff --git a/Engine/Source/Runtime/Landscape/Private/LandscapeEdgeFixup.cpp b/Engine/Source/Runtime/Landscape/Private/LandscapeEdgeFixup.cpp
--- a/Engine/Source/Runtime/Landscape/Private/LandscapeEdgeFixup.cpp
+++ b/Engine/Source/Runtime/Landscape/Private/LandscapeEdgeFixup.cpp
@@ -777,11 +777,14 @@ void ULandscapeHeightmapTextureEdgeFixup::SetActiveComponent(ULandscapeComponent
 	ActiveComponent = InComponent;
 	ActiveGroup = InGroup;
 
-	InGroup->Map(this, InComponent);
+	if (InComponent != nullptr)
+	{
+		InGroup->Map(this, InComponent);
 
-	FIXUP_DEBUG_LOG_DETAIL(TEXT("    Map %p (%d,%d) <Group %d Key %d> edgefixup:%p MAPPED:%d TOTAL:%d "),
-		InComponent, GroupCoord.X, GroupCoord.Y, InComponent->GetLandscapeProxy()->LODGroupKey, InGroup->LandscapeGroupKey, this,
-		InGroup->XYToEdgeFixupMap.Num(), InGroup->AllRegisteredFixups.Num());
+		FIXUP_DEBUG_LOG_DETAIL(TEXT("    Map %p (%d,%d) <Group %d Key %d> edgefixup:%p MAPPED:%d TOTAL:%d "),
+			InComponent, GroupCoord.X, GroupCoord.Y, InComponent->GetLandscapeProxy()->LODGroupKey, InGroup->LandscapeGroupKey, this,
+			InGroup->XYToEdgeFixupMap.Num(), InGroup->AllRegisteredFixups.Num());
+	}
 }
 
 void ULandscapeHeightmapTextureEdgeFixup::RequestEdgeTexturePatchingForNeighbors(UE::Landscape::ENeighborFlags NeighborsNeedingPatching)
diff --git a/Engine/Source/Runtime/Landscape/Private/LandscapeGroup.cpp b/Engine/Source/Runtime/Landscape/Private/LandscapeGroup.cpp
--- a/Engine/Source/Runtime/Landscape/Private/LandscapeGroup.cpp
+++ b/Engine/Source/Runtime/Landscape/Private/LandscapeGroup.cpp
@@ -130,9 +130,7 @@ void FLandscapeGroup::RegisterComponent(ULandscapeComponent* Component)
 			FLandscapeGroup* ActiveGroup = EdgeFixup->ActiveGroup;
 			check(ActiveGroup != nullptr);
 
-			GROUP_DEBUG_LOG_DETAIL(TEXT("Group %p RegisterComponent %p -- Override Existing EdgeFixup"), this, Component);
-
-			HeightmapTextureToActiveComponent[Heightmap] = Component;
+			GROUP_DEBUG_LOG_DETAIL(TEXT("Group %p RegisterComponent %p -- Use Existing EdgeFixup (Active Component %p, new component disabled: %d)"), this, Component, ActiveComponent, bIsDisabled);
 		}
 		else
 		{
@@ -157,8 +155,6 @@ void FLandscapeGroup::RegisterComponent(ULandscapeComponent* Component)
 			// there should be no active component yet
 			check(EdgeFixup->ActiveGroup == nullptr);
 			check(EdgeFixup->ActiveComponent == nullptr);
-
-			HeightmapTextureToActiveComponent.Add(Heightmap, Component);
 		}
 	}
 
@@ -177,6 +173,7 @@ void FLandscapeGroup::RegisterComponent(ULandscapeComponent* Component)
 	{
 		// set active will also map the fixup within the group
 		EdgeFixup->SetActiveComponent(Component, this);
+		HeightmapTextureToActiveComponent.FindOrAdd(Heightmap) = Component;
 
 		GROUP_DEBUG_LOG_DETAIL(TEXT("  -- EdgeFixup Set Active At Coord (%d, %d)"), EdgeFixup->GroupCoord.X, EdgeFixup->GroupCoord.Y);
 
@@ -251,10 +248,11 @@ void FLandscapeGroup::UnregisterComponent(ULandscapeComponent* Component)
 					FLandscapeGroup* DisabledGroup = DisabledComponent->RegisteredLandscapeGroup;
 					check(DisabledGroup != nullptr);
 
-					GROUP_DEBUG_LOG_DETAIL(TEXT("Group %p UnregisterComponent %p %d (%d,%d) -- Reactivate Other Disabled"), this, Component, EdgeFixup->bMapped, EdgeFixup->GroupCoord.X, EdgeFixup->GroupCoord.Y);
+					GROUP_DEBUG_LOG_DETAIL(TEXT("Group %p UnregisterComponent %p %d (%d,%d) -- Reactivate Disabled Component %p"), this, Component, EdgeFixup->bMapped, EdgeFixup->GroupCoord.X, EdgeFixup->GroupCoord.Y, DisabledComponent);
 
 					const bool bDisableCurrentActive = false; // just unmap it
 					EdgeFixup->SetActiveComponent(DisabledComponent, DisabledGroup, bDisableCurrentActive);
+					HeightmapTextureToActiveComponent[EdgeFixup->HeightmapTexture] = DisabledComponent;
 
 					bReactivatedDisabled = true;
 					break;
@@ -267,15 +265,8 @@ void FLandscapeGroup::UnregisterComponent(ULandscapeComponent* Component)
 			GROUP_DEBUG_LOG_DETAIL(TEXT("Group %p UnregisterComponent %p %d (%d,%d) -- Deactivating"), this, Component, EdgeFixup->bMapped, EdgeFixup->GroupCoord.X, EdgeFixup->GroupCoord.Y);
 
 			// no disabled available to replace the active component -- just set NO active component
-			Unmap(EdgeFixup);
-
-			GROUP_DEBUG_LOG_DETAIL(TEXT("  Unmap %p (%d,%d) <Group %d Key %d> edgefixup:%p MAPPED:%d TOTAL:%d"),
-				EdgeFixup->ActiveComponent, EdgeFixup->GroupCoord.X, EdgeFixup->GroupCoord.Y, EdgeFixup->ActiveComponent->GetLandscapeProxy()->LODGroupKey, EdgeFixup->ActiveGroup->LandscapeGroupKey, this,
-				EdgeFixup->ActiveGroup->XYToEdgeFixupMap.Num(), EdgeFixup->ActiveGroup->AllRegisteredFixups.Num());
-
-			EdgeFixup->ActiveComponent = nullptr;
-			EdgeFixup->ActiveGroup = nullptr;
-
+			const bool bDisableCurrentActive = false; // just unmap it
+			EdgeFixup->SetActiveComponent(nullptr, nullptr, bDisableCurrentActive);
 			HeightmapTextureToActiveComponent.Remove(EdgeFixup->HeightmapTexture);
 		}
 	}
@@ -591,7 +582,7 @@ void FLandscapeGroup::TickEdgeFixup(ULandscapeSubsystem* LandscapeSubsystem, boo
 			check(Landscape);
 			if (!Landscape->bGrassUpdateEnabled)
 			{
-				GROUP_DEBUG_LOG_DETAIL(TEXT("-- Snap Component %p (%d,%d) -------- PAUSED (editing)"), EdgeFixup->ActiveComponent, EdgeFixup->GroupCoord.X, EdgeFixup->GroupCoord.Y);
+				GROUP_DEBUG_LOG_DETAIL(TEXT("-- Snap Component %p (%d,%d) -------- PAUSED (editing)"), EdgeFixup->ActiveComponent.Get(), EdgeFixup->GroupCoord.X, EdgeFixup->GroupCoord.Y);
 				continue;
 			}
 #endif // WITH_EDITORONLY_DATA
@@ -600,7 +591,7 @@ void FLandscapeGroup::TickEdgeFixup(ULandscapeSubsystem* LandscapeSubsystem, boo
 			SnapshotsCaptured++;
 			SnapshotEdgesChanged += FGenericPlatformMath::CountBits((uint64)ChangedEdges);
 
-			GROUP_DEBUG_LOG_DETAIL(TEXT("-- Snap Component %p (%d,%d) -------- UPDATED, Changed edges: %s"), EdgeFixup->ActiveComponent, EdgeFixup->GroupCoord.X, EdgeFixup->GroupCoord.Y, *DirectionFlagsToString(ChangedEdges));
+			GROUP_DEBUG_LOG_DETAIL(TEXT("-- Snap Component %p (%d,%d) -------- UPDATED, Changed edges: %s"), EdgeFixup->ActiveComponent.Get(), EdgeFixup->GroupCoord.X, EdgeFixup->GroupCoord.Y, *DirectionFlagsToString(ChangedEdges));
 			if (ChangedEdges != EEdgeFlags::None)
 			{
 				// if any edges were changed -- queue it for texture edge patching
@@ -635,28 +626,28 @@ void FLandscapeGroup::TickEdgeFixup(ULandscapeSubsystem* LandscapeSubsystem, boo
 			check(Landscape);
 			if (!Landscape->bGrassUpdateEnabled)
 			{
-				GROUP_DEBUG_LOG_DETAIL(TEXT("-- Patch Component %p (%d,%d) -------- PAUSED (editing)"), EdgeFixup->ActiveComponent, EdgeFixup->GroupCoord.X, EdgeFixup->GroupCoord.Y);
+				GROUP_DEBUG_LOG_DETAIL(TEXT("-- Patch Component %p (%d,%d) -------- PAUSED (editing)"), EdgeFixup->ActiveComponent.Get(), EdgeFixup->GroupCoord.X, EdgeFixup->GroupCoord.Y);
 				continue;
 			}
 #endif // WITH_EDITORONLY_DATA
 
 			if (EdgeFixup->HeightmapTexture->HasPendingInitOrStreaming())
 			{
-				GROUP_DEBUG_LOG_DETAIL(TEXT("-- Patch Component %p (%d,%d) -------- PAUSED (pending)"), EdgeFixup->ActiveComponent, EdgeFixup->GroupCoord.X, EdgeFixup->GroupCoord.Y);
+				GROUP_DEBUG_LOG_DETAIL(TEXT("-- Patch Component %p (%d,%d) -------- PAUSED (pending)"), EdgeFixup->ActiveComponent.Get(), EdgeFixup->GroupCoord.X, EdgeFixup->GroupCoord.Y);
 				continue;
 			}
 
 #if WITH_EDITOR
 			if (EdgeFixup->IsTextureEdgePatchingPaused())
 			{
-				GROUP_DEBUG_LOG_DETAIL(TEXT("-- Patch Component %p (%d,%d) -------- PAUSED (readback)"), EdgeFixup->ActiveComponent, EdgeFixup->GroupCoord.X, EdgeFixup->GroupCoord.Y);
+				GROUP_DEBUG_LOG_DETAIL(TEXT("-- Patch Component %p (%d,%d) -------- PAUSED (readback)"), EdgeFixup->ActiveComponent.Get(), EdgeFixup->GroupCoord.X, EdgeFixup->GroupCoord.Y);
 				continue;
 			}
 
 			// wait til we are not using the default texture to apply edge patching
 			if (EdgeFixup->GetHeightmapTexture()->IsDefaultTexture())
 			{
-				GROUP_DEBUG_LOG_DETAIL(TEXT("-- Patch Component %p (%d,%d) -------- WAITING (texture load)"), EdgeFixup->ActiveComponent, EdgeFixup->GroupCoord.X, EdgeFixup->GroupCoord.Y);
+				GROUP_DEBUG_LOG_DETAIL(TEXT("-- Patch Component %p (%d,%d) -------- WAITING (texture load)"), EdgeFixup->ActiveComponent.Get(), EdgeFixup->GroupCoord.X, EdgeFixup->GroupCoord.Y);
 				continue;
 			}
 #endif // WITH_EDITOR
@@ -670,7 +661,7 @@ void FLandscapeGroup::TickEdgeFixup(ULandscapeSubsystem* LandscapeSubsystem, boo
 			EdgesPatched += PatchedEdgeCount;
 
 			It.RemoveCurrent();
-			GROUP_DEBUG_LOG_DETAIL(TEXT("-- Patch Component %p (%d,%d) -------- PATCHED %d edges"), EdgeFixup->ActiveComponent, EdgeFixup->GroupCoord.X, EdgeFixup->GroupCoord.Y, PatchedEdgeCount);
+			GROUP_DEBUG_LOG_DETAIL(TEXT("-- Patch Component %p (%d,%d) -------- PATCHED %d edges"), EdgeFixup->ActiveComponent.Get(), EdgeFixup->GroupCoord.X, EdgeFixup->GroupCoord.Y, PatchedEdgeCount);
 		}
 	}
 
