Made new experimental landscape APIs UE_INTERNAL (because there's no other way to mark an API for experimental usage ATM)

diff --git a/Engine/Source/Runtime/Landscape/Classes/Landscape.h b/Engine/Source/Runtime/Landscape/Classes/Landscape.h
--- a/Engine/Source/Runtime/Landscape/Classes/Landscape.h
+++ b/Engine/Source/Runtime/Landscape/Classes/Landscape.h
@@ -130,6 +130,7 @@ struct FLandscapeLayerBrush
 	void SetOwner(ALandscape* InOwner);
 
 	//~ Begin UE::Landscape::EditLayers::IEditLayerRendererProvider implementation
+	LANDSCAPE_EDIT_LAYERS_BATCHED_MERGE_EXPERIMENTAL
 	LANDSCAPE_API virtual TArray<UE::Landscape::EditLayers::FEditLayerRendererState> GetEditLayerRendererStates(const ULandscapeInfo* InLandscapeInfo, bool bInSkipBrush) override;
 	//~ End UE::Landscape::EditLayers::IEditLayerRendererProvider implementation
 #endif // WITH_EDITOR
@@ -172,6 +173,7 @@ struct FLandscapeLayer
 
 #if WITH_EDITOR
 	//~ Begin UE::Landscape::EditLayers::IEditLayerRendererProvider implementation
+	LANDSCAPE_EDIT_LAYERS_BATCHED_MERGE_EXPERIMENTAL
 	LANDSCAPE_API virtual TArray<UE::Landscape::EditLayers::FEditLayerRendererState> GetEditLayerRendererStates(const ULandscapeInfo* InLandscapeInfo, bool bInSkipBrush) override;
 	//~ End UE::Landscape::EditLayers::IEditLayerRendererProvider implementation
 #endif // WITH_EDITOR
diff --git a/Engine/Source/Runtime/Landscape/Classes/LandscapeEditLayer.h b/Engine/Source/Runtime/Landscape/Classes/LandscapeEditLayer.h
--- a/Engine/Source/Runtime/Landscape/Classes/LandscapeEditLayer.h
+++ b/Engine/Source/Runtime/Landscape/Classes/LandscapeEditLayer.h
@@ -3,6 +3,7 @@
 #pragma once
 
 #include "LandscapeEditLayerRenderer.h"
+#include "LandscapeEditTypes.h"
 
 #include "LandscapeEditLayer.generated.h"
 
@@ -196,10 +197,14 @@ public:
 
 #if WITH_EDITOR
 	//~ Begin ILandscapeEditLayerRenderer implementation
+	LANDSCAPE_EDIT_LAYERS_BATCHED_MERGE_EXPERIMENTAL
 	LANDSCAPE_API virtual void GetRendererStateInfo(const ULandscapeInfo* InLandscapeInfo,
 		UE::Landscape::EditLayers::FEditLayerTargetTypeState& OutSupportedTargetTypeState, UE::Landscape::EditLayers::FEditLayerTargetTypeState& OutEnabledTargetTypeState, TArray<TSet<FName>>& OutRenderGroups) const override;
+	LANDSCAPE_EDIT_LAYERS_BATCHED_MERGE_EXPERIMENTAL
 	LANDSCAPE_API virtual TArray<UE::Landscape::EditLayers::FEditLayerRenderItem> GetRenderItems(const ULandscapeInfo* InLandscapeInfo) const override;
+	LANDSCAPE_EDIT_LAYERS_BATCHED_MERGE_EXPERIMENTAL
 	LANDSCAPE_API virtual void RenderLayer(ILandscapeEditLayerRenderer::FRenderParams& InRenderParams) override;
+	LANDSCAPE_EDIT_LAYERS_BATCHED_MERGE_EXPERIMENTAL
 	LANDSCAPE_API virtual FString GetEditLayerRendererDebugName() const override;
 	//~ End ILandscapeEditLayerRenderer implementation
 #endif // WITH_EDITOR
diff --git a/Engine/Source/Runtime/Landscape/Private/LandscapeEditLayerRenderer.cpp b/Engine/Source/Runtime/Landscape/Private/LandscapeEditLayerRenderer.cpp
--- a/Engine/Source/Runtime/Landscape/Private/LandscapeEditLayerRenderer.cpp
+++ b/Engine/Source/Runtime/Landscape/Private/LandscapeEditLayerRenderer.cpp
@@ -954,6 +954,34 @@ TArray<FName> FMergeRenderContext::ConvertTargetLayerBitIndicesToNamesChecked(co
 	return Names;
 }
 
+bool FMergeRenderContext::IsValid() const
+{
+	return !RenderBatches.IsEmpty();
+}
+
+int32 FMergeRenderContext::GetTargetLayerIndexForName(const FName& InName) const
+{
+	return AllTargetLayerNames.Find(InName);
+}
+
+int32 FMergeRenderContext::GetTargetLayerIndexForNameChecked(const FName& InName) const
+{
+	int32 Index = AllTargetLayerNames.Find(InName); 
+	check(Index != INDEX_NONE); 
+	return Index;
+}
+
+FName FMergeRenderContext::GetTargetLayerNameForIndex(int32 InIndex) const
+{
+	return AllTargetLayerNames.IsValidIndex(InIndex) ? AllTargetLayerNames[InIndex] : NAME_None;
+}
+
+FName FMergeRenderContext::GetTargetLayerNameForIndexChecked(int32 InIndex) const
+{
+	check(AllTargetLayerNames.IsValidIndex(InIndex)); 
+	return AllTargetLayerNames[InIndex];
+}
+
 TArray<ULandscapeLayerInfoObject*> FMergeRenderContext::ConvertTargetLayerBitIndicesToLayerInfos(const TBitArray<>& InTargetLayerBitIndices) const
 {
 	const int32 NumTargetLayerInfos = AllTargetLayerNames.Num();
diff --git a/Engine/Source/Runtime/Landscape/Public/LandscapeBlueprintBrushBase.h b/Engine/Source/Runtime/Landscape/Public/LandscapeBlueprintBrushBase.h
--- a/Engine/Source/Runtime/Landscape/Public/LandscapeBlueprintBrushBase.h
+++ b/Engine/Source/Runtime/Landscape/Public/LandscapeBlueprintBrushBase.h
@@ -110,10 +110,14 @@ public:
 
 #if WITH_EDITOR
 	//~ Begin ILandscapeEditLayerRenderer implementation
+	LANDSCAPE_EDIT_LAYERS_BATCHED_MERGE_EXPERIMENTAL
 	LANDSCAPE_API virtual void GetRendererStateInfo(const ULandscapeInfo* InLandscapeInfo, 
 		UE::Landscape::EditLayers::FEditLayerTargetTypeState& OutSupportedTargetTypeState, UE::Landscape::EditLayers::FEditLayerTargetTypeState& OutEnabledTargetTypeState, TArray<TSet<FName>>& OutRenderGroups) const override;
+	LANDSCAPE_EDIT_LAYERS_BATCHED_MERGE_EXPERIMENTAL
 	LANDSCAPE_API virtual TArray<UE::Landscape::EditLayers::FEditLayerRenderItem> GetRenderItems(const ULandscapeInfo* InLandscapeInfo) const override;
+	LANDSCAPE_EDIT_LAYERS_BATCHED_MERGE_EXPERIMENTAL
 	LANDSCAPE_API virtual void RenderLayer(ILandscapeEditLayerRenderer::FRenderParams& InRenderParams) override;
+	LANDSCAPE_EDIT_LAYERS_BATCHED_MERGE_EXPERIMENTAL
 	LANDSCAPE_API virtual FString GetEditLayerRendererDebugName() const override;
 	//~ End ILandscapeEditLayerRenderer implementation
 
diff --git a/Engine/Source/Runtime/Landscape/Public/LandscapeEditLayerRenderer.h b/Engine/Source/Runtime/Landscape/Public/LandscapeEditLayerRenderer.h
--- a/Engine/Source/Runtime/Landscape/Public/LandscapeEditLayerRenderer.h
+++ b/Engine/Source/Runtime/Landscape/Public/LandscapeEditLayerRenderer.h
@@ -64,11 +64,14 @@ struct FWeightmapPaintLayerInfo
  *  whether the renderer's heightmaps and/or visibility and/or weightmaps are enabled (and if so, which weightmap is enabled exactly)  
  *  It is meant to be provided by the edit layer renderer's GetRendererStateInfo implementation.
  */
-struct FEditLayerTargetTypeState
+class FEditLayerTargetTypeState
 {
 	friend class ILandscapeEditLayerRenderer;
 
+public:
+	LANDSCAPE_EDIT_LAYERS_BATCHED_MERGE_EXPERIMENTAL
 	LANDSCAPE_API FEditLayerTargetTypeState() = default;
+	LANDSCAPE_EDIT_LAYERS_BATCHED_MERGE_EXPERIMENTAL
 	LANDSCAPE_API FEditLayerTargetTypeState(ELandscapeToolTargetTypeFlags InTargetTypeMask, const TArrayView<const FName>& InSupportedWeightmaps = TArrayView<const FName>());
 	
 	/**
@@ -76,33 +79,43 @@ struct FEditLayerTargetTypeState
 	 * @param InTargetType the requested target type (heightmap/weightmap/visibility)
 	 * @param InWeightmapLayerName (optional) is the requested weightmap, only relevant for the ELandscapeToolTargetType::Weightmap case
 	 */
+	LANDSCAPE_EDIT_LAYERS_BATCHED_MERGE_EXPERIMENTAL
 	LANDSCAPE_API bool IsActive(ELandscapeToolTargetType InTargetType, FName InWeightmapLayerName = NAME_None) const;
 
 	/** Returns the currently active weightmaps, if Weightmap is amongst the supported target types */
+	LANDSCAPE_EDIT_LAYERS_BATCHED_MERGE_EXPERIMENTAL
 	LANDSCAPE_API const TArray<FName>& GetActiveWeightmaps() const;
 	
 	/** Returns the target type mask (i.e. same as ELandscapeToolTargetType, but as bit flags) */
+	LANDSCAPE_EDIT_LAYERS_BATCHED_MERGE_EXPERIMENTAL
 	LANDSCAPE_API ELandscapeToolTargetTypeFlags GetTargetTypeMask() const { return TargetTypeMask; }
 
 	/** Sets the target type mask (i.e. same as ELandscapeToolTargetType, but as bit flags) */
+	LANDSCAPE_EDIT_LAYERS_BATCHED_MERGE_EXPERIMENTAL
 	LANDSCAPE_API void SetTargetTypeMask(ELandscapeToolTargetTypeFlags InTargetTypeMask);
 
 	/** Adds the target type in parameter to the mask of active target types */
+	LANDSCAPE_EDIT_LAYERS_BATCHED_MERGE_EXPERIMENTAL
 	LANDSCAPE_API void AddTargetType(ELandscapeToolTargetType InTargetType);
 
 	/** Appends the target type mask in parameter to the mask of active target types */
+	LANDSCAPE_EDIT_LAYERS_BATCHED_MERGE_EXPERIMENTAL
 	LANDSCAPE_API void AddTargetTypeMask(ELandscapeToolTargetTypeFlags InTargetTypeMask);
 
 	/** Removes a single target type from the mask of active target types */
+	LANDSCAPE_EDIT_LAYERS_BATCHED_MERGE_EXPERIMENTAL
 	LANDSCAPE_API void RemoveTargetType(ELandscapeToolTargetType InTargetType);
 
 	/** Removes the target type mask in parameter from the mask of active target types */
+	LANDSCAPE_EDIT_LAYERS_BATCHED_MERGE_EXPERIMENTAL
 	LANDSCAPE_API void RemoveTargetTypeMask(ELandscapeToolTargetTypeFlags InTargetTypeMask);
 
 	/** Adds a weightmap to the list of active weightmaps (make sure ELandscapeToolTargetType::Weightmap is amongst the supported target types) */
+	LANDSCAPE_EDIT_LAYERS_BATCHED_MERGE_EXPERIMENTAL
 	LANDSCAPE_API void AddWeightmap(FName InWeightmapLayerName);
 
 	/** Removes a weightmap from the list of active weightmaps */
+	LANDSCAPE_EDIT_LAYERS_BATCHED_MERGE_EXPERIMENTAL
 	LANDSCAPE_API void RemoveWeightmap(FName InWeightmapLayerName);
 
 	/**
@@ -135,12 +148,16 @@ private:
  *  A target type must be both supported and enabled on a given edit layer renderer in order for this renderer to render anything. 
  *  It also describes the render groups this renderer needs when rendering its weightmap (i.e. which weightmap needs to be rendered with which weightmaps : e.g. for weight blending)
  */
-struct FEditLayerRendererState
+class FEditLayerRendererState
 {
+public:
+	LANDSCAPE_EDIT_LAYERS_BATCHED_MERGE_EXPERIMENTAL
 	LANDSCAPE_API FEditLayerRendererState() = default;
+	LANDSCAPE_EDIT_LAYERS_BATCHED_MERGE_EXPERIMENTAL
 	LANDSCAPE_API explicit FEditLayerRendererState(ILandscapeEditLayerRenderer* InRenderer, const ULandscapeInfo* InLandscapeInfo);
 
 	/** Returns the edit layer renderer which this state relates to */
+	LANDSCAPE_EDIT_LAYERS_BATCHED_MERGE_EXPERIMENTAL
 	LANDSCAPE_API ILandscapeEditLayerRenderer* GetRenderer() const { return Renderer; }
 
 	/**
@@ -148,27 +165,35 @@ struct FEditLayerRendererState
 	 * @param InTargetType the requested target type (heightmap/weightmap/visibility)
 	 * @param InWeightmapLayerName (optional) is the requested weightmap, only relevant for the ELandscapeToolTargetType::Weightmap case
 	 */
+	LANDSCAPE_EDIT_LAYERS_BATCHED_MERGE_EXPERIMENTAL
 	LANDSCAPE_API bool IsTargetSupported(ELandscapeToolTargetType InTargetType, FName InWeightmapLayerName = NAME_None) const;
 
 	/** Returns a mask of all target types supported by this renderer */
+	LANDSCAPE_EDIT_LAYERS_BATCHED_MERGE_EXPERIMENTAL
 	LANDSCAPE_API const FEditLayerTargetTypeState& GetSupportedTargetTypeState() const { return SupportedTargetTypeState; }
 
 	/** Returns a mask of all target types supported by this renderer */
+	LANDSCAPE_EDIT_LAYERS_BATCHED_MERGE_EXPERIMENTAL
 	LANDSCAPE_API const FEditLayerTargetTypeState& GetEnabledTargetTypeState() const { return EnabledTargetTypeState; }
 
 	/** Returns a list of all weightmaps supported by this renderer (only relevant for ELandscapeToolTargetType::Weightmap (and ELandscapeToolTargetType::Visibility))*/
+	LANDSCAPE_EDIT_LAYERS_BATCHED_MERGE_EXPERIMENTAL
 	LANDSCAPE_API const TArray<FName>& GetSupportedTargetWeightmaps() const;
 
 	/** Mutates the EnabledTargetTypeState by adding the target type in parameter to the mask of active target types */
+	LANDSCAPE_EDIT_LAYERS_BATCHED_MERGE_EXPERIMENTAL
 	LANDSCAPE_API void EnableTargetType(ELandscapeToolTargetType InTargetType);
 	
 	/** Mutates the EnabledTargetTypeState by adding the target type mask in parameter to the mask of active target types */
+	LANDSCAPE_EDIT_LAYERS_BATCHED_MERGE_EXPERIMENTAL
 	LANDSCAPE_API void EnableTargetTypeMask(ELandscapeToolTargetTypeFlags InTargetTypeMask);
 
 	/** Mutates the EnabledTargetTypeState by removing the target type in parameter from the mask of active target types */
+	LANDSCAPE_EDIT_LAYERS_BATCHED_MERGE_EXPERIMENTAL
 	LANDSCAPE_API void DisableTargetType(ELandscapeToolTargetType InTargetType);
 
 	/** Mutates the EnabledTargetTypeState by removing the target type mask in parameter from the mask of active target types */
+	LANDSCAPE_EDIT_LAYERS_BATCHED_MERGE_EXPERIMENTAL
 	LANDSCAPE_API void DisableTargetTypeMask(ELandscapeToolTargetTypeFlags InTargetTypeMask);
 
 	/**
@@ -176,19 +201,24 @@ struct FEditLayerRendererState
 	 * @param InTargetType the requested target type (heightmap/weightmap/visibility)
 	 * @param InWeightmapLayerName (optional) is the requested weightmap, only relevant for the ELandscapeToolTargetType::Weightmap case
 	 */
+	LANDSCAPE_EDIT_LAYERS_BATCHED_MERGE_EXPERIMENTAL
 	LANDSCAPE_API bool IsTargetEnabled(ELandscapeToolTargetType InTargetType, FName InWeightmapLayerName = NAME_None) const;
 
 	/** Mutates the EnabledTargetTypeState by adding the weightmap in parameter to the list of enabled weightmaps */
+	LANDSCAPE_EDIT_LAYERS_BATCHED_MERGE_EXPERIMENTAL
 	LANDSCAPE_API bool EnableTarget(ELandscapeToolTargetType InTargetType, FName InWeightmapLayerName = NAME_None);
 
 	/** Mutates the EnabledTargetTypeState by removing the weightmap in parameter from the list of enabled weightmaps */
+	LANDSCAPE_EDIT_LAYERS_BATCHED_MERGE_EXPERIMENTAL
 	LANDSCAPE_API void DisableTarget(FName InWeightmapLayerName);
 
 	/** Returns the list of all weightmaps currently enabled */
+	LANDSCAPE_EDIT_LAYERS_BATCHED_MERGE_EXPERIMENTAL
 	LANDSCAPE_API TArray<FName> GetEnabledTargetWeightmaps() const;
 
 	/** Returns the render groups associated with this renderer. A render group is a set of target layers (weightmaps) that depend on one another in order to produce the output target layers.
 	 This allows to implement "horizontal blending", where weightmaps can be blended with one another at each step of the landscape edit layers merge algorithm */
+	LANDSCAPE_EDIT_LAYERS_BATCHED_MERGE_EXPERIMENTAL
 	LANDSCAPE_API TArray<TSet<FName>> GetWeightmapRenderGroups() const { return RenderGroups; }
 
 private:
@@ -375,10 +405,7 @@ public:
 	FMergeRenderContext& operator=(const FMergeRenderContext& Other) = default;
 	FMergeRenderContext& operator=(FMergeRenderContext&& Other) = default;
 
-	bool IsValid() const
-	{ 
-		return !RenderBatches.IsEmpty(); 
-	}
+	bool IsValid() const;
 
 	/**
 	 * Cycle between the 3 render targets used for blending:
@@ -389,11 +416,16 @@ public:
 	 * 
 	 * @param InDesiredWriteAccess RHI state of the write RT after the operation, if specified
 	 */
+	LANDSCAPE_EDIT_LAYERS_BATCHED_MERGE_EXPERIMENTAL
 	LANDSCAPE_API void CycleBlendRenderTargets(ERHIAccess InDesiredWriteAccess = ERHIAccess::None);
 
+	LANDSCAPE_EDIT_LAYERS_BATCHED_MERGE_EXPERIMENTAL
 	LANDSCAPE_API ULandscapeScratchRenderTarget* GetBlendRenderTargetWrite() const;
+	LANDSCAPE_EDIT_LAYERS_BATCHED_MERGE_EXPERIMENTAL
 	LANDSCAPE_API ULandscapeScratchRenderTarget* GetBlendRenderTargetRead() const;
+	LANDSCAPE_EDIT_LAYERS_BATCHED_MERGE_EXPERIMENTAL
 	LANDSCAPE_API ULandscapeScratchRenderTarget* GetBlendRenderTargetReadPrevious() const;
+	LANDSCAPE_EDIT_LAYERS_BATCHED_MERGE_EXPERIMENTAL
 	LANDSCAPE_API ULandscapeScratchRenderTarget* GetValidityRenderTarget(const FName& InTargetLayerName) const;
 	
 	struct FOnRenderBatchTargetGroupDoneParams
@@ -427,19 +459,29 @@ public:
 	};
 	void Render(TFunction<void(const FOnRenderBatchTargetGroupDoneParams&)> OnRenderBatchTargetGroupDone);
 
+	LANDSCAPE_EDIT_LAYERS_BATCHED_MERGE_EXPERIMENTAL
 	inline ALandscape* GetLandscape() const { return Landscape; }
+	LANDSCAPE_EDIT_LAYERS_BATCHED_MERGE_EXPERIMENTAL
 	inline bool IsHeightmapMerge() const { return bIsHeightmapMerge; }
 
+	LANDSCAPE_EDIT_LAYERS_BATCHED_MERGE_EXPERIMENTAL
 	inline FIntPoint GetMaxNeededResolution() const { return MaxNeededResolution; }
+	LANDSCAPE_EDIT_LAYERS_BATCHED_MERGE_EXPERIMENTAL
 	inline const TArray<FMergeRenderBatch>& GetRenderBatches() const { return RenderBatches; }
+	LANDSCAPE_EDIT_LAYERS_BATCHED_MERGE_EXPERIMENTAL
 	LANDSCAPE_API const FMergeRenderBatch* GetCurrentRenderBatch() const;
 
+	LANDSCAPE_EDIT_LAYERS_BATCHED_MERGE_EXPERIMENTAL
 	LANDSCAPE_API FTransform ComputeVisualLogTransform(const FTransform& InTransform) const;
+	LANDSCAPE_EDIT_LAYERS_BATCHED_MERGE_EXPERIMENTAL
 	LANDSCAPE_API void IncrementVisualLogOffset();
+	LANDSCAPE_EDIT_LAYERS_BATCHED_MERGE_EXPERIMENTAL
 	LANDSCAPE_API void ResetVisualLogOffset();
 
 #if ENABLE_VISUAL_LOG
+	LANDSCAPE_EDIT_LAYERS_BATCHED_MERGE_EXPERIMENTAL
 	LANDSCAPE_API static int32 GetVisualLogAlpha();
+	LANDSCAPE_EDIT_LAYERS_BATCHED_MERGE_EXPERIMENTAL
 	LANDSCAPE_API bool IsVisualLogEnabled() const;
 #endif // ENABLE_VISUAL_LOG
 
@@ -450,15 +492,19 @@ public:
 	void RenderExpandedRenderTarget(const FMergeRenderBatch& InRenderBatch);
 
 	const TArray<FName>& GetAllTargetLayerNames() const { return AllTargetLayerNames; }
-	inline int32 GetTargetLayerIndexForName(const FName& InName) const { return AllTargetLayerNames.Find(InName); }
-	inline int32 GetTargetLayerIndexForNameChecked(const FName& InName) const { int32 Index = AllTargetLayerNames.Find(InName); check(Index); return Index; }
-	inline FName GetTargetLayerNameForIndex(int32 InIndex) const { return AllTargetLayerNames.IsValidIndex(InIndex) ? AllTargetLayerNames[InIndex] : NAME_None; }
-	inline FName GetTargetLayerNameForIndexChecked(int32 InIndex) const { check(AllTargetLayerNames.IsValidIndex(InIndex)); return AllTargetLayerNames[InIndex]; }
+	int32 GetTargetLayerIndexForName(const FName& InName) const;
+	int32 GetTargetLayerIndexForNameChecked(const FName& InName) const;
+	FName GetTargetLayerNameForIndex(int32 InIndex) const;
+	FName GetTargetLayerNameForIndexChecked(int32 InIndex) const;
 	TBitArray<> ConvertTargetLayerNamesToBitIndices(TConstArrayView<FName> InTargetLayerNames) const;
 	TBitArray<> ConvertTargetLayerNamesToBitIndicesChecked(TConstArrayView<FName> InTargetLayerNames) const;
+	LANDSCAPE_EDIT_LAYERS_BATCHED_MERGE_EXPERIMENTAL
 	LANDSCAPE_API TArray<FName> ConvertTargetLayerBitIndicesToNames(const TBitArray<>& InTargetLayerBitIndices) const;
+	LANDSCAPE_EDIT_LAYERS_BATCHED_MERGE_EXPERIMENTAL
 	LANDSCAPE_API TArray<FName> ConvertTargetLayerBitIndicesToNamesChecked(const TBitArray<>& InTargetLayerBitIndices) const;
+	LANDSCAPE_EDIT_LAYERS_BATCHED_MERGE_EXPERIMENTAL
 	LANDSCAPE_API TArray<ULandscapeLayerInfoObject*> ConvertTargetLayerBitIndicesToLayerInfos(const TBitArray<>& InTargetLayerBitIndices) const;
+	LANDSCAPE_EDIT_LAYERS_BATCHED_MERGE_EXPERIMENTAL
 	LANDSCAPE_API TArray<ULandscapeLayerInfoObject*> ConvertTargetLayerBitIndicesToLayerInfosChecked(const TBitArray<>& InTargetLayerBitIndices) const;
 
 	/**
@@ -469,10 +515,15 @@ public:
 	 *     return continueLoop ? true : false;
 	 * });
 	 */
+	LANDSCAPE_EDIT_LAYERS_BATCHED_MERGE_EXPERIMENTAL
 	LANDSCAPE_API void ForEachTargetLayer(const TBitArray<>& InTargetLayerBitIndices, TFunctionRef<bool(int32 /*InTargetLayerIndex*/, FName /*InTargetLayerName*/)> Fn) const;
+	LANDSCAPE_EDIT_LAYERS_BATCHED_MERGE_EXPERIMENTAL
 	LANDSCAPE_API void ForEachTargetLayerChecked(const TBitArray<>& InTargetLayerBitIndices, TFunctionRef<bool(int32 /*InTargetLayerIndex*/, FName /*InTargetLayerName*/)> Fn) const;
 
-	const TBitArray<>& GetFinalTargetLayerBitIndices() const { return FinalTargetLayerBitIndices; }
+	const TBitArray<>& GetFinalTargetLayerBitIndices() const 
+	{ 
+		return FinalTargetLayerBitIndices; 
+	}
 
 private:
 	/** Allocate all needed render targets for this merge */
@@ -609,6 +660,7 @@ private:
 	/** Area around the component that is needed in the EType::LocalComponent / EType::SpecificComponent case (in component coordinates (see ULandscapeComponent::GetComponentKey()), 
 	 e.g. use (-1, -1, 1, 1) for the component and its immediate neighbors all around) */
 	FIntRect LocalArea; 
+
 	/** World space object-oriented box in the EType::OOBBox case */
 	FOOBox2D OOBox2D;
 };
@@ -730,6 +782,7 @@ public:
 	 *
 	 * @return a list of renderer states (i.e. a ILandscapeEditLayerRenderer and its current state) to be processed in that order by the merge operation
 	 */
+	LANDSCAPE_EDIT_LAYERS_BATCHED_MERGE_EXPERIMENTAL
 	LANDSCAPE_API virtual TArray<FEditLayerRendererState> GetEditLayerRendererStates(const ULandscapeInfo* InLandscapeInfo, bool bInSkipBrush) PURE_VIRTUAL(IEditLayerRendererProvider::GetEditLayerRendererStates, return {}; );
 };
 #endif // WITH_EDITOR
@@ -822,6 +875,7 @@ public:
 	 *  Depending on the other renderer's needs, the final render groups might contain more layers than was requested by a given renderer. This only means that more layers will be processed together
 	 *  and if this renderer doesn't act on one of these layers, it will simply do nothing with it in its RenderLayer function
 	 */
+	LANDSCAPE_EDIT_LAYERS_BATCHED_MERGE_EXPERIMENTAL
 	LANDSCAPE_API virtual void GetRendererStateInfo(const ULandscapeInfo* InLandscapeInfo,
 		UE::Landscape::EditLayers::FEditLayerTargetTypeState& OutSupportedTargetTypeState, UE::Landscape::EditLayers::FEditLayerTargetTypeState& OutEnabledTargetTypeState, TArray<TSet<FName>>& OutRenderGroups) const
 		PURE_VIRTUAL(ILandscapeEditLayerRenderer::GetRendererStateInfo, );
@@ -829,6 +883,7 @@ public:
 	/**
 	* @return the a debug name for this renderer
 	*/
+	LANDSCAPE_EDIT_LAYERS_BATCHED_MERGE_EXPERIMENTAL
 	LANDSCAPE_API virtual FString GetEditLayerRendererDebugName() const 
 		PURE_VIRTUAL(ILandscapeEditLayerRenderer::GetEditLayerRendererDebugName, return TEXT(""); );
 
@@ -839,12 +894,14 @@ public:
 	 *
 	 * @return list of all render items that affect this renderer
 	*/
+	LANDSCAPE_EDIT_LAYERS_BATCHED_MERGE_EXPERIMENTAL
 	LANDSCAPE_API virtual TArray<UE::Landscape::EditLayers::FEditLayerRenderItem> GetRenderItems(const ULandscapeInfo* InLandscapeInfo) const 
 		PURE_VIRTUAL(ILandscapeEditLayerRenderer::GetRenderItems, return { }; );
 
 	/**
 	 * Indicates whether the renderer actually does anything in the render phase
 	*/
+	LANDSCAPE_EDIT_LAYERS_BATCHED_MERGE_EXPERIMENTAL
 	LANDSCAPE_API virtual bool CanRender() const { return true; }
 
 	/**
@@ -852,6 +909,7 @@ public:
 	 * It operates on a limited set of components (depending on the size of the render batches) and on a set of target layers (e.g. multiple weightmaps).
 	 * It guarantees access to merged result from preceding layers of each target layer
 	*/
+	LANDSCAPE_EDIT_LAYERS_BATCHED_MERGE_EXPERIMENTAL
 	LANDSCAPE_API virtual void RenderLayer(ILandscapeEditLayerRenderer::FRenderParams& InRenderParams) {}
 
 #endif // WITH_EDITOR
diff --git a/Engine/Source/Runtime/Landscape/Public/LandscapeEditResourcesSubsystem.h b/Engine/Source/Runtime/Landscape/Public/LandscapeEditResourcesSubsystem.h
--- a/Engine/Source/Runtime/Landscape/Public/LandscapeEditResourcesSubsystem.h
+++ b/Engine/Source/Runtime/Landscape/Public/LandscapeEditResourcesSubsystem.h
@@ -4,6 +4,7 @@
 
 #include "Subsystems/EngineSubsystem.h"
 #include "Engine/TextureRenderTarget2D.h"
+#include "LandscapeEditTypes.h"
 #include "RHIAccess.h"
 #include "RHITransition.h"
 
@@ -16,7 +17,7 @@ class UTextureRenderTarget2DArray;
 namespace UE::Landscape 
 {
 
-struct FScratchRenderTargetParams
+struct LANDSCAPE_EDIT_LAYERS_BATCHED_MERGE_EXPERIMENTAL FScratchRenderTargetParams
 {
 	FScratchRenderTargetParams() = default;
 	explicit FScratchRenderTargetParams(const FString& InDebugName, bool bInExactDimensions, bool bInUseUAV, bool bInTargetArraySlicesIndependently, const FIntPoint& InResolution, int32 InNumSlices, ETextureRenderTargetFormat InFormat, const FLinearColor& InClearColor, ERHIAccess InInitialState)
@@ -51,7 +52,7 @@ struct FScratchRenderTargetParams
 	ERHIAccess InitialState = ERHIAccess::None;
 };
 
-struct FScratchRenderTargetScope : public FNoncopyable
+struct LANDSCAPE_EDIT_LAYERS_BATCHED_MERGE_EXPERIMENTAL FScratchRenderTargetScope : public FNoncopyable
 {
 	FScratchRenderTargetScope() = delete;
 	FScratchRenderTargetScope(const FScratchRenderTargetParams& InParams);
@@ -81,27 +82,38 @@ public:
 	ULandscapeScratchRenderTarget();
 
 	/** @return a debug name while this render target is in use. This helps track things down, since the underlying resource can be recycled */
+	LANDSCAPE_EDIT_LAYERS_BATCHED_MERGE_EXPERIMENTAL
 	LANDSCAPE_API const FString& GetDebugName() const;
 	/** @return Getter for the internal UTextureRenderTarget */
 	UTextureRenderTarget* GetRenderTarget() const { return RenderTarget; }
 	/** @return Getter for the internal UTextureRenderTarget2D (only when CurrentRenderTargetParams.NumSlices == 0, asserts if CurrentRenderTargetParams.NumSlices > 0) */
+	LANDSCAPE_EDIT_LAYERS_BATCHED_MERGE_EXPERIMENTAL
 	LANDSCAPE_API UTextureRenderTarget2D* GetRenderTarget2D() const;
 	/** @return the internal UTextureRenderTarget2D (only when CurrentRenderTargetParams.NumSlices == 0, returns nullptr if CurrentRenderTargetParams.NumSlices > 0) */
+	LANDSCAPE_EDIT_LAYERS_BATCHED_MERGE_EXPERIMENTAL
 	LANDSCAPE_API UTextureRenderTarget2D* TryGetRenderTarget2D() const;
 	/**  @return the internal UTextureRenderTarget2DArray (only when CurrentRenderTargetParams.NumSlices > 0, asserts if CurrentRenderTargetParams.NumSlices == 0)  */
+	LANDSCAPE_EDIT_LAYERS_BATCHED_MERGE_EXPERIMENTAL
 	LANDSCAPE_API UTextureRenderTarget2DArray* GetRenderTarget2DArray() const;
 	/** @return the internal UTextureRenderTarget2DArray (only when CurrentRenderTargetParams.NumSlices > 0, returns nullptr if CurrentRenderTargetParams.NumSlices == 0) */
+	LANDSCAPE_EDIT_LAYERS_BATCHED_MERGE_EXPERIMENTAL
 	LANDSCAPE_API UTextureRenderTarget2DArray* TryGetRenderTarget2DArray() const;
 	/** @return the internal render target's true resolution (can be different than the effective resolution if CurrentRenderTargetParams.bExactDimensions is false) */
+	LANDSCAPE_EDIT_LAYERS_BATCHED_MERGE_EXPERIMENTAL
 	LANDSCAPE_API FIntPoint GetResolution() const;
 	/**@return the internal render target's effective resolution (can be different than the actual resolution if CurrentRenderTargetParams.bExactDimensions is false) while this render target is in use */
+	LANDSCAPE_EDIT_LAYERS_BATCHED_MERGE_EXPERIMENTAL
 	LANDSCAPE_API FIntPoint GetEffectiveResolution() const;
 	/** @return the internal render target's number of slices (can be different than the actual number of slices if CurrentRenderTargetParams.bExactDimensions is false) */
+	LANDSCAPE_EDIT_LAYERS_BATCHED_MERGE_EXPERIMENTAL
 	LANDSCAPE_API int32 GetNumSlices() const;
 	/** @return the internal render target's effective number of slices (can be different than the actual number of slices if CurrentRenderTargetParams.bExactDimensions is false) while this render target is in use */
+	LANDSCAPE_EDIT_LAYERS_BATCHED_MERGE_EXPERIMENTAL
 	LANDSCAPE_API int32 GetEffectiveNumSlices() const;
 	bool IsTexture2DArray() const { return (TryGetRenderTarget2DArray() != nullptr); }
+	LANDSCAPE_EDIT_LAYERS_BATCHED_MERGE_EXPERIMENTAL
 	LANDSCAPE_API ETextureRenderTargetFormat GetFormat() const;
+	LANDSCAPE_EDIT_LAYERS_BATCHED_MERGE_EXPERIMENTAL
 	LANDSCAPE_API FLinearColor GetClearColor() const;
 	ERHIAccess GetCurrentState() const { return CurrentState; }
 	const UE::Landscape::FScratchRenderTargetParams& GetCurrentRenderTargetParams() { return CurrentRenderTargetParams; }
@@ -169,12 +181,16 @@ public:
 	};
 
 	/** Copies the content of the texture in parameter to the scratch texture (assuming the input texture is in CopySrc state already). Transitions the scratch texture's RHIAccess */
+	LANDSCAPE_EDIT_LAYERS_BATCHED_MERGE_EXPERIMENTAL
 	LANDSCAPE_API void CopyFrom(const FCopyFromTextureParams& InCopyParams);
 	/** Copies the content of the scratch texture in parameter to the scratch texture. Transitions both scratch texture's RHIAccess */
+	LANDSCAPE_EDIT_LAYERS_BATCHED_MERGE_EXPERIMENTAL
 	LANDSCAPE_API void CopyFrom(const FCopyFromScratchRenderTargetParams& InCopyParams);
 	/** Perform a transition of the scratch texture's RHIAccess */
+	LANDSCAPE_EDIT_LAYERS_BATCHED_MERGE_EXPERIMENTAL
 	LANDSCAPE_API void TransitionTo(ERHIAccess InStateAfter);
 	/** Perform a clear of the scratch texture */
+	LANDSCAPE_EDIT_LAYERS_BATCHED_MERGE_EXPERIMENTAL
 	LANDSCAPE_API void Clear();
 	/** @return true if this scratch texture is compatible (and can therefore be used) with the requested render target params */
 	bool IsCompatibleWith(const UE::Landscape::FScratchRenderTargetParams& InParams) const;
@@ -238,7 +254,3 @@ private:
 	UPROPERTY(Transient, NonTransactional)
 	TArray<TObjectPtr<ULandscapeScratchRenderTarget>> ScratchRenderTargets;
 };
-
-#if UE_ENABLE_INCLUDE_ORDER_DEPRECATED_IN_5_2
-#include "CoreMinimal.h"
-#endif
diff --git a/Engine/Source/Runtime/Landscape/Public/LandscapeEditTypes.h b/Engine/Source/Runtime/Landscape/Public/LandscapeEditTypes.h
--- a/Engine/Source/Runtime/Landscape/Public/LandscapeEditTypes.h
+++ b/Engine/Source/Runtime/Landscape/Public/LandscapeEditTypes.h
@@ -7,6 +7,12 @@
 
 #include "LandscapeEditTypes.generated.h"
 
+// APIs marked with LANDSCAPE_EDIT_LAYERS_BATCHED_MERGE_EXPERIMENTAL are experimental and therefore purposefully marked as UE_INTERNAL : 
+//  These APIs are part of a large refactor of the landscape edit layers system (batched merge) that will eventually replace the previous global 
+//  and local merge code paths but is still optional and therefore considered experimental for the time being. Since the API will likely evolve 
+//  until then, use at your own risk.
+#define LANDSCAPE_EDIT_LAYERS_BATCHED_MERGE_EXPERIMENTAL UE_INTERNAL
+
 UENUM()
 enum class ELandscapeToolTargetType : uint8
 {
