Removed deprecated/unused landscape grass CVar

diff --git a/Engine/Source/Runtime/Landscape/Private/LandscapeGrass.cpp b/Engine/Source/Runtime/Landscape/Private/LandscapeGrass.cpp
--- a/Engine/Source/Runtime/Landscape/Private/LandscapeGrass.cpp
+++ b/Engine/Source/Runtime/Landscape/Private/LandscapeGrass.cpp
@@ -167,12 +167,6 @@ static FAutoConsoleVariableRef CVarDisableGPUCull(
 	GDisableGPUCull,
 	TEXT("For debugging. Set this to zero to see where the grass is generated. Useful for tweaking the guard bands."));
 
-static int32 GPrerenderGrassmaps = 1;
-static FAutoConsoleVariableRef CVarPrerenderGrassmaps(
-	TEXT("grass.PrerenderGrassmaps"),
-	GPrerenderGrassmaps,
-	TEXT("1: Pre-render grass maps for all components in the editor; 0: Generate grass maps on demand while moving through the editor"));
-
 static int32 GDisableDynamicShadows = 0;
 static FAutoConsoleVariableRef CVarDisableDynamicShadows(
 	TEXT("grass.DisableDynamicShadows"),
