Stop unnecessary render passes in landscape grassmap runtime generation when bBakeMaterialPositionOffsetIntoCollision is enabled.

The height mip data generated by these passes is only ever used in editor, and is always discarded / stripped in non-editor builds.

This change makes it not bother to render at runtime, saving a little bit of memory, bandwidth and GPU time.

Technically either file change would fix the issue, but why not both!

diff --git a/Engine/Source/Runtime/Landscape/Private/LandscapeGrassMapsBuilder.cpp b/Engine/Source/Runtime/Landscape/Private/LandscapeGrassMapsBuilder.cpp
--- a/Engine/Source/Runtime/Landscape/Private/LandscapeGrassMapsBuilder.cpp
+++ b/Engine/Source/Runtime/Landscape/Private/LandscapeGrassMapsBuilder.cpp
@@ -1401,6 +1401,8 @@ void FLandscapeGrassMapsBuilder::KickOffRenderAndReadback(FComponentState& State
 
 	ULandscapeComponent* Component = State.Component;
 
+#if WITH_EDITORONLY_DATA
+	// the height mips data is discarded unless we are in editor, don't bother to request it otherwise
 	const bool bBakeMaterialPositionOffsetIntoCollision = (Component->GetLandscapeProxy() && Component->GetLandscapeProxy()->bBakeMaterialPositionOffsetIntoCollision);
 	if (bBakeMaterialPositionOffsetIntoCollision)
 	{
@@ -1413,6 +1415,7 @@ void FLandscapeGrassMapsBuilder::KickOffRenderAndReadback(FComponentState& State
 			HeightMips.Add(Component->SimpleCollisionMipLevel);
 		}
 	}
+#endif // WITH_EDITORONLY_DATA
 
 	check(State.ActiveRender == nullptr);
 	constexpr bool bInNeedsGrassmap = true;
diff --git a/Engine/Source/Runtime/Landscape/Private/LandscapeGrassWeightExporter.h b/Engine/Source/Runtime/Landscape/Private/LandscapeGrassWeightExporter.h
--- a/Engine/Source/Runtime/Landscape/Private/LandscapeGrassWeightExporter.h
+++ b/Engine/Source/Runtime/Landscape/Private/LandscapeGrassWeightExporter.h
@@ -61,13 +61,18 @@ public:
 			int32 NumGrassMaps = RequestedGrassTypes.Num();
 			if (bInNeedsHeightmap || NumGrassMaps > 0)
 			{
+				// 2 channels for the heightmap, and one channel for each grass map, packed into 4 channel render targets
 				NumPasses += FMath::DivideAndRoundUp(2 /* heightmap */ + NumGrassMaps, 4);
 			}
+
+#if WITH_EDITORONLY_DATA
+			// since we don't read HeightMips unless we are in editor, there's no reason to add passes for it unless we are in editor
 			if (InHeightMips.Num() > 0)
 			{
 				FirstHeightMipsPassIndex = NumPasses;
 				NumPasses += InHeightMips.Num();
 			}
+#endif // WITH_EDITORONLY_DATA
 		}
 	};
 
