Fix a crash in LandscapeInfo after undoing a deletion of an actor

diff --git a/Engine/Source/Runtime/Landscape/Private/LandscapeEdit.cpp b/Engine/Source/Runtime/Landscape/Private/LandscapeEdit.cpp
--- a/Engine/Source/Runtime/Landscape/Private/LandscapeEdit.cpp
+++ b/Engine/Source/Runtime/Landscape/Private/LandscapeEdit.cpp
@@ -973,7 +973,12 @@ void ULandscapeComponent::PostEditUndo()
 		{
 			// Here we create and register with the LandscapeInfo early (before PostRegisterAllComponents)
 			ULandscapeInfo* LandscapeInfo = ComponentProxy->GetLandscapeInfo();
-			if (!LandscapeInfo->IsRegistered(ComponentProxy))
+			if (!IsValid(LandscapeInfo))
+			{
+				LandscapeInfo = ComponentProxy->CreateLandscapeInfo(true);
+				check(LandscapeInfo->IsRegistered(ComponentProxy));
+			}
+			else if (!LandscapeInfo->IsRegistered(ComponentProxy))
 			{
 				LandscapeInfo->RegisterActor(ComponentProxy, true);
 			}
