Fix for data loss when enabling edit layers on landscape without having all landscape streaming proxies loaded

It was not realizing it needed to be upgrade because FixupSharedData stomps on the bHasLayerContent when copying layers over from the parent landscape actor.

diff --git a/Engine/Source/Runtime/Landscape/Classes/LandscapeProxy.h b/Engine/Source/Runtime/Landscape/Classes/LandscapeProxy.h
--- a/Engine/Source/Runtime/Landscape/Classes/LandscapeProxy.h
+++ b/Engine/Source/Runtime/Landscape/Classes/LandscapeProxy.h
@@ -959,6 +959,11 @@ PRAGMA_ENABLE_DEPRECATION_WARNINGS
 	UPROPERTY()
 	bool bHasLayersContent;
 
+	/** Because bHasLayersContent can be stomped by the FixupSharedData copying layers over from the parent LandscapeActor,
+	  * we preserve the initial state of bHasLayersContent until we get to registration and can apply upgrade steps
+	  **/
+	bool bHadLayersContentAtPostLoadTime = false;
+
 	/** Enable compressed heightmap texture storage. */
 	UPROPERTY(EditAnywhere, Category = Landscape, AdvancedDisplay, meta = (LandscapeInherited))
 	bool bUseCompressedHeightmapStorage = false;
diff --git a/Engine/Source/Runtime/Landscape/Private/Landscape.cpp b/Engine/Source/Runtime/Landscape/Private/Landscape.cpp
--- a/Engine/Source/Runtime/Landscape/Private/Landscape.cpp
+++ b/Engine/Source/Runtime/Landscape/Private/Landscape.cpp
@@ -4159,6 +4159,9 @@ void ALandscapeProxy::PostLoad()
 {
 	TRACE_CPUPROFILER_EVENT_SCOPE(ALandscapeProxy::PostLoad);
 
+	// save the load time state of layers content before doing anything else (specifically FixupSharedData can stomp it)
+	bHadLayersContentAtPostLoadTime = bHasLayersContent;
+
 	Super::PostLoad();
 
 	PostLoadFixupLandscapeGuidsIfInstanced();
@@ -5613,16 +5616,13 @@ void ULandscapeInfo::RegisterLandscapeActorWithProxyInternal(ALandscapeProxy* Pr
 
 		Proxy->UpdateCachedHasLayersContent(true);
 
-		// Cache the value at this point as CreateLandscapeInfo (-> RegisterActor) might create/destroy layers content if there was a mismatch between landscape & proxy
-		// Check the actual flag here not HasLayersContent() which could return true if the LandscapeActorRef is valid.
-		bool bHasLayersContentBefore = Proxy->bHasLayersContent;
-
+		bool bHasLayersContentBefore = Proxy->bHadLayersContentAtPostLoadTime;
 		check(Proxy->WeightmapFixupVersion == Proxy->CurrentVersion);
-
 		const bool bNeedOldDataMigration = !bHasLayersContentBefore && CanHaveLayersContent();
 		if (bNeedOldDataMigration && LandscapeActor->HasLayersContent())
 		{
 			LandscapeActor->CopyOldDataToDefaultLayer(Proxy);
+			Proxy->bHadLayersContentAtPostLoadTime = false;
 		}
 	}
 #endif // WITH_EDITOR
