Adding a cvar command to evict GrassMaps, for testing and debugging of runtime grassmap generation

diff --git a/Engine/Source/Runtime/Landscape/Private/LandscapeGrassMapsBuilder.cpp b/Engine/Source/Runtime/Landscape/Private/LandscapeGrassMapsBuilder.cpp
--- a/Engine/Source/Runtime/Landscape/Private/LandscapeGrassMapsBuilder.cpp
+++ b/Engine/Source/Runtime/Landscape/Private/LandscapeGrassMapsBuilder.cpp
@@ -42,6 +42,12 @@ FAutoConsoleVariableRef CVarGrassMapUseRuntimeGeneration(
 	GGrassMapUseRuntimeGeneration,
 	TEXT("Enable runtime grass map generation to save disk space and runtime memory.  When enabled the grass density maps are not serialized and are built on the fly at runtime."));
 
+int32 GGrassMapEvictAll = 0;
+FAutoConsoleVariableRef CVarGrassMapEvictAll(
+	TEXT("grass.GrassMap.EvictAll"),
+	GGrassMapEvictAll,
+	TEXT("Remove all grass maps. In editor, or if runtime generation is enabled, the grass maps will repopulate automatically."));
+
 int32 GGrassMapUseAsyncFetch = 0;
 FAutoConsoleVariableRef CVarGrassMapUseAsyncFetch(
 	TEXT("grass.GrassMap.UseAsyncFetch"),
@@ -840,12 +846,17 @@ void FLandscapeGrassMapsBuilder::AmortizedUpdateGrassMaps(
 	}
 
 #if WITH_EDITOR
-	const bool bCancelAndEvictAllImmediately = false;
+	// In editor we want to build all of the grass maps, and generally keep them around so we don't have to rebuild them
+	// because we will want to serialize all of them to disk on save.
+	const bool bCancelAndEvictAllImmediately = (GGrassMapEvictAll != 0);
 	const bool bEvictWhenBeyondEvictionRange = false;
 #else
-	const bool bCancelAndEvictAllImmediately = !GGrassEnable;
-	const bool bEvictWhenBeyondEvictionRange = !!GGrassMapUseRuntimeGeneration;
+	// At runtime, if we are using runtime generation then we evict anything beyond the eviction range
+	// And we also evict everything if grass is disabled generally
+	const bool bCancelAndEvictAllImmediately = (GGrassMapEvictAll != 0) || (GGrassEnable == 0);
+	const bool bEvictWhenBeyondEvictionRange = (GGrassMapUseRuntimeGeneration != 0);
 #endif // WITH_EDITOR
+	GGrassMapEvictAll = 0;
 	UpdateTrackedComponents(Cameras, AmortizedMaxRendering, GGrassMapMaxDiscardChecksPerFrame, bCancelAndEvictAllImmediately, bEvictWhenBeyondEvictionRange);
 
 	// no point in looking to start new grass map generation if nothing is pending, if grass is disabled or there are no cameras
