Removed unused forward declaration

#rnx

diff --git a/Engine/Source/Runtime/Landscape/Public/LandscapeRender.h b/Engine/Source/Runtime/Landscape/Public/LandscapeRender.h
--- a/Engine/Source/Runtime/Landscape/Public/LandscapeRender.h
+++ b/Engine/Source/Runtime/Landscape/Public/LandscapeRender.h
@@ -36,7 +36,6 @@ LandscapeRender.h: New terrain rendering
 #define LANDSCAPE_MAX_SUBSECTION_NUM 2
 
 class FLandscapeComponentSceneProxy;
-enum class ERuntimeVirtualTextureMaterialType : uint8;
 enum EShaderPlatform : uint16;
 
 namespace UE::Renderer::Private { class IShadowInvalidatingInstances; }
