Fixed assert in debug when using batched merge

diff --git a/Engine/Source/Runtime/Landscape/Private/LandscapeEditLayerRenderer.cpp b/Engine/Source/Runtime/Landscape/Private/LandscapeEditLayerRenderer.cpp
--- a/Engine/Source/Runtime/Landscape/Private/LandscapeEditLayerRenderer.cpp
+++ b/Engine/Source/Runtime/Landscape/Private/LandscapeEditLayerRenderer.cpp
@@ -732,6 +732,7 @@ void FMergeRenderContext::RenderExpandedRenderTarget(const FMergeRenderBatch& In
 		FSceneViewFamilyContext ViewFamily(FSceneViewFamily::ConstructionValues(nullptr, SceneInterface, FEngineShowFlags(ESFIM_Game)).SetTime(FGameTime::GetTimeSinceAppStart()));
 		FSceneViewInitOptions ViewInitOptions;
 		ViewInitOptions.ViewFamily = &ViewFamily;
+		ViewInitOptions.SetViewRectangle(FIntRect(0, 0, 1, 1)); // Use a dummy rect to avoid a check(slow)
 		GetRendererModule().CreateAndInitSingleView(InRHICmdList, &ViewFamily, &ViewInitOptions);
 		const FSceneView* View = ViewFamily.Views[0];
 
diff --git a/Engine/Source/Runtime/Landscape/Private/LandscapeEditLayers.cpp b/Engine/Source/Runtime/Landscape/Private/LandscapeEditLayers.cpp
--- a/Engine/Source/Runtime/Landscape/Private/LandscapeEditLayers.cpp
+++ b/Engine/Source/Runtime/Landscape/Private/LandscapeEditLayers.cpp
@@ -5710,6 +5710,7 @@ void ULandscapeEditLayerPersistent::RenderLayer(ILandscapeEditLayerRenderer::FRe
 			FSceneViewFamilyContext ViewFamily(FSceneViewFamily::ConstructionValues(nullptr, SceneInterface, FEngineShowFlags(ESFIM_Game)).SetTime(FGameTime::GetTimeSinceAppStart()));
 			FSceneViewInitOptions ViewInitOptions;
 			ViewInitOptions.ViewFamily = &ViewFamily;
+			ViewInitOptions.SetViewRectangle(FIntRect(0, 0, 1, 1)); // Use a dummy rect to avoid a check(slow)
 			GetRendererModule().CreateAndInitSingleView(InRHICmdList, &ViewFamily, &ViewInitOptions);
 			const FSceneView* View = ViewFamily.Views[0];
 
