Fix landscape heightmap import in expand mode when the target region extends outside the landscape.

diff --git a/Engine/Source/Runtime/Landscape/Public/LandscapeConfigHelper.h b/Engine/Source/Runtime/Landscape/Public/LandscapeConfigHelper.h
--- a/Engine/Source/Runtime/Landscape/Public/LandscapeConfigHelper.h
+++ b/Engine/Source/Runtime/Landscape/Public/LandscapeConfigHelper.h
@@ -98,36 +98,43 @@ public:
 		// Regions are in ComponentQuads and we want Vertices (+1)
 		const int32 SrcWidth = InSrcRegion.Width() + 1;
 		const int32 SrcHeight = InSrcRegion.Height() + 1;
-		const int32 DestWidth = InDestRegion.Width() + 1;
-		const int32 DestHeight = InDestRegion.Height() + 1;
+		const int32 DstWidth = InDestRegion.Width() + 1;
+		const int32 DstHeight = InDestRegion.Height() + 1;
 		const int32 OffsetX = bOffset ? InDestRegion.Min.X - InSrcRegion.Min.X : 0;
 		const int32 OffsetY = bOffset ? InDestRegion.Min.Y - InSrcRegion.Min.X : 0;
 		
-		OutData.Empty(DestWidth * DestHeight);
-		OutData.AddUninitialized(DestWidth * DestHeight);
+		OutData.Empty(DstWidth * DstHeight);
+		OutData.AddUninitialized(DstWidth * DstHeight);
 
-		for (int32 Y = 0; Y < DestHeight; ++Y)
+		for (int32 DstY = 0; DstY < DstHeight; ++DstY)
 		{
-			const int32 OldY = FMath::Clamp<int32>(Y + OffsetY, 0, SrcHeight - 1);
-						
+			const int32 SrcY = FMath::Clamp<int32>(DstY + OffsetY, 0, SrcHeight - 1);
+
 			// Pad anything to the left
-			const T PadLeft = InData[OldY * SrcWidth];
-			for (int32 X = 0; X < -OffsetX; ++X)
+			const T PadLeft = InData[SrcY * SrcWidth];
+			int32 EndPadLeft = FMath::Min<int32>(-OffsetX, DstWidth);
+			for (int32 DstX = 0; DstX < EndPadLeft; ++DstX)
 			{
-				OutData[Y * DestWidth + X] = PadLeft;
+				OutData[DstY * DstWidth + DstX] = PadLeft;
 			}
 
 			{
-				const int32 X = FMath::Max(0, -OffsetX);
-				const int32 OldX = FMath::Clamp<int32>(X + OffsetX, 0, SrcWidth - 1);
-				const int32 CopySize = FMath::Min<int32>(SrcWidth - OldX, DestWidth) * sizeof(T);
-				FMemory::Memcpy(&OutData[Y * DestWidth + X], &InData[OldY * SrcWidth + OldX], CopySize);
+				const int32 DstX = FMath::Max(0, -OffsetX);
+				const int32 SrcX = FMath::Clamp<int32>(DstX + OffsetX, 0, SrcWidth - 1);
+				// Limit to DstWidth-DstX to avoid writing past the end of the destination scanline.  This will also make CopySize negative if the target
+				// starts past the end of the scanline.
+				const int32 CopySize = FMath::Min<int32>(SrcWidth - SrcX, DstWidth - DstX) * sizeof(T);
+				if (CopySize > 0)
+				{
+					FMemory::Memcpy(&OutData[DstY * DstWidth + DstX], &InData[SrcY * SrcWidth + SrcX], CopySize);
+				}
 			}
 
-			const T PadRight = InData[OldY * SrcWidth + SrcWidth - 1];
-			for (int32 X = -OffsetX + SrcWidth; X < DestWidth; ++X)
+			const T PadRight = InData[SrcY * SrcWidth + SrcWidth - 1];
+			int32 StartPadRight = FMath::Max<int32>(-OffsetX + SrcWidth, 0);
+			for (int32 DstX = StartPadRight; DstX < DstWidth; ++DstX)
 			{
-				OutData[Y * DestWidth + X] = PadRight;
+				OutData[DstY * DstWidth + DstX] = PadRight;
 			}
 		}
 	}
