Fix for crash in landscape storage provider compression when the heightmap is not square

diff --git a/Engine/Source/Runtime/Landscape/Private/LandscapeTextureStorageProvider.cpp b/Engine/Source/Runtime/Landscape/Private/LandscapeTextureStorageProvider.cpp
--- a/Engine/Source/Runtime/Landscape/Private/LandscapeTextureStorageProvider.cpp
+++ b/Engine/Source/Runtime/Landscape/Private/LandscapeTextureStorageProvider.cpp
@@ -513,8 +513,8 @@ ULandscapeTextureStorageProviderFactory* ULandscapeTextureStorageProviderFactory
 			CompressMipToBulkData(MipIndex, MipWidth, MipHeight, MipData.GetData(), MipData.Num(), Mip->BulkData);
 		}
 
-		MipWidth >>= 1;
-		MipHeight >>= 1;
+		MipWidth = FMath::Max(MipWidth / 2, 1);
+		MipHeight = FMath::Max(MipHeight / 2, 1);
 	}
 
 	return Factory;
