Fix landscape heightmap export tile filenames when regions are not loaded.
Hide UI control for export mode (loaded-only vs all) when export is selected, where it had no effect.

diff --git a/Engine/Source/Editor/LandscapeEditor/Private/LandscapeEditorDetailCustomization_ImportExport.cpp b/Engine/Source/Editor/LandscapeEditor/Private/LandscapeEditorDetailCustomization_ImportExport.cpp
--- a/Engine/Source/Editor/LandscapeEditor/Private/LandscapeEditorDetailCustomization_ImportExport.cpp
+++ b/Engine/Source/Editor/LandscapeEditor/Private/LandscapeEditorDetailCustomization_ImportExport.cpp
@@ -40,6 +40,9 @@
 #include "LandscapeRegionUtils.h"
 #include "LandscapeEditorUtils.h"
 
+
+#include "LandscapeEditorPrivate.h"
+
 #define LOCTEXT_NAMESPACE "LandscapeEditor.ImportExport"
 
 TSharedRef<IDetailCustomization> FLandscapeEditorDetailCustomization_ImportExport::MakeInstance()
@@ -188,6 +191,7 @@ void FLandscapeEditorDetailCustomization_ImportExport::CustomizeDetails(IDetailL
 	TSharedPtr<SWidget> NameWidget;
 	TSharedPtr<SWidget> ValueWidget;
 	IDetailPropertyRow& ImportExportModeRow = ImportExportCategory.AddProperty(PropertyHandle_ImportExportMode);
+	ImportExportModeRow.Visibility(MakeAttributeLambda([]() { return FLandscapeEditorDetailCustomization_ImportExport::GetImportExportVisibility(true); }));
 	ImportExportModeRow.GetDefaultWidgets(NameWidget, ValueWidget);
 	ImportExportModeRow.CustomWidget()
 		.NameContent()
@@ -580,12 +584,10 @@ FReply FLandscapeEditorDetailCustomization_ImportExport::OnImportExportButtonCli
 
 	if (ULandscapeInfo* LandscapeInfo = LandscapeEdMode->CurrentToolTarget.LandscapeInfo.Get())
 	{
-		FIntRect LandscapeExtent;
-		LandscapeInfo->GetLandscapeExtent(LandscapeExtent);
-
+		FIntRect LandscapeExtent = LandscapeInfo->GetCompleteLandscapeExtent();
 		LandscapeExtent.Max.X += 1;
 		LandscapeExtent.Max.Y += 1;
-	
+
 		if (IsImporting())
 		{
 			FGuid CurrentLayerGuid = LandscapeEdMode->GetCurrentLayerGuid();
@@ -797,10 +799,8 @@ FReply FLandscapeEditorDetailCustomization_ImportExport::OnImportExportButtonCli
 				});
 			}
 
-			FIntRect Extent;
-			LandscapeInfo->GetLandscapeExtent(Extent);
-			// Set the import gizmo location to landscape minx/miny
-			LandscapeEdMode->UISettings->ImportLandscape_GizmoLocalPosition = Extent.Min;
+			// Reset the import gizmo location to landscape minx/miny
+			LandscapeEdMode->UISettings->ImportLandscape_GizmoLocalPosition = LandscapeExtent.Min;
 		}
 	}
 
diff --git a/Engine/Source/Runtime/Landscape/Classes/LandscapeInfo.h b/Engine/Source/Runtime/Landscape/Classes/LandscapeInfo.h
--- a/Engine/Source/Runtime/Landscape/Classes/LandscapeInfo.h
+++ b/Engine/Source/Runtime/Landscape/Classes/LandscapeInfo.h
@@ -211,6 +211,7 @@ public:
 	LANDSCAPE_API bool GetLandscapeExtent(ALandscapeProxy* Proxy, FIntRect& ProxyExtent) const;
 	LANDSCAPE_API bool GetLandscapeExtent(FIntRect& LandscapeExtent) const;
 	LANDSCAPE_API bool GetLandscapeExtent(int32& MinX, int32& MinY, int32& MaxX, int32& MaxY) const;
+	LANDSCAPE_API FIntRect GetCompleteLandscapeExtent() const;
 	LANDSCAPE_API bool GetLandscapeXYComponentBounds(FIntRect& OutXYComponentBounds) const;
 	LANDSCAPE_API void ForAllLandscapeComponents(TFunctionRef<void(ULandscapeComponent*)> Fn) const;
 	LANDSCAPE_API void ExportHeightmap(const FString& Filename);
diff --git a/Engine/Source/Runtime/Landscape/Private/LandscapeEdit.cpp b/Engine/Source/Runtime/Landscape/Private/LandscapeEdit.cpp
--- a/Engine/Source/Runtime/Landscape/Private/LandscapeEdit.cpp
+++ b/Engine/Source/Runtime/Landscape/Private/LandscapeEdit.cpp
@@ -4752,6 +4752,57 @@ bool ULandscapeInfo::GetLandscapeExtent(int32& MinX, int32& MinY, int32& MaxX, i
 	return (MinX != MAX_int32);
 }
 
+// Returns the full landscape extent regardless of which proxies are loaded.  Similar behavior to GetCompleteBounds.
+// Extent is measured in quads, matching GetLandscapeExtent.
+FIntRect ULandscapeInfo::GetCompleteLandscapeExtent() const
+{
+	FIntRect Extent(MAX_int32, MAX_int32, MIN_int32, MIN_int32);
+	ALandscape* Landscape = LandscapeActor.Get();
+
+	// In WP mode, examine each proxy
+	if (Landscape && Landscape->GetWorld() && Landscape->GetWorld()->GetWorldPartition())
+	{
+		FWorldPartitionHelpers::ForEachActorDescInstance<ALandscapeProxy>(Landscape->GetWorld()->GetWorldPartition(), [this, &Extent, Landscape](const FWorldPartitionActorDescInstance* ActorDescInstance)
+		{
+			FLandscapeActorDesc* LandscapeActorDesc = (FLandscapeActorDesc*)ActorDescInstance->GetActorDesc();
+			ALandscapeProxy* LandscapeProxy = Cast<ALandscapeProxy>(ActorDescInstance->GetActor());
+
+			if (LandscapeProxy && (LandscapeProxy->GetGridGuid() == LandscapeGuid))
+			{
+				// Prioritize extent of loaded proxies, as the ActorDesc is not updated after edits until the level is saved.
+				// Note that the parent ALandscape actor will have an extent in the ActorDesc (which does not match the true extent) but usually
+				// no components.  That will hit this case and return false from GetLandscapeExtent.  Ignore this phantom extent.
+				FIntRect ProxyExtent;
+				bool bRet = GetLandscapeExtent(LandscapeProxy, ProxyExtent);
+				if (bRet)
+				{
+					Extent.Union(ProxyExtent);
+				}
+			}
+			else if(LandscapeActorDesc->GridGuid == LandscapeGuid)
+			{
+				// If not loaded, use the values from the FLandscapeActorDesc
+				int32 GridSize = IntCastChecked<int32>(LandscapeActorDesc->GridSize);
+				int32 X = IntCastChecked<int32>(LandscapeActorDesc->GridIndexX);
+				int32 Y = IntCastChecked<int32>(LandscapeActorDesc->GridIndexY);
+				FIntRect DescExtent(X * GridSize, Y * GridSize,
+					(X + 1) * GridSize, (Y + 1) * GridSize);
+
+				Extent.Union(DescExtent);
+			}
+
+			return true;
+		});
+	}
+	else
+	{
+		// Otherwise use the extent straight from the components.
+		GetLandscapeExtent(Extent);
+	}
+
+	return Extent;
+}
+
 bool ULandscapeInfo::GetLandscapeXYComponentBounds(FIntRect& OutXYComponentBounds) const
 {
 	OutXYComponentBounds = XYComponentBounds;
