Declared batched merge APIs as experimental rather than internal

diff --git a/Engine/Source/Runtime/Landscape/Public/LandscapeEditTypes.h b/Engine/Source/Runtime/Landscape/Public/LandscapeEditTypes.h
--- a/Engine/Source/Runtime/Landscape/Public/LandscapeEditTypes.h
+++ b/Engine/Source/Runtime/Landscape/Public/LandscapeEditTypes.h
@@ -7,11 +7,11 @@
 
 #include "LandscapeEditTypes.generated.h"
 
-// APIs marked with LANDSCAPE_EDIT_LAYERS_BATCHED_MERGE_EXPERIMENTAL are experimental and therefore purposefully marked as UE_INTERNAL : 
+// APIs marked with LANDSCAPE_EDIT_LAYERS_BATCHED_MERGE_EXPERIMENTAL are experimental and therefore purposefully marked as UE_EXPERIMENTAL : 
 //  These APIs are part of a large refactor of the landscape edit layers system (batched merge) that will eventually replace the previous global 
 //  and local merge code paths but is still optional and therefore considered experimental for the time being. Since the API will likely evolve 
 //  until then, use at your own risk.
-#define LANDSCAPE_EDIT_LAYERS_BATCHED_MERGE_EXPERIMENTAL UE_INTERNAL
+#define LANDSCAPE_EDIT_LAYERS_BATCHED_MERGE_EXPERIMENTAL UE_EXPERIMENTAL(5.5, "Landscape edit layers batched merge is still being actively developed, this API may change.")
 
 UENUM()
 enum class ELandscapeToolTargetType : uint8
