Fix for potential race condition in landscape grass map readback from the GPU.

It is only valid to call IsReady() after the GPU fence has been written, and it's very difficult to know when that has happened (potential async shenanigans).

By always calling IsReady() from the render thread, we should be safe.

Side effect is that we now queue a render command for every call to CheckAndUpdate(), and may queue more commands than needed.
The overhead is minimal as we only call it once per tick in current usage.

diff --git a/Engine/Source/Runtime/Landscape/Private/LandscapeAsyncTextureReadback.cpp b/Engine/Source/Runtime/Landscape/Private/LandscapeAsyncTextureReadback.cpp
--- a/Engine/Source/Runtime/Landscape/Private/LandscapeAsyncTextureReadback.cpp
+++ b/Engine/Source/Runtime/Landscape/Private/LandscapeAsyncTextureReadback.cpp
@@ -19,9 +19,28 @@ void FLandscapeAsyncTextureReadback::StartReadback_RenderThread(FRDGBuilder& Gra
 	bStartedOnRenderThread = true;
 }
 
+void FLandscapeAsyncTextureReadback::CheckAndUpdate_RenderThread(const bool bInForceFinish)
+{
+	// StartReadback_RenderThread() must execute before CheckAndUpdate_RenderThread()
+	check(bStartedOnRenderThread);
+	if (bFinishedOnRenderThread)
+	{
+		// already finished, nothing to do!
+		return;
+	}
+	check(AsyncReadback.IsValid());
+	if (bInForceFinish || AsyncReadback->IsReady())
+	{
+		FinishReadback_RenderThread();
+	}
+}
+
 void FLandscapeAsyncTextureReadback::FinishReadback_RenderThread()
 {
 	check(bStartedOnRenderThread && AsyncReadback.IsValid());
+
+	// always run the lock first -- even if we are cancelling the operation
+	// this ensures that we wait until the readback is complete
 	int32 RowPitchInPixels = 0;
 	int32 BufferHeight = 0;
 	void* SrcData = AsyncReadback->Lock(RowPitchInPixels, &BufferHeight);	// this will block if the readback is not yet ready
@@ -29,7 +48,7 @@ void FLandscapeAsyncTextureReadback::FinishReadback_RenderThread()
 	check(RowPitchInPixels >= TextureWidth);
 	check(BufferHeight >= TextureHeight);
 
-	if (!bCancel)	// skip the copy work if we're cancelling
+	if (!bCancel)	// we can skip the copy work if we're cancelling
 	{
 		// copy into ReadbackResults
 		ReadbackResults.SetNumUninitialized(TextureWidth * TextureHeight);
@@ -79,38 +98,28 @@ void FLandscapeAsyncTextureReadback::FinishReadback_RenderThread()
 	AsyncReadback->Unlock();
 	AsyncReadback.Reset();
 
-	FPlatformMisc::MemoryBarrier();
 	bFinishedOnRenderThread = true;
 }
 
-bool FLandscapeAsyncTextureReadback::CheckAndUpdate(bool& bOutFinishCommandQueued, const bool bInForceFinish)
+bool FLandscapeAsyncTextureReadback::CheckAndUpdate(bool& bOutRenderCommandQueued, const bool bInForceFinish)
 {
-	// if we already queued the finish commands to render thread, then we're just waiting on it signaling readback complete	
-	if (bFinishQueuedFromGameThread)
+	// if we're already finished, then nothing to do!
+	if (bFinishedOnRenderThread)
 	{
-		return bFinishedOnRenderThread;
+		return true;
 	}
 	
-	// if we haven't started, or if the readback is not yet ready, then we have nothing to do but wait
-	if (!bStartedOnRenderThread || (!bInForceFinish && !AsyncReadback->IsReady()))
-	{
-		return false;
-	}
-
-	// the readback was started and it is ready, but we have not yet queued the finish command.
-	// queue it to make the data available to the game thread
+	// We can only safely check the status on the render thread, so queue a command to check it
 	FLandscapeAsyncTextureReadback* Readback = this;
-	ENQUEUE_RENDER_COMMAND(FLandscapeAsyncTextureReadback_FinishReadback)(
+	PendingRenderThreadCommands++;
+	ENQUEUE_RENDER_COMMAND(FLandscapeAsyncTextureReadback_CheckAndUpdate)(
 		[Readback, bInForceFinish](FRHICommandListImmediate& RHICmdList)
 		{
-			// sanity check the state 
-			check(Readback->bStartedOnRenderThread && !Readback->bFinishedOnRenderThread);
-			check(Readback->AsyncReadback.IsValid() && (bInForceFinish || Readback->AsyncReadback->IsReady()));
-			Readback->FinishReadback_RenderThread();
+			Readback->CheckAndUpdate_RenderThread(bInForceFinish);
+			Readback->PendingRenderThreadCommands--;
 		});
 
-	bFinishQueuedFromGameThread = true;
-	bOutFinishCommandQueued = true;
+	bOutRenderCommandQueued = true;
 
 	return false;
 }
@@ -120,26 +129,25 @@ void FLandscapeAsyncTextureReadback::CancelAndSelfDestruct()
 	// set the cancel flag, which will reduce work done by the finish command
 	bCancel = true;
 
-	// run a render thread finish command if it hasn't been queued  yet
-	bool bNeedsFinish = !bFinishQueuedFromGameThread;
-	bFinishQueuedFromGameThread = true;
-
 	FLandscapeAsyncTextureReadback* Readback = this;
-	ENQUEUE_RENDER_COMMAND(FCancelAndSelfDestructCommand)(
-		[Readback, bNeedsFinish](FRHICommandListImmediate& RHICmdList)
+	PendingRenderThreadCommands++;
+	ENQUEUE_RENDER_COMMAND(FLandscapeAsyncTextureReadback_CancelAndSelfDestruct)(
+		[Readback](FRHICommandListImmediate& RHICmdList)
 		{
 			check(Readback->bStartedOnRenderThread);
 			check(Readback->bCancel);
 
-			// if not yet finished, run the finish command
-			if (bNeedsFinish)
+			// check no commands were queued AFTER calling CancelAndSelfDestruct()
+			check(Readback->PendingRenderThreadCommands == 1);
+
+			// if not yet finished, force run the finish command (may stall til completion)
+			if (!Readback->bFinishedOnRenderThread)
 			{
 				Readback->FinishReadback_RenderThread();
 			}
-
 			check(Readback->bFinishedOnRenderThread);
 
-			// self destruct
+			// self destruct (from render thread)
 			delete Readback;
 		});
 }
@@ -151,9 +159,14 @@ void FLandscapeAsyncTextureReadback::QueueDeletionFromGameThread()
 	check(bFinishedOnRenderThread);
 
 	FLandscapeAsyncTextureReadback* Readback = this;
-	ENQUEUE_RENDER_COMMAND(FLandscapeAsyncTextureReadback_CheckAndUpdate)(
+	PendingRenderThreadCommands++;
+	ENQUEUE_RENDER_COMMAND(FLandscapeAsyncTextureReadback_QueueDeletion)(
 		[Readback](FRHICommandListImmediate& RHICmdList)
 		{
+			// check no commands were queued AFTER calling QueueDeletionFromGameThread()
+			check(Readback->PendingRenderThreadCommands == 1);
+
+			// self destruct (from render thread)
 			delete Readback;
 		});
 }
diff --git a/Engine/Source/Runtime/Landscape/Private/LandscapeAsyncTextureReadback.h b/Engine/Source/Runtime/Landscape/Private/LandscapeAsyncTextureReadback.h
--- a/Engine/Source/Runtime/Landscape/Private/LandscapeAsyncTextureReadback.h
+++ b/Engine/Source/Runtime/Landscape/Private/LandscapeAsyncTextureReadback.h
@@ -12,9 +12,9 @@
 class FLandscapeAsyncTextureReadback
 {
 private:
-	// game thread state
-	bool bFinishQueuedFromGameThread = false;					// game thread has sent the finish command to the render thread
-	bool bCancel = false;										// game thread wants to cancel readback, don't need to produce valid results
+	// game thread / render thread shared state
+	std::atomic_int PendingRenderThreadCommands = 0;			// number of render thread commands queued and not yet executed (for sanity checks)
+	std::atomic_bool bCancel = false;							// set by game thread to signal cancel readback: we don't need to produce valid results
 
 	// render thread state
 	std::atomic_bool bStartedOnRenderThread = false;			// render thread start command has sent the readback command to the GPU
@@ -40,19 +40,21 @@ public:
 	// use this to start an async readback operation from the render thread (on a render graph texture)
 	void StartReadback_RenderThread(FRDGBuilder& GraphBuilder, FRDGTextureRef RDGTexture);
 
+	// checks if the operation is ready to finish, and if it is, calls FinishReadback_RenderThread() to complete the readback.
+	// bInForceFinish will force the readback to finish by stalling the render thread until it is ready
+	void CheckAndUpdate_RenderThread(const bool bInForceFinish);
+
 	// use this to finish an async readback operation from the render thread
 	// Calling this when !IsComplete() will cause a stall until the GPU has completed the readback
 	void FinishReadback_RenderThread();
 
-	// TODO [chris.tchou] : add other StartReadback functions as needed
-
 	// Non-blocking call from the game thread to check readback status and start the finish command if needed.
 	// Returns true when the AsyncReadbackResults are available to the game thread, at which point you can call TakeResults() to access them.
-	// bOutFinishCommandQueued is set to true if the finish render command was queued in this call, othwerwise unchanged.
+	// bOutRenderCommandQueued is set to true if any render command was queued in this call, othwerwise unchanged.
 	// You must call this function occasionally or the readback will never complete, as the finish command is required.
 	// bInForceFinish will force the finish process to be queued on the render thread (potentially stalling render thread, but forcing it to finish the readback)
-	// bInForceFinish may not make the results immediately available, but will ensure they are available after the render thread executes the command.
-	bool CheckAndUpdate(bool& bOutFinishCommandQueued, const bool bInForceFinish);
+	// bInForceFinish DOES NOT generally make the results immediately available, but it will ensure they are available after the render thread executes the command.
+	bool CheckAndUpdate(bool& bOutRenderCommandQueued, const bool bInForceFinish);
 
 	// Call from game thread to terminate any readback in flight and queue deletion of this (FLandscapeAsyncTextureReadback) on the render thread
 	void CancelAndSelfDestruct();
@@ -78,7 +80,7 @@ public:
 	FString ToString()
 	{
 		FString Result;
-		Result.Appendf(TEXT("FLandscapeAsyncTextureReadback { RTStart: %d RTComplete: %d Queued: %d Cancel: %d AsyncReadback: %p }"), bStartedOnRenderThread.load(), bFinishedOnRenderThread.load(), bFinishQueuedFromGameThread, bCancel, AsyncReadback.Get());
+		Result.Appendf(TEXT("FLandscapeAsyncTextureReadback { RTStart: %d RTComplete: %d PendingCommands: %d Cancel: %d AsyncReadback: %p }"), bStartedOnRenderThread.load(), bFinishedOnRenderThread.load(), PendingRenderThreadCommands.load(), bCancel.load(), AsyncReadback.Get());
 		return Result;
 	}
 
diff --git a/Engine/Source/Runtime/Landscape/Private/LandscapeGrassMapsBuilder.cpp b/Engine/Source/Runtime/Landscape/Private/LandscapeGrassMapsBuilder.cpp
--- a/Engine/Source/Runtime/Landscape/Private/LandscapeGrassMapsBuilder.cpp
+++ b/Engine/Source/Runtime/Landscape/Private/LandscapeGrassMapsBuilder.cpp
@@ -1102,18 +1102,28 @@ bool FLandscapeGrassMapsBuilder::CancelAndEvict(FComponentState& State, bool bCa
 				if (bCancelImmediately)
 				{
 					// release active render from the state, and cancel and self destruct it
+					// (this may cause a stall in the render thread waiting for GPU completion)
 					FLandscapeGrassWeightExporter* Render = State.ActiveRender.Release();
 					Render->CancelAndSelfDestruct();
 				}
-				else if (!State.ActiveRender->CheckAndUpdateAsyncReadback(bNewRenderCommands, State.TickCount > 4))
-				{
-					// we can't cancel yet.. must wait for the readback to complete
-					return false;
-				}
 				else
 				{
-					// readback is complete, we can delete the active render
-					State.ActiveRender.Reset();
+					// generally we arrive here when the Component was unregistered,
+					// but we must wait for the readback to complete before exiting the rendering state.
+					// Because Components can be re-registered and pick up where they left off, we don't want to cancel the operation outright just yet.
+					// To avoid potential render thread stalls we let the operation continue asynchronously for up to four ticks,
+					// before forcing completion (this also gives us more time for re-register)
+					if (!State.ActiveRender->CheckAndUpdateAsyncReadback(bNewRenderCommands, State.TickCount > 4))
+					{
+						// we can't evict yet.. must wait for the readback to complete
+						return false;
+					}
+					else
+					{
+						// readback is complete, we can delete the active render and move to the pending state
+						// (under the hood this queues deletion on render thread)
+						State.ActiveRender.Reset();
+					}
 				}
 			}
 			check(RenderingCount > 0);
