- Allow editing a landscape part of an edit level instance (Can be enabled using the experimental editor settings flag 'Enable Level Instance Landscape Edit').
- Fixed multiple issues with Actor Partition system not using the correct level.
- Fixed Edit Level Instance post-effect not working properly on FLandscapeComponentSceneProxy.

diff --git a/Engine/Source/Editor/LandscapeEditor/Private/LandscapeEdMode.cpp b/Engine/Source/Editor/LandscapeEditor/Private/LandscapeEdMode.cpp
--- a/Engine/Source/Editor/LandscapeEditor/Private/LandscapeEdMode.cpp
+++ b/Engine/Source/Editor/LandscapeEditor/Private/LandscapeEdMode.cpp
@@ -2330,6 +2330,16 @@ int32 FEdModeLandscape::UpdateLandscapeList()
 {
 	LandscapeList.Empty();
 
+	// Use EditorModeManager to skip landscape actors that can't be edited
+	// One example is the ULevelInstanceEditorMode when in level instance restricted mode and 
+	// the actor is not part of the editing level instance.
+	FLevelEditorModule& LevelEditorModule = FModuleManager::GetModuleChecked<FLevelEditorModule>("LevelEditor");
+	TSharedPtr<ILevelEditor> FirstLevelEditor = LevelEditorModule.GetFirstLevelEditor();
+	auto IsEditingDisallowedByEdModes = [&FirstLevelEditor](AActor* Actor)
+	{
+		return FirstLevelEditor && Actor && FirstLevelEditor->GetEditorModeManager().IsEditingDisallowed(Actor);
+	};
+
 	if (!CurrentGizmoActor.IsValid())
 	{
 		ALandscapeGizmoActiveActor* GizmoActor = nullptr;
@@ -2351,7 +2361,9 @@ int32 FEdModeLandscape::UpdateLandscapeList()
 		for (auto It = LandscapeInfoMap.Map.CreateIterator(); It; ++It)
 		{
 			ULandscapeInfo* LandscapeInfo = It.Value();
-			if (IsValid(LandscapeInfo) && LandscapeInfo->SupportsLandscapeEditing())
+			if (IsValid(LandscapeInfo) && 
+				LandscapeInfo->SupportsLandscapeEditing() &&
+				!IsEditingDisallowedByEdModes(LandscapeInfo->LandscapeActor.Get()))
 			{
 				if (ALandscape* Landscape = LandscapeInfo->LandscapeActor.Get())
 				{
diff --git a/Engine/Source/Editor/LevelInstanceEditor/Private/LevelInstanceEditorMode.cpp b/Engine/Source/Editor/LevelInstanceEditor/Private/LevelInstanceEditorMode.cpp
--- a/Engine/Source/Editor/LevelInstanceEditor/Private/LevelInstanceEditorMode.cpp
+++ b/Engine/Source/Editor/LevelInstanceEditor/Private/LevelInstanceEditorMode.cpp
@@ -9,6 +9,7 @@
 #include "EditorModes.h"
 #include "Engine/World.h"
 #include "EngineUtils.h"
+#include "LevelInstance/LevelInstanceSettings.h"
 #include "LevelInstance/LevelInstanceSubsystem.h"
 #include "LevelInstance/LevelInstanceInterface.h"
 #include "LevelInstance/ILevelInstanceEditorModule.h"
@@ -267,7 +268,7 @@ void ULevelInstanceEditorMode::ModeTick(float DeltaTime)
 
 bool ULevelInstanceEditorMode::IsCompatibleWith(FEditorModeID OtherModeID) const
 {
-	return (OtherModeID != FBuiltinEditorModes::EM_Foliage) && (OtherModeID != FBuiltinEditorModes::EM_Landscape);
+	return (OtherModeID != FBuiltinEditorModes::EM_Foliage) && ((OtherModeID != FBuiltinEditorModes::EM_Landscape) || ULevelInstanceSettings::Get()->IsLevelInstanceEditCompatibleWithLandscapeEdit());
 }
 
 void ULevelInstanceEditorMode::BindCommands()
@@ -307,6 +308,11 @@ void ULevelInstanceEditorMode::BindCommands()
 		FIsActionChecked::CreateUObject(this, &ULevelInstanceEditorMode::IsContextRestrictionCommandEnabled));
 }
 
+bool ULevelInstanceEditorMode::IsEditingDisallowed(AActor* InActor) const
+{
+	return IsSelectionDisallowed(InActor, true);
+}
+
 bool ULevelInstanceEditorMode::IsSelectionDisallowed(AActor* InActor, bool bInSelection) const
 {
 	UWorld* World = InActor->GetWorld();
diff --git a/Engine/Source/Editor/LevelInstanceEditor/Private/LevelInstanceEditorMode.h b/Engine/Source/Editor/LevelInstanceEditor/Private/LevelInstanceEditorMode.h
--- a/Engine/Source/Editor/LevelInstanceEditor/Private/LevelInstanceEditorMode.h
+++ b/Engine/Source/Editor/LevelInstanceEditor/Private/LevelInstanceEditorMode.h
@@ -46,6 +46,7 @@ public:
 
 	/** Check to see if an actor can be selected in this mode - no side effects */
 	virtual bool IsSelectionDisallowed(AActor* InActor, bool bInSelection) const override;
+	virtual bool IsEditingDisallowed(AActor* InActor) const override;
 	/** Only accept saving on current asset */
 	virtual bool IsOperationSupportedForCurrentAsset(EAssetOperation InOperation) const { return InOperation == EAssetOperation::Save; }
 	// End UEdMode
diff --git a/Engine/Source/Editor/UnrealEd/Classes/Settings/EditorExperimentalSettings.h b/Engine/Source/Editor/UnrealEd/Classes/Settings/EditorExperimentalSettings.h
--- a/Engine/Source/Editor/UnrealEd/Classes/Settings/EditorExperimentalSettings.h
+++ b/Engine/Source/Editor/UnrealEd/Classes/Settings/EditorExperimentalSettings.h
@@ -149,6 +149,10 @@ public:
 	UPROPERTY(EditAnywhere, config, Category = LevelInstance)
 	bool bEnableLevelInstancePropertyOverrides;
 
+	/** Allows editing of landscapes inside an editing level instance. */
+	UPROPERTY(EditAnywhere, config, Category = LevelInstance)
+	bool bEnableLevelInstanceLandscapeEdit;
+
 	UPROPERTY(EditAnywhere, config, Category = Sequencer, meta = (DisplayName = "Enable Music Mode Editing", ToolTip = "Enables \"music mode\" editing features"))
 	bool bMusicModeEditingInSequencerEnabled;
 
diff --git a/Engine/Source/Editor/UnrealEd/Private/EditorModeManager.cpp b/Engine/Source/Editor/UnrealEd/Private/EditorModeManager.cpp
--- a/Engine/Source/Editor/UnrealEd/Private/EditorModeManager.cpp
+++ b/Engine/Source/Editor/UnrealEd/Private/EditorModeManager.cpp
@@ -272,6 +272,18 @@ void FEditorModeTools::SetSelectionHasSceneComponent(bool bHasSceneComponent)
 	bSelectionHasSceneComponent = bHasSceneComponent;
 }
 
+bool FEditorModeTools::IsEditingDisallowed(AActor* InActor) const
+{
+	for (const UEdMode* Mode : ActiveScriptableModes)
+	{
+		if (Mode->IsEditingDisallowed(InActor))
+		{
+			return true;
+		}
+	}
+	return false;
+}
+
 bool FEditorModeTools::IsSelectionAllowed(AActor* InActor, const bool bInSelected) const
 {
 	bool bSelectionAllowed = (ActiveScriptableModes.Num() == 0);
diff --git a/Engine/Source/Editor/UnrealEd/Public/EditorModeManager.h b/Engine/Source/Editor/UnrealEd/Public/EditorModeManager.h
--- a/Engine/Source/Editor/UnrealEd/Public/EditorModeManager.h
+++ b/Engine/Source/Editor/UnrealEd/Public/EditorModeManager.h
@@ -497,6 +497,8 @@ public:
 
 	UNREALED_API bool IsSelectionHandled(AActor* InActor, const bool bInSelected) const;
 
+	UNREALED_API bool IsEditingDisallowed(AActor* InActor) const;
+
 	UNREALED_API bool ProcessEditDuplicate();
 	UNREALED_API bool ProcessEditDelete();
 	UNREALED_API bool ProcessEditCut();
diff --git a/Engine/Source/Editor/UnrealEd/Public/Tools/UEdMode.h b/Engine/Source/Editor/UnrealEd/Public/Tools/UEdMode.h
--- a/Engine/Source/Editor/UnrealEd/Public/Tools/UEdMode.h
+++ b/Engine/Source/Editor/UnrealEd/Public/Tools/UEdMode.h
@@ -197,9 +197,12 @@ public:
 	/** Check to see if an actor can be selected in this mode - no side effects */
 	virtual bool IsSelectionAllowed(AActor* InActor, bool bInSelection) const { return true; }
 
-	/** Check to see if an actor selection is exclusively dissallowed by a mode -- no side effects */
+	/** Check to see if an actor selection is exclusively disallowed by a mode -- no side effects */
 	virtual bool IsSelectionDisallowed(AActor* InActor, bool bInSelection) const { return false; }
 
+	/** Check to see if an actor edit is exclusively disallowed by a mode */
+	virtual bool IsEditingDisallowed(AActor* InActor) const { return false; }
+
 	/** Returns the editor mode identifier. */
 	FEditorModeID GetID() const { return Info.ID; }
 
diff --git a/Engine/Source/Runtime/Engine/Private/ActorPartition/ActorPartitionSubsystem.cpp b/Engine/Source/Runtime/Engine/Private/ActorPartition/ActorPartitionSubsystem.cpp
--- a/Engine/Source/Runtime/Engine/Private/ActorPartition/ActorPartitionSubsystem.cpp
+++ b/Engine/Source/Runtime/Engine/Private/ActorPartition/ActorPartitionSubsystem.cpp
@@ -243,7 +243,8 @@ public:
 			return true;
 		};
 
-		UWorldPartition* WorldPartition = World->GetWorldPartition();
+		UWorldPartition* WorldPartition = InCellCoord.Level->GetWorldPartition();
+		check(WorldPartition);
 
 		FBox CellBounds = UActorPartitionSubsystem::FCellCoord::GetCellBounds(InCellCoord, InGridSize);
 		if (bInBoundsSearch)
@@ -259,10 +260,10 @@ public:
 		{
 			return nullptr;
 		}
-				
+
 		if (!FoundActor && bInCreate)
 		{
-			const FString ActorName = APartitionActor::GetActorName(World, InActorPartitionId, InGridSize, InCellCoord.X, InCellCoord.Y, InCellCoord.Z);
+			const FString ActorName = APartitionActor::GetActorName(InCellCoord.Level->GetTypedOuter<UWorld>(), InActorPartitionId, InGridSize, InCellCoord.X, InCellCoord.Y, InCellCoord.Z);
 
 			FActorSpawnParameters SpawnParams;
 			SpawnParams.OverrideLevel = InCellCoord.Level;
@@ -270,7 +271,7 @@ public:
 			SpawnParams.NameMode = FActorSpawnParameters::ESpawnActorNameMode::Required_Fatal;
 
 			// Handle the case where the actor already exists, but is in the undo stack (was deleted)
-			if (UObject* ExistingObject = StaticFindObject(nullptr, World->PersistentLevel, *SpawnParams.Name.ToString()))
+			if (UObject* ExistingObject = StaticFindObject(nullptr, InCellCoord.Level, *SpawnParams.Name.ToString()))
 			{
 				AActor* ExistingActor = CastChecked<AActor>(ExistingObject);
 				check(!IsValidChecked(ExistingActor));
diff --git a/Engine/Source/Runtime/Engine/Private/LevelInstance/LevelInstanceSettings.cpp b/Engine/Source/Runtime/Engine/Private/LevelInstance/LevelInstanceSettings.cpp
--- a/Engine/Source/Runtime/Engine/Private/LevelInstance/LevelInstanceSettings.cpp
+++ b/Engine/Source/Runtime/Engine/Private/LevelInstance/LevelInstanceSettings.cpp
@@ -17,6 +17,11 @@ bool ULevelInstanceSettings::IsPropertyOverrideEnabled() const
 	return GetDefault<UEditorExperimentalSettings>()->bEnableLevelInstancePropertyOverrides && !!PropertyOverridePolicy;
 }
 
+bool ULevelInstanceSettings::IsLevelInstanceEditCompatibleWithLandscapeEdit() const
+{
+	return GetDefault<UEditorExperimentalSettings>()->bEnableLevelInstanceLandscapeEdit;
+}
+
 void ULevelInstanceSettings::UpdatePropertyOverridePolicy()
 {
 	// Policy hasn't changed
diff --git a/Engine/Source/Runtime/Engine/Public/LevelInstance/LevelInstanceSettings.h b/Engine/Source/Runtime/Engine/Public/LevelInstance/LevelInstanceSettings.h
--- a/Engine/Source/Runtime/Engine/Public/LevelInstance/LevelInstanceSettings.h
+++ b/Engine/Source/Runtime/Engine/Public/LevelInstance/LevelInstanceSettings.h
@@ -19,6 +19,7 @@ public:
 
 #if WITH_EDITOR
 	bool IsPropertyOverrideEnabled() const;
+	bool IsLevelInstanceEditCompatibleWithLandscapeEdit() const;
 
 	UE_DEPRECATED(5.5, "Use SetIsLevelInstanceDisabled instead.")
 	void DisableLevelInstanceSupport() { bIsLevelInstanceDisabled = true; }
diff --git a/Engine/Source/Runtime/Landscape/Private/LandscapeConfigHelper.cpp b/Engine/Source/Runtime/Landscape/Private/LandscapeConfigHelper.cpp
--- a/Engine/Source/Runtime/Landscape/Private/LandscapeConfigHelper.cpp
+++ b/Engine/Source/Runtime/Landscape/Private/LandscapeConfigHelper.cpp
@@ -58,7 +58,7 @@ ALandscapeProxy* FLandscapeConfigHelper::FindOrAddLandscapeStreamingProxy(ULands
 	UActorPartitionSubsystem* ActorPartitionSubsystem = World->GetSubsystem<UActorPartitionSubsystem>();
 	const uint32 GridSize = InLandscapeInfo->LandscapeActor->GetGridSize();
 
-	UActorPartitionSubsystem::FCellCoord CellCoord = UActorPartitionSubsystem::FCellCoord::GetCellCoord(InSectionBase, World->PersistentLevel, GridSize);
+	UActorPartitionSubsystem::FCellCoord CellCoord = UActorPartitionSubsystem::FCellCoord::GetCellCoord(InSectionBase, InLandscapeInfo->LandscapeActor->GetLevel(), GridSize);
 	return FLandscapeConfigHelper::FindOrAddLandscapeStreamingProxy(ActorPartitionSubsystem, InLandscapeInfo, CellCoord);
 }
 
@@ -122,7 +122,7 @@ bool FLandscapeConfigHelper::ChangeGridSize(ULandscapeInfo* InLandscapeInfo, uin
 
 	TSet<ALandscapeProxy*> ProxiesToDelete;
 	
-	FActorPartitionGridHelper::ForEachIntersectingCell(ALandscapeStreamingProxy::StaticClass(), Extent, World->PersistentLevel, [ActorPartitionSubsystem, InLandscapeInfo, InNewGridSizeInComponents, &LandscapeComponents, &ProxiesToDelete](const UActorPartitionSubsystem::FCellCoord& CellCoord, const FIntRect& CellBounds)
+	FActorPartitionGridHelper::ForEachIntersectingCell(ALandscapeStreamingProxy::StaticClass(), Extent, InLandscapeInfo->LandscapeActor->GetLevel(), [ActorPartitionSubsystem, InLandscapeInfo, InNewGridSizeInComponents, &LandscapeComponents, &ProxiesToDelete](const UActorPartitionSubsystem::FCellCoord& CellCoord, const FIntRect& CellBounds)
 	{
 		TMap<ULandscapeComponent*, UMaterialInterface*> ComponentMaterials;
 		TMap<ULandscapeComponent*, UMaterialInterface*> ComponentHoleMaterials;
@@ -595,7 +595,7 @@ ULandscapeInfo* FLandscapeConfigHelper::ChangeConfiguration(ULandscapeInfo* InLa
 				continue;
 			}
 
-			UActorPartitionSubsystem::FCellCoord CellCoord = UActorPartitionSubsystem::FCellCoord::GetCellCoord(NewSectionBase, World->PersistentLevel, NewLandscape->GetGridSize());
+			UActorPartitionSubsystem::FCellCoord CellCoord = UActorPartitionSubsystem::FCellCoord::GetCellCoord(NewSectionBase, NewLandscape->GetLevel(), NewLandscape->GetGridSize());
 			ALandscapeProxy* LandscapeProxy = FLandscapeConfigHelper::FindOrAddLandscapeStreamingProxy(ActorPartitionSubsystem, NewLandscapeInfo, CellCoord);
 			check(LandscapeProxy);
 
diff --git a/Engine/Source/Runtime/Landscape/Private/LandscapeEdit.cpp b/Engine/Source/Runtime/Landscape/Private/LandscapeEdit.cpp
--- a/Engine/Source/Runtime/Landscape/Private/LandscapeEdit.cpp
+++ b/Engine/Source/Runtime/Landscape/Private/LandscapeEdit.cpp
@@ -2885,34 +2885,38 @@ bool ULandscapeInfo::HasUnloadedComponentsInRegion(int32 X1, int32 Y1, int32 X2,
 	if (LandscapeActor.IsValid())
 	{
 		UWorld* World = LandscapeActor->GetWorld();
-
-		int32 ComponentIndexX1, ComponentIndexY1, ComponentIndexX2, ComponentIndexY2;
-		ALandscape::CalcComponentIndicesOverlap(X1, Y1, X2, Y2, ComponentSizeQuads, ComponentIndexX1, ComponentIndexY1, ComponentIndexX2, ComponentIndexY2);
-
-		const UActorPartitionSubsystem::FCellCoord MinCoord = UActorPartitionSubsystem::FCellCoord::GetCellCoord(FIntPoint(ComponentIndexX1 * ComponentSizeQuads, ComponentIndexY1 * ComponentSizeQuads), World->PersistentLevel, LandscapeActor->GetGridSize());
-		const UActorPartitionSubsystem::FCellCoord MaxCoord = UActorPartitionSubsystem::FCellCoord::GetCellCoord(FIntPoint(ComponentIndexX2 * ComponentSizeQuads, ComponentIndexY2 * ComponentSizeQuads), World->PersistentLevel, LandscapeActor->GetGridSize());
+		ULevel* Level = LandscapeActor->GetLevel();
 
 		if (UWorldPartition* WorldPartition = World->GetWorldPartition())
 		{
-			FWorldPartitionHelpers::ForEachActorDescInstance<ALandscapeProxy>(WorldPartition, [this, World, &MinCoord, &MaxCoord, &bResult](const FWorldPartitionActorDescInstance* ActorDescInstance)
+			if (Level == World->PersistentLevel)
 			{
-				FLandscapeActorDesc* LandscapeActorDesc = (FLandscapeActorDesc*)ActorDescInstance->GetActorDesc();
+				int32 ComponentIndexX1, ComponentIndexY1, ComponentIndexX2, ComponentIndexY2;
+				ALandscape::CalcComponentIndicesOverlap(X1, Y1, X2, Y2, ComponentSizeQuads, ComponentIndexX1, ComponentIndexY1, ComponentIndexX2, ComponentIndexY2);
 
-				if (LandscapeActorDesc->GridGuid == LandscapeGuid)
+				const UActorPartitionSubsystem::FCellCoord MinCoord = UActorPartitionSubsystem::FCellCoord::GetCellCoord(FIntPoint(ComponentIndexX1 * ComponentSizeQuads, ComponentIndexY1 * ComponentSizeQuads), Level, LandscapeActor->GetGridSize());
+				const UActorPartitionSubsystem::FCellCoord MaxCoord = UActorPartitionSubsystem::FCellCoord::GetCellCoord(FIntPoint(ComponentIndexX2 * ComponentSizeQuads, ComponentIndexY2 * ComponentSizeQuads), Level, LandscapeActor->GetGridSize());
+
+				FWorldPartitionHelpers::ForEachActorDescInstance<ALandscapeProxy>(WorldPartition, [this, World, &MinCoord, &MaxCoord, &bResult](const FWorldPartitionActorDescInstance* ActorDescInstance)
 				{
-					const UActorPartitionSubsystem::FCellCoord ActorCoord(LandscapeActorDesc->GridIndexX, LandscapeActorDesc->GridIndexY, LandscapeActorDesc->GridIndexZ, World->PersistentLevel);
-					if (ActorCoord.X >= MinCoord.X && ActorCoord.Y >= MinCoord.Y && ActorCoord.X <= MaxCoord.X && ActorCoord.Y <= MaxCoord.Y)
+					FLandscapeActorDesc* LandscapeActorDesc = (FLandscapeActorDesc*)ActorDescInstance->GetActorDesc();
+
+					if (LandscapeActorDesc->GridGuid == LandscapeGuid)
 					{
-						if (!ActorDescInstance->IsLoaded())
+						const UActorPartitionSubsystem::FCellCoord ActorCoord(LandscapeActorDesc->GridIndexX, LandscapeActorDesc->GridIndexY, LandscapeActorDesc->GridIndexZ, World->PersistentLevel);
+						if (ActorCoord.X >= MinCoord.X && ActorCoord.Y >= MinCoord.Y && ActorCoord.X <= MaxCoord.X && ActorCoord.Y <= MaxCoord.Y)
 						{
-							bResult = true;
-							return false;
+							if (!ActorDescInstance->IsLoaded())
+							{
+								bResult = true;
+								return false;
+							}
 						}
 					}
-				}
 
-				return true;
-			});
+					return true;
+				});
+			}
 		}
 	}
 
diff --git a/Engine/Source/Runtime/Landscape/Private/LandscapeRender.cpp b/Engine/Source/Runtime/Landscape/Private/LandscapeRender.cpp
--- a/Engine/Source/Runtime/Landscape/Private/LandscapeRender.cpp
+++ b/Engine/Source/Runtime/Landscape/Private/LandscapeRender.cpp
@@ -2117,6 +2117,9 @@ FPrimitiveViewRelevance FLandscapeComponentSceneProxy::GetViewRelevance(const FS
 	else
 	{
 		Result.bStaticRelevance = true;
+#if WITH_EDITOR
+		Result.bEditorVisualizeLevelInstanceRelevance = IsEditingLevelInstanceChild();
+#endif
 	}
 
 	Result.bShadowRelevance = (GAllowLandscapeShadows > 0) && IsShadowCast(View) && View->Family->EngineShowFlags.Landscape;
