Fix for crash when weightmap layers have no name or don't exist.
Here we treat them as if the weightmap layer is not allocated, when getting the combination material.

diff --git a/Engine/Source/Runtime/Landscape/Private/LandscapeEdit.cpp b/Engine/Source/Runtime/Landscape/Private/LandscapeEdit.cpp
--- a/Engine/Source/Runtime/Landscape/Private/LandscapeEdit.cpp
+++ b/Engine/Source/Runtime/Landscape/Private/LandscapeEdit.cpp
@@ -511,8 +511,11 @@ FString ULandscapeComponent::GetLayerAllocationKey(const TArray<FWeightmapLayerA
 	TArray<FString> LayerStrings;
 	for (int32 LayerIdx = 0; LayerIdx < Allocations.Num(); LayerIdx++)
 	{
-		const bool bNoWeightBlend = Allocations[LayerIdx].LayerInfo && Allocations[LayerIdx].LayerInfo->bNoWeightBlend;
-		LayerStrings.Add(FString::Printf(TEXT("_%s_%s%d"), *Allocations[LayerIdx].GetLayerName().ToString(), bNoWeightBlend ? TEXT("n") : TEXT("w"), Allocations[LayerIdx].WeightmapTextureIndex));
+		if (Allocations[LayerIdx].GetLayerName() != NAME_None) // ignore layers that have no name
+		{
+			const bool bNoWeightBlend = Allocations[LayerIdx].LayerInfo && Allocations[LayerIdx].LayerInfo->bNoWeightBlend;
+			LayerStrings.Add(FString::Printf(TEXT("_%s_%s%d"), *Allocations[LayerIdx].GetLayerName().ToString(), bNoWeightBlend ? TEXT("n") : TEXT("w"), Allocations[LayerIdx].WeightmapTextureIndex));
+		}
 	}
 	// Sort them alphabetically so we can share across components even if the order is different
 	LayerStrings.Sort(TGreater<FString>());
@@ -590,8 +593,10 @@ UMaterialInstanceConstant* ULandscapeComponent::GetCombinationMaterial(FMaterial
 				if (Allocation.LayerInfo)
 				{
 					FName LayerName = Allocation.GetLayerName();
-					check(LayerName != NAME_None);
-					StaticParameters.EditorOnly.TerrainLayerWeightParameters.Add(FStaticTerrainLayerWeightParameter(LayerName, Allocation.WeightmapTextureIndex, !Allocation.LayerInfo->bNoWeightBlend));
+					if (LayerName != NAME_None) // ignore allocations that have no name
+					{
+						StaticParameters.EditorOnly.TerrainLayerWeightParameters.Add(FStaticTerrainLayerWeightParameter(LayerName, Allocation.WeightmapTextureIndex, !Allocation.LayerInfo->bNoWeightBlend));
+					}
 				}
 			}
 
