Fix for crash when landscape mobile material instance cannot be created.
These arrays being filled out are not indexed in a specific order, so we can just skip adding entries when it fails.

diff --git a/Engine/Source/Runtime/Landscape/Private/LandscapeEdit.cpp b/Engine/Source/Runtime/Landscape/Private/LandscapeEdit.cpp
--- a/Engine/Source/Runtime/Landscape/Private/LandscapeEdit.cpp
+++ b/Engine/Source/Runtime/Landscape/Private/LandscapeEdit.cpp
@@ -8046,7 +8046,10 @@ void ULandscapeComponent::GenerateMobilePlatformPixelData(bool bIsCooking, const
 
 			// Find or set a matching MIC in the Landscape's map.
 			MobileCombinationMaterialInstances[MaterialIndex] = GetCombinationMaterial(nullptr, MobileWeightmapLayerAllocations, MaterialLOD, true);
-			check(MobileCombinationMaterialInstances[MaterialIndex] != nullptr);
+			if (MobileCombinationMaterialInstances[MaterialIndex] == nullptr)
+			{
+				continue;
+			}
 
 			if (bIsCooking)
 			{
