Landscape Spline Actor is now created in the same level as it's landscape actor.

diff --git a/Engine/Source/Runtime/Landscape/Private/Landscape.cpp b/Engine/Source/Runtime/Landscape/Private/Landscape.cpp
--- a/Engine/Source/Runtime/Landscape/Private/Landscape.cpp
+++ b/Engine/Source/Runtime/Landscape/Private/Landscape.cpp
@@ -5857,7 +5857,7 @@ ALandscapeSplineActor* ULandscapeInfo::CreateSplineActor(const FVector& Location
 	UWorld* World = LandscapeActor->GetWorld();
 	check(World);
 	FActorSpawnParameters SpawnParams;
-	SpawnParams.OverrideLevel = World->PersistentLevel;
+	SpawnParams.OverrideLevel = LandscapeActor->GetLevel();
 	SpawnParams.bNoFail = true;
 	SpawnParams.ObjectFlags |= RF_Transactional;
 	ALandscapeSplineActor* SplineActor = World->SpawnActor<ALandscapeSplineActor>(Location, FRotator::ZeroRotator, SpawnParams);
