Fixed crash with landscape edit layers batched merge

diff --git a/Engine/Source/Runtime/Landscape/Private/LandscapeEditLayers.cpp b/Engine/Source/Runtime/Landscape/Private/LandscapeEditLayers.cpp
--- a/Engine/Source/Runtime/Landscape/Private/LandscapeEditLayers.cpp
+++ b/Engine/Source/Runtime/Landscape/Private/LandscapeEditLayers.cpp
@@ -5817,6 +5817,8 @@ void ULandscapeEditLayerPersistent::RenderLayer(ILandscapeEditLayerRenderer::FRe
 					{
 						FlushCopyQuads();
 						check(CurrentSourceTextureResources.IsEmpty());
+
+						CurrentDestinationArrayIndex = InSingleQuadParams.DestinationArrayIndex;
 					}
 					CurrentSourceTextureResources.Add(InSingleQuadParams.SourceTextureResource);
 				}
