Added proper material compilation error when reaching the max number of landscape physical material outputs

diff --git a/Engine/Shaders/Private/LandscapePhysicalMaterial.usf b/Engine/Shaders/Private/LandscapePhysicalMaterial.usf
--- a/Engine/Shaders/Private/LandscapePhysicalMaterial.usf
+++ b/Engine/Shaders/Private/LandscapePhysicalMaterial.usf
@@ -111,7 +111,8 @@ void PSMain(
 #if NUM_MATERIAL_OUTPUTS_GETPHYSICALMATERIAL > 15
 	DomininantMaterial = UpdateDominantMaterial(15, GetPhysicalMaterial15(MaterialParameters), DomininantMaterial);
 #endif
-#if NUM_MATERIAL_OUTPUTS_GETPHYSICALMATERIAL > 16
+	
+#if NUM_MATERIAL_OUTPUTS_GETPHYSICALMATERIAL > 16 // Corresponds to : UMaterialExpressionLandscapePhysicalMaterialOutput::MaxPhysicalMaterials == 16. 
 	// We pick an arbitrary limit to support here, but the only cost of extending this is more typing.
 #error Too many physical materials
 #endif
diff --git a/Engine/Source/Runtime/Landscape/Classes/Materials/MaterialExpressionLandscapePhysicalMaterialOutput.h b/Engine/Source/Runtime/Landscape/Classes/Materials/MaterialExpressionLandscapePhysicalMaterialOutput.h
--- a/Engine/Source/Runtime/Landscape/Classes/Materials/MaterialExpressionLandscapePhysicalMaterialOutput.h
+++ b/Engine/Source/Runtime/Landscape/Classes/Materials/MaterialExpressionLandscapePhysicalMaterialOutput.h
@@ -34,6 +34,9 @@ class UMaterialExpressionLandscapePhysicalMaterialOutput : public UMaterialExpre
 {
 	GENERATED_UCLASS_BODY()
 
+	// Maximum number of supported grass types on a given landscape material. Whenever adjusting this, make sure to update LandscapeGrassWeight.usf accordingly:
+	static constexpr int32 MaxPhysicalMaterials = 16;
+
 	/** Array of physical material inputs. */
 	UPROPERTY(EditAnywhere, Category = PhysicalMaterial)
 	TArray< FPhysicalMaterialInput > Inputs;
@@ -57,6 +60,7 @@ class UMaterialExpressionLandscapePhysicalMaterialOutput : public UMaterialExpre
 
 	//~ Begin UMaterialExpressionCustomOutput Interface
 	virtual int32 GetNumOutputs() const override { return Inputs.Num(); }
+	virtual int32 GetMaxOutputs() const { return MaxPhysicalMaterials; };
 	virtual FString GetFunctionName() const override { return TEXT("GetPhysicalMaterial"); }
 
 #if WITH_EDITOR
