Optimization for landscape sculpt when many proxies are loaded.  Don't lock UTextures for write when we know the data isn't changing.  Avoids expensive ComputeChannelLinearMinMax and FIoHash for unchanged proxies at the end of the stroke.  Most applicable with a large number of landscape proxies loaded.

diff --git a/Engine/Source/Runtime/Landscape/Classes/Landscape.h b/Engine/Source/Runtime/Landscape/Classes/Landscape.h
--- a/Engine/Source/Runtime/Landscape/Classes/Landscape.h
+++ b/Engine/Source/Runtime/Landscape/Classes/Landscape.h
@@ -545,7 +545,7 @@ private:
 	int32 PerformLayersWeightmapsBatchedMerge(FUpdateLayersContentContext& InUpdateLayersContentContext, const FEditLayersWeightmapMergeParams& InMergeParams);
 	void ResolveLayersWeightmapTexture(const FTextureToComponentHelper& MapHelper, const TSet<UTexture2D*>& WeightmapsToResolve, bool bIntermediateRender, TArray<FLandscapeEditLayerComponentReadbackResult>& InOutComponentReadbackResults);
 
-	bool HasTextureDataChanged(TArrayView<const FColor> InOldData, TArrayView<const FColor> InNewData, bool bInIsWeightmap, uint64 InPreviousHash, uint64& OutNewHash, TOptional<uint8>& OutChangedWeightmapChannelsMasks) const;
+	bool HasTextureDataChanged(TArrayView<const FColor> InOldData, TArrayView<const FColor> InNewData, bool bInIsWeightmap, uint64 InPreviousHash, uint64 InNewHash, TOptional<uint8>& OutChangedWeightmapChannelsMasks) const;
 	bool ResolveLayersTexture(FTextureToComponentHelper const& MapHelper, FLandscapeEditLayerReadback* InCPUReadBack, UTexture2D* InOutputTexture, bool bIntermediateRender,
 		TArray<FLandscapeEditLayerComponentReadbackResult>& InOutComponentReadbackResults, bool bIsWeightmap);
 
diff --git a/Engine/Source/Runtime/Landscape/Private/LandscapeEditLayers.cpp b/Engine/Source/Runtime/Landscape/Private/LandscapeEditLayers.cpp
--- a/Engine/Source/Runtime/Landscape/Private/LandscapeEditLayers.cpp
+++ b/Engine/Source/Runtime/Landscape/Private/LandscapeEditLayers.cpp
@@ -8184,12 +8184,11 @@ void ALandscape::OnDirtyHeightmap(FTextureToComponentHelper const& MapHelper, UT
 	}
 }
 
-bool ALandscape::HasTextureDataChanged(TArrayView<const FColor> InOldData, TArrayView<const FColor> InNewData, bool bInIsWeightmap, uint64 InPreviousHash, uint64& OutNewHash, TOptional<uint8>& OutChangedWeightmapChannelsMasks) const
+bool ALandscape::HasTextureDataChanged(TArrayView<const FColor> InOldData, TArrayView<const FColor> InNewData, bool bInIsWeightmap, uint64 InPreviousHash, uint64 InNewHash, TOptional<uint8>& OutChangedWeightmapChannelsMasks) const
 {
 	TRACE_CPUPROFILER_EVENT_SCOPE(ALandscape::HasTextureDataChanged);
 
-	OutNewHash = FLandscapeEditLayerReadback::CalculateHash((uint8*)InNewData.GetData(), InNewData.Num() * sizeof(FColor));
-	if (OutNewHash != InPreviousHash)
+	if (InPreviousHash != InNewHash)
 	{
 		int32 TextureSize = InOldData.Num();
 		check(TextureSize == InNewData.Num());
@@ -8209,11 +8208,11 @@ bool ALandscape::HasTextureDataChanged(TArrayView<const FColor> InOldData, TArra
 					const FColor& NewColor = InNewData[Index];
 					if (OldColor != NewColor)
 					{
-						auto DiffChannel = [DirtyWeightmapThreshold] (uint8 InOldValue, uint8 InNewValue) -> bool
-						{
-							uint8 Diff = (InNewValue > InOldValue) ? (InNewValue - InOldValue) : (InOldValue - InNewValue);
-							return (Diff > DirtyWeightmapThreshold);
-						};
+						auto DiffChannel = [DirtyWeightmapThreshold](uint8 InOldValue, uint8 InNewValue) -> bool
+							{
+								uint8 Diff = (InNewValue > InOldValue) ? (InNewValue - InOldValue) : (InOldValue - InNewValue);
+								return (Diff > DirtyWeightmapThreshold);
+							};
 
 						uint8 DiffMask = ((uint8)(DiffChannel(OldColor.R, NewColor.R) ? 1 : 0) << 0)
 							| ((uint8)(DiffChannel(OldColor.G, NewColor.G) ? 1 : 0) << 1)
@@ -8239,7 +8238,7 @@ bool ALandscape::HasTextureDataChanged(TArrayView<const FColor> InOldData, TArra
 				return bHasChanged;
 			}
 		}
-		else 
+		else
 		{
 			int32 DirtyHeightmapHeightThreshold = CVarLandscapeDirtyHeightmapHeightThreshold.GetValueOnGameThread();
 			int32 DirtyHeightmapNormalThreshold = CVarLandscapeDirtyHeightmapNormalThreshold.GetValueOnGameThread();
@@ -8274,11 +8273,11 @@ bool ALandscape::HasTextureDataChanged(TArrayView<const FColor> InOldData, TArra
 			}
 		}
 
-		// Hash is different and we don't need a deep comparison : 
+		// Hash is different and we don't need a deep comparison :
 		return true;
 	}
 
-	// Hash is identical : 
+	// Hash is identical :
 	return false;
 }
 
@@ -8298,7 +8297,7 @@ bool ALandscape::ResolveLayersTexture(
 
 	bool bUserTriggered = false;
 
-	if (TArray<ULandscapeComponent*> const * Components = bIsWeightmap ? MapHelper.WeightmapToComponents.Find(InOutputTexture) : MapHelper.HeightmapToComponents.Find(InOutputTexture))
+	if (TArray<ULandscapeComponent*> const* Components = bIsWeightmap ? MapHelper.WeightmapToComponents.Find(InOutputTexture) : MapHelper.HeightmapToComponents.Find(InOutputTexture))
 	{
 		for (ULandscapeComponent* Component : *Components)
 		{
@@ -8309,13 +8308,13 @@ bool ALandscape::ResolveLayersTexture(
 			}
 		}
 	}
-	
+
 	bool bChanged = false;
 	TOptional<uint8> ChangedChannelsMask;
 	if (bIsWeightmap)
 	{
 		// Request a precise report of which channel have changed if we need to dump the weightmap diffs :
-		if ((CVarLandscapeDumpWeightmapDiff.GetValueOnGameThread() != 0) 
+		if ((CVarLandscapeDumpWeightmapDiff.GetValueOnGameThread() != 0)
 			|| CVarLandscapeTrackDirty.GetValueOnGameThread() != 0)
 		{
 			ChangedChannelsMask.Emplace(0);
@@ -8329,41 +8328,63 @@ bool ALandscape::ResolveLayersTexture(
 		// Copy final result to texture source.
 		TArray<TArray<FColor>> const& OutMipsData = InCPUReadback->GetResult(CompletedReadbackNum - 1);
 
+		// Keep track if we locked the texture for write or not.  bChanged implies bLockedForWrite, but not the other way around.  We might have
+		// bLockedForWrite because of a data hash change, but bChanged=false due to the threshold check.
+		bool bLockedForWrite = false;
+
 		for (int8 MipIndex = 0; MipIndex < OutMipsData.Num(); ++MipIndex)
 		{
 			int32 TextureSize = OutMipsData[MipIndex].Num();
 			if (TextureSize > 0)
 			{
-				uint8* TextureData = InOutputTexture->Source.LockMip(MipIndex);
+				uint8* TextureData = nullptr;
 
 				// Do dirty detection on first mip.
 				// Don't do this for intermediate renders.
 				if (MipIndex == 0 && !bIntermediateRender)
 				{
-					uint64 Hash = 0;
-					if (HasTextureDataChanged(MakeArrayView(reinterpret_cast<const FColor*>(TextureData), TextureSize), MakeArrayView(OutMipsData[MipIndex].GetData(), TextureSize), bIsWeightmap, InCPUReadback->GetHash(), Hash, ChangedChannelsMask))
-					{
-						bChanged |= InCPUReadback->SetHash(Hash);
-						check(bChanged);
-					}
+					uint64 OldHash = InCPUReadback->GetHash();
+					uint64 NewHash = FLandscapeEditLayerReadback::CalculateHash((uint8*)OutMipsData[MipIndex].GetData(), TextureSize * sizeof(FColor));
 
-					if (bChanged)
+					if (NewHash != OldHash)
 					{
-						// We're about to modify the texture's source data, the texture needs to know so that it can handle properly update cached platform data (additionally, the package needs to be dirtied) :
-						if (GetDefault<ULandscapeSettings>()->LandscapeDirtyingMode == ELandscapeDirtyingMode::InLandscapeModeAndUserTriggeredChanges)
+						// Defer locking the texture for ReadWrite until after we know that the data is changing at all.  Unlocking after a ReadWrite causes an 
+						// expensive rehashing using the slower hash function used in IoHash.  
+						TextureData = InOutputTexture->Source.LockMip(MipIndex);
+						check(!bLockedForWrite);
+						bLockedForWrite = true;
+
+						if (HasTextureDataChanged(MakeArrayView(reinterpret_cast<const FColor*>(TextureData), TextureSize), MakeArrayView(OutMipsData[MipIndex].GetData(), TextureSize), bIsWeightmap, OldHash, NewHash, ChangedChannelsMask))
 						{
-							FLandscapeDirtyOnlyInModeScope Scope(GetLandscapeInfo(), !bUserTriggered);
-							GetLandscapeInfo()->ModifyObject(InOutputTexture);
+							bChanged |= InCPUReadback->SetHash(NewHash);
+							check(bChanged);
 						}
-						else
+
+						if (bChanged)
 						{
-							GetLandscapeInfo()->ModifyObject(InOutputTexture);
+							// We're about to modify the texture's source data, the texture needs to know so that it can handle properly update cached platform data (additionally, the package needs to be dirtied) :
+							if (GetDefault<ULandscapeSettings>()->LandscapeDirtyingMode == ELandscapeDirtyingMode::InLandscapeModeAndUserTriggeredChanges)
+							{
+								FLandscapeDirtyOnlyInModeScope Scope(GetLandscapeInfo(), !bUserTriggered);
+								GetLandscapeInfo()->ModifyObject(InOutputTexture);
+							}
+							else
+							{
+								GetLandscapeInfo()->ModifyObject(InOutputTexture);
+							}
 						}
 					}
 				}
+				else if (bLockedForWrite || bIntermediateRender)
+				{
+					// If we locked mip 0, lock the rest too, even though FTextureSource isn't tracking mip-level locking.
+					TextureData = InOutputTexture->Source.LockMip(MipIndex);
+					bLockedForWrite = true;
+				}
 
 				if (bChanged)
 				{
+					check(bLockedForWrite && TextureData);
 					if (bIsWeightmap)
 					{
 						OnDirtyWeightmap(MapHelper, InOutputTexture, (FColor*)TextureData, OutMipsData[MipIndex].GetData(), MipIndex, ChangedChannelsMask.Get(0));
@@ -8374,23 +8395,29 @@ bool ALandscape::ResolveLayersTexture(
 					}
 				}
 
-				TRACE_CPUPROFILER_EVENT_SCOPE(ReadbackToCPU);
-				FMemory::Memcpy(TextureData, OutMipsData[MipIndex].GetData(), OutMipsData[MipIndex].Num() * sizeof(FColor));
+				if (bLockedForWrite)
+				{
+					TRACE_CPUPROFILER_EVENT_SCOPE(ReadbackToCPU);
+					FMemory::Memcpy(TextureData, OutMipsData[MipIndex].GetData(), OutMipsData[MipIndex].Num() * sizeof(FColor));
+				}
 			}
 		}
 
-		// Unlock all mips at once because there's a lock counter in FTextureSource that recomputes the content hash when reaching 0 (which means we'd recompute the hash several times over if we Lock/Unlock/Lock/Unloc/... for each mip ):
-		for (int8 MipIndex = 0; MipIndex < OutMipsData.Num(); ++MipIndex)
+		if (bLockedForWrite)
 		{
-			if (OutMipsData[MipIndex].Num() > 0)
+			// Unlock all mips at once because there's a lock counter in FTextureSource that recomputes the content hash when reaching 0 (which means we'd recompute the hash several times over if we Lock/Unlock/Lock/Unloc/... for each mip ):
+			for (int8 MipIndex = 0; MipIndex < OutMipsData.Num(); ++MipIndex)
 			{
-				InOutputTexture->Source.UnlockMip(MipIndex);
+				if (OutMipsData[MipIndex].Num() > 0)
+				{
+					InOutputTexture->Source.UnlockMip(MipIndex);
+				}
 			}
 		}
 
 		// change lighting guid to be the hash of the source data (so we can use lighting guid to detect when it actually changes)
 		// TODO [chris.tchou] we should check(InOutputTexture->Source.bGuidIsHash); .. but it's not a public value currently.
- 		InOutputTexture->SetLightingGuid(InOutputTexture->Source.GetId());
+		InOutputTexture->SetLightingGuid(InOutputTexture->Source.GetId());
 
 		// Find out whether some channels from this weightmap are now all zeros : 
 		static constexpr uint32 AllChannelsAllZerosMask = 15;
