Prepare BodyCollisionData for Landscape HeightFields

diff --git a/Engine/Source/Runtime/Landscape/Private/LandscapeCollision.cpp b/Engine/Source/Runtime/Landscape/Private/LandscapeCollision.cpp
--- a/Engine/Source/Runtime/Landscape/Private/LandscapeCollision.cpp
+++ b/Engine/Source/Runtime/Landscape/Private/LandscapeCollision.cpp
@@ -396,6 +396,10 @@ void ULandscapeHeightfieldCollisionComponent::OnCreatePhysicsState()
 
 				TUniquePtr<Chaos::FPerShapeData> NewShape = Chaos::FShapeInstanceProxy::Make(ShapeArray.Num(), ChaosHeightFieldFromCooked);
 
+				FBodyCollisionData BodyCollisionData;
+				FBodyInstance::BuildBodyCollisionFlags(BodyCollisionData.CollisionFlags, GetCollisionEnabled(), false);
+				const FBodyCollisionFlags& CollisionFlags = BodyCollisionData.CollisionFlags;
+
 				// Setup filtering
 				FCollisionFilterData QueryFilterData, SimFilterData;
 				CreateShapeFilterData(static_cast<uint8>(GetCollisionObjectType()), FMaskFilter(0), GetOwner()->GetUniqueID(), GetCollisionResponseToChannels(), 
@@ -409,6 +413,14 @@ void ULandscapeHeightfieldCollisionComponent::OnCreatePhysicsState()
 				NewShape->SetSimData(SimFilterData);
 				NewShape->SetMaterials(HeightfieldRef->UsedChaosMaterials);
 
+				const bool bNewQueryShape = CollisionFlags.bEnableQueryCollision;
+				const bool bNewSimShape = CollisionFlags.bEnableSimCollisionSimple;
+				const bool bNewProbeShape = CollisionFlags.bEnableProbeCollision;
+				
+				NewShape->SetIsProbe(bNewProbeShape);
+				NewShape->SetSimEnabled(bNewSimShape);
+				NewShape->SetQueryEnabled(bNewQueryShape);
+
 				Geoms.Emplace(MoveTemp(ChaosHeightFieldFromCooked));
 				ShapeArray.Emplace(MoveTemp(NewShape));
 
@@ -431,6 +443,10 @@ void ULandscapeHeightfieldCollisionComponent::OnCreatePhysicsState()
 					NewSimpleShape->SetSimData(SimFilterDataSimple);
 					NewSimpleShape->SetMaterials(HeightfieldRef->UsedChaosMaterials);
 
+					NewSimpleShape->SetIsProbe(bNewProbeShape);
+					NewSimpleShape->SetSimEnabled(bNewSimShape);
+					NewSimpleShape->SetQueryEnabled(bNewQueryShape);
+
 					Geoms.Emplace(MoveTemp(ChaosSimpleHeightFieldFromCooked));
 					ShapeArray.Emplace(MoveTemp(NewSimpleShape));
 				}
@@ -457,6 +473,10 @@ void ULandscapeHeightfieldCollisionComponent::OnCreatePhysicsState()
 					NewEditorShape->SetSimData(SimFilterDataEd);
 					NewEditorShape->SetMaterials(HeightfieldRef->UsedChaosMaterials);
 
+					NewEditorShape->SetIsProbe(bNewProbeShape);
+					NewEditorShape->SetSimEnabled(bNewSimShape);
+					NewEditorShape->SetQueryEnabled(bNewQueryShape);
+
 					Geoms.Emplace(MoveTemp(ChaosEditorHeightFieldFromCooked));
 					ShapeArray.Emplace(MoveTemp(NewEditorShape));
 				}
