Fixed checkSlow when building Nanite landscape with visibility

diff --git a/Engine/Source/Runtime/Landscape/Private/LandscapeEdit.cpp b/Engine/Source/Runtime/Landscape/Private/LandscapeEdit.cpp
--- a/Engine/Source/Runtime/Landscape/Private/LandscapeEdit.cpp
+++ b/Engine/Source/Runtime/Landscape/Private/LandscapeEdit.cpp
@@ -3946,6 +3946,7 @@ namespace UE::Landscape
 	// don.boogert-todo: we don't consider saddle points correctly but it doesn't seem to be an issue for our application in landscape geometry export.
 	void GenerateMarchingSquaresGeometry(const TStaticArray<float, UE_ARRAY_COUNT(QuadPattern)>& InVisibilities, float InThreshold, const TStaticArray<FVector, UE_ARRAY_COUNT(QuadPattern)>& InPositions, TArray<int32, TInlineAllocator<12>>& OutIndices, TArray<FVector, TInlineAllocator<6>>& OutPositions)
 	{
+		static_assert(UE_ARRAY_COUNT(QuadPattern) == 4, "Invalid quad pattern");
 		uint8 V0 = InVisibilities[0] < InThreshold ? 1 : 0;
 		uint8 V1 = InVisibilities[1] < InThreshold ? 1 : 0;
 		uint8 V2 = InVisibilities[2] < InThreshold ? 1 : 0;
@@ -3953,16 +3954,17 @@ namespace UE::Landscape
 
 		uint8 Case = (V1 << 3) | (V2 << 2) | (V3 << 1) | V0;
 
-		const float Alpha0 = (InThreshold - InVisibilities[0]) / (InVisibilities[1] - InVisibilities[0]);
-		const float Alpha1 = (InThreshold - InVisibilities[1]) / (InVisibilities[2] - InVisibilities[1]);
-		const float Alpha2 = (InThreshold - InVisibilities[2]) / (InVisibilities[3] - InVisibilities[2]);
-		const float Alpha3 = (InThreshold - InVisibilities[3]) / (InVisibilities[0] - InVisibilities[3]);
-
-		FVector Interpolated0 = FMath::Lerp(InPositions[0], InPositions[1], Alpha0);
-		FVector Interpolated1 = FMath::Lerp(InPositions[1], InPositions[2], Alpha1);
-		FVector Interpolated2 = FMath::Lerp(InPositions[2], InPositions[3], Alpha2);
-		FVector Interpolated3 = FMath::Lerp(InPositions[3], InPositions[0], Alpha3);
-
+		auto InterpolateWithVisibility = [&InVisibilities, &InPositions, InThreshold](uint8 InPointIndex) -> FVector
+			{
+				const uint32 OtherPointIndex = (InPointIndex + 1) % 4;
+				float Visibility = InVisibilities[InPointIndex];
+				float OtherVisibility = InVisibilities[OtherPointIndex];
+				const float Alpha = (Visibility != OtherVisibility) 
+					? FMath::Clamp((InThreshold - Visibility) / (OtherVisibility - Visibility), 0.0f, 1.0f)
+					: 0.5f;
+				return FMath::Lerp(InPositions[InPointIndex], InPositions[OtherPointIndex], Alpha);
+			};
+			
 		switch (Case)
 		{
 		case 0:
@@ -3970,91 +3972,59 @@ namespace UE::Landscape
 			OutIndices.Empty();
 			break;
 		case 1:
-			check(Alpha0 >= 0.0f && Alpha0 <= 1.0f);
-			check(Alpha3 >= 0.0f && Alpha3 <= 1.0f);
-			OutPositions = { InPositions[0], Interpolated0, Interpolated3 };
+			OutPositions = { InPositions[0], InterpolateWithVisibility(0), InterpolateWithVisibility(3)};
 			OutIndices = { 0, 1, 2 };
 			break;
 		case 2:
-			check(Alpha2 >= 0.0f && Alpha2 <= 1.0f);
-			check(Alpha3 >= 0.0f && Alpha3 <= 1.0f);
-			OutPositions = { InPositions[3], Interpolated3, Interpolated2 };
+			OutPositions = { InPositions[3],  InterpolateWithVisibility(3),  InterpolateWithVisibility(2) };
 			OutIndices = { 0, 1, 2 };
 			break;
 		case 3:
-			check(Alpha0 >= 0.0f && Alpha0 <= 1.0f);
-			check(Alpha2 >= 0.0f && Alpha2 <= 1.0f);
-			OutPositions = { InPositions[0], Interpolated0, Interpolated2, InPositions[3] };
+			OutPositions = { InPositions[0],  InterpolateWithVisibility(0),  InterpolateWithVisibility(2), InPositions[3] };
 			OutIndices = { 0, 1, 2, 0, 2, 3 };
 			break;
 		case 4:
-			check(Alpha1 >= 0.0f && Alpha1 <= 1.0f);
-			check(Alpha2 >= 0.0f && Alpha2 <= 1.0f);
-			OutPositions = { InPositions[2], Interpolated2, Interpolated1 };
+			OutPositions = { InPositions[2],  InterpolateWithVisibility(2),  InterpolateWithVisibility(1) };
 			OutIndices = { 0, 1, 2 };
 			break;
 		case 5:
-			check(Alpha0 >= 0.0f && Alpha0 <= 1.0f);
-			check(Alpha1 >= 0.0f && Alpha1 <= 1.0f);
-			check(Alpha2 >= 0.0f && Alpha2 <= 1.0f);
-			check(Alpha3 >= 0.0f && Alpha3 <= 1.0f);
-			OutPositions = { InPositions[0], Interpolated0, Interpolated1, InPositions[2],Interpolated2, Interpolated3 };
+			OutPositions = { InPositions[0],  InterpolateWithVisibility(0),  InterpolateWithVisibility(1), InPositions[2], InterpolateWithVisibility(2),  InterpolateWithVisibility(3) };
 			OutIndices = { 0, 1, 5, 1, 2, 5, 2, 4, 5, 2, 3, 4 };
 			break;
 		case 6:
-			check(Alpha1 >= 0.0f && Alpha1 <= 1.0f);
-			check(Alpha3 >= 0.0f && Alpha3 <= 1.0f);
-			OutPositions = { InPositions[3], Interpolated3, Interpolated1, InPositions[2] };
+			OutPositions = { InPositions[3],  InterpolateWithVisibility(3),  InterpolateWithVisibility(1), InPositions[2] };
 			OutIndices = { 0, 1, 2, 0, 2, 3 };
 			break;
 		case 7:
-			check(Alpha0 >= 0.0f && Alpha0 <= 1.0f);
-			check(Alpha1 >= 0.0f && Alpha1 <= 1.0f);
-			OutPositions = { InPositions[0], Interpolated0, Interpolated1, InPositions[2], InPositions[3] };
+			OutPositions = { InPositions[0],  InterpolateWithVisibility(0),  InterpolateWithVisibility(1), InPositions[2], InPositions[3] };
 			OutIndices = { 0, 1, 2, 0, 2, 3, 0, 3, 4 };
 			break;
 		case 8:
-			check(Alpha0 >= 0.0f && Alpha0 <= 1.0f);
-			check(Alpha1 >= 0.0f && Alpha1 <= 1.0f);
-			OutPositions = { InPositions[1], Interpolated1, Interpolated0 };
+			OutPositions = { InPositions[1],  InterpolateWithVisibility(1),  InterpolateWithVisibility(0) };
 			OutIndices = { 0, 1, 2 };
 			break;
 		case 9:
-			check(Alpha1 >= 0.0f && Alpha1 <= 1.0f);
-			check(Alpha3 >= 0.0f && Alpha3 <= 1.0f);
-			OutPositions = { InPositions[0], InPositions[1], Interpolated1, Interpolated3 };
+			OutPositions = { InPositions[0], InPositions[1],  InterpolateWithVisibility(1),  InterpolateWithVisibility(3) };
 			OutIndices = { 0, 1, 2, 0, 2, 3 };
 			break;
 		case 10:
-			check(Alpha0 >= 0.0f && Alpha0 <= 1.0f);
-			check(Alpha1 >= 0.0f && Alpha1 <= 1.0f);
-			check(Alpha2 >= 0.0f && Alpha2 <= 1.0f);
-			check(Alpha3 >= 0.0f && Alpha3 <= 1.0f);
-			OutPositions = { InPositions[1], Interpolated1, Interpolated2, InPositions[3],Interpolated3, Interpolated0 };
+			OutPositions = { InPositions[1],  InterpolateWithVisibility(1),  InterpolateWithVisibility(2), InPositions[3], InterpolateWithVisibility(3),  InterpolateWithVisibility(0) };
 			OutIndices = { 0, 1, 5, 1, 2, 5, 2, 4, 5, 2, 3, 4 };
 			break;
 		case 11:
-			check(Alpha1 >= 0.0f && Alpha1 <= 1.0f);
-			check(Alpha2 >= 0.0f && Alpha2 <= 1.0f);
-			OutPositions = { InPositions[1], Interpolated1, Interpolated2, InPositions[3], InPositions[0] };
+			OutPositions = { InPositions[1],  InterpolateWithVisibility(1),  InterpolateWithVisibility(2), InPositions[3], InPositions[0] };
 			OutIndices = { 0, 1, 2, 0, 2, 3, 0, 3, 4 };
 			break;
 		case 12:
-			check(Alpha0 >= 0.0f && Alpha0 <= 1.0f);
-			check(Alpha2 >= 0.0f && Alpha2 <= 1.0f);
-			OutPositions = { InPositions[2], Interpolated2, Interpolated0, InPositions[1] };
+			OutPositions = { InPositions[2],  InterpolateWithVisibility(2),  InterpolateWithVisibility(0), InPositions[1] };
 			OutIndices = { 0, 1, 2, 0, 2, 3 };
 			break;
 		case 13:
-			check(Alpha2 >= 0.0f && Alpha2 <= 1.0f);
-			check(Alpha3 >= 0.0f && Alpha3 <= 1.0f);
-			OutPositions = { InPositions[2], Interpolated2, Interpolated3, InPositions[0], InPositions[1] };
+			OutPositions = { InPositions[2],  InterpolateWithVisibility(2),  InterpolateWithVisibility(3), InPositions[0], InPositions[1] };
 			OutIndices = { 0, 1, 2, 0, 2, 3, 0, 3, 4 };
 			break;
 		case 14:
-			check(Alpha0 >= 0.0f && Alpha0 <= 1.0f);
-			check(Alpha3 >= 0.0f && Alpha3 <= 1.0f);
-			OutPositions = { InPositions[3], Interpolated3, Interpolated0, InPositions[1], InPositions[2] };
+			OutPositions = { InPositions[3],  InterpolateWithVisibility(3),  InterpolateWithVisibility(0), InPositions[1], InPositions[2] };
 			OutIndices = { 0, 1, 2, 0, 2, 3, 0, 3, 4 };
 			break;
 		case 15:
