Added debug color to physical materials, to have them maintain the same color from run to run and be tweakable by the user. Currently only supported by landscape collision display (with landscape.CollisionMesh.ShowPhysicalMaterial). Initialized with a random color.

diff --git a/Engine/Source/Editor/UnrealEd/Private/Factories/EditorFactories.cpp b/Engine/Source/Editor/UnrealEd/Private/Factories/EditorFactories.cpp
--- a/Engine/Source/Editor/UnrealEd/Private/Factories/EditorFactories.cpp
+++ b/Engine/Source/Editor/UnrealEd/Private/Factories/EditorFactories.cpp
@@ -2087,6 +2087,9 @@ UObject* UPhysicalMaterialFactoryNew::FactoryCreateNew(UClass* Class, UObject* I
 	check(ClassToUse->IsChildOf(UPhysicalMaterial::StaticClass()));
 
 	UPhysicalMaterial* NewMaterial = NewObject<UPhysicalMaterial>(InParent, ClassToUse, Name, MaterialFlags);
+	// Assign a random debug color in the factory. The user can then change it at will : 
+	NewMaterial->DebugColor = FLinearColor::MakeRandomSeededColor(GetTypeHash(NewMaterial->GetPathName()));
+
 	// A call to get will ensure any physics-engine specific data is built
 	NewMaterial->GetPhysicsMaterial();
 	return NewMaterial;
diff --git a/Engine/Source/Runtime/Core/Public/UObject/FortniteMainBranchObjectVersion.h b/Engine/Source/Runtime/Core/Public/UObject/FortniteMainBranchObjectVersion.h
--- a/Engine/Source/Runtime/Core/Public/UObject/FortniteMainBranchObjectVersion.h
+++ b/Engine/Source/Runtime/Core/Public/UObject/FortniteMainBranchObjectVersion.h
@@ -550,6 +550,9 @@ struct FFortniteMainBranchObjectVersion
 		// Make sure scene graph entities are not public by default
 		SceneGraphEntitiesPrivateByDefault,
 
+		// Added debug color for physical materials
+		DebugColorForPhysicalMaterials,
+
 		// -----<new versions can be added above this line>-------------------------------------------------
 		VersionPlusOne,
 		LatestVersion = VersionPlusOne - 1
diff --git a/Engine/Source/Runtime/Landscape/Private/LandscapeCollision.cpp b/Engine/Source/Runtime/Landscape/Private/LandscapeCollision.cpp
--- a/Engine/Source/Runtime/Landscape/Private/LandscapeCollision.cpp
+++ b/Engine/Source/Runtime/Landscape/Private/LandscapeCollision.cpp
@@ -592,9 +592,7 @@ FPrimitiveSceneProxy* ULandscapeHeightfieldCollisionComponent::CreateSceneProxy(
 			{
 				Chaos::FChaosPhysicsMaterial* ChaosMaterial = MaterialHandle.Get();
 				UPhysicalMaterial* PhysicalMaterial = (ChaosMaterial != nullptr) ? FChaosUserData::Get<UPhysicalMaterial>(ChaosMaterial->UserData) : nullptr;
-				uint32 Hash = PointerHash(PhysicalMaterial);
-				uint8* HashElement = reinterpret_cast<uint8*>(&Hash);
-				MaterialIndexColors.Emplace(HashElement[0], HashElement[1], HashElement[2]);
+				MaterialIndexColors.Emplace(PhysicalMaterial->DebugColor.ToFColor(/*bSRGB = */false));
 			}
 
 			for (int32 I = 0; I < NumVerts; I++)
diff --git a/Engine/Source/Runtime/Landscape/Private/LandscapeModule.cpp b/Engine/Source/Runtime/Landscape/Private/LandscapeModule.cpp
--- a/Engine/Source/Runtime/Landscape/Private/LandscapeModule.cpp
+++ b/Engine/Source/Runtime/Landscape/Private/LandscapeModule.cpp
@@ -10,6 +10,7 @@
 #include "LandscapeComponent.h"
 #include "LandscapeVersion.h"
 #include "LandscapeInfoMap.h"
+#include "LandscapeHeightfieldCollisionComponent.h"
 #include "Materials/MaterialInstance.h"
 
 #include "LandscapeProxy.h"
@@ -19,11 +20,14 @@
 #include "Engine/Texture2D.h"
 #include "EngineUtils.h"
 
+#include "Chaos/PhysicalMaterials.h"
+
 #if WITH_EDITOR
 #include "WorldPartition/WorldPartitionActorDesc.h"
 #include "UObject/UE5MainStreamObjectVersion.h"
 #endif
 
+#define ALLOW_LANDSCAPE_COLLISION_COMPONENT_DEBUG_DISPLAY !(UE_BUILD_SHIPPING || UE_BUILD_TEST)
 
 // Register the custom version with core
 FCustomVersionRegistration GRegisterLandscapeCustomVersion(FLandscapeCustomVersion::GUID, FLandscapeCustomVersion::LatestVersion, TEXT("Landscape"));
@@ -39,13 +43,23 @@ public:
 
 	virtual void SetLandscapeEditorServices(ILandscapeEditorServices* InLandscapeEditorServices) override { LandscapeEditorServices = InLandscapeEditorServices; }
 	virtual ILandscapeEditorServices* GetLandscapeEditorServices() const override { return LandscapeEditorServices; }
+
 private:
 	void OnPostEngineInit();
 	void OnEnginePreExit();
 
+#if ALLOW_LANDSCAPE_COLLISION_COMPONENT_DEBUG_DISPLAY
+	void OnUpdatePhysicalMaterial(Chaos::FMaterialHandle InHandle);
+#endif // ALLOW_LANDSCAPE_COLLISION_COMPONENT_DEBUG_DISPLAY
+
 private:
 	TSharedPtr<FLandscapeSceneViewExtension, ESPMode::ThreadSafe> SceneViewExtension;
 	ILandscapeEditorServices* LandscapeEditorServices = nullptr;
+
+	/** Event binding handles */
+#if ALLOW_LANDSCAPE_COLLISION_COMPONENT_DEBUG_DISPLAY
+	FDelegateHandle OnUpdatePhysicalMaterialHandle;
+#endif // ALLOW_LANDSCAPE_COLLISION_COMPONENT_DEBUG_DISPLAY
 };
 
 /**
@@ -218,6 +232,14 @@ void FLandscapeModule::StartupModule()
 	FCoreDelegates::OnPostEngineInit.AddRaw(this, &FLandscapeModule::OnPostEngineInit);
 	FCoreDelegates::OnEnginePreExit.AddRaw(this, &FLandscapeModule::OnEnginePreExit);
 
+	// Bind material events
+#if ALLOW_LANDSCAPE_COLLISION_COMPONENT_DEBUG_DISPLAY
+	{
+		Chaos::FPhysicalMaterialManager& MaterialManager = Chaos::FPhysicalMaterialManager::Get();
+		OnUpdatePhysicalMaterialHandle = MaterialManager.OnMaterialUpdated.Add(Chaos::FMaterialUpdatedDelegate::CreateRaw(this, &FLandscapeModule::OnUpdatePhysicalMaterial));
+	}
+#endif // ALLOW_LANDSCAPE_COLLISION_COMPONENT_DEBUG_DISPLAY
+
 #if WITH_EDITOR
 	// Register LandscapeSplineActorDesc Deprecation
 	FWorldPartitionActorDesc::RegisterActorDescDeprecator(ALandscapeSplineActor::StaticClass(), [](FArchive& Ar, FWorldPartitionActorDesc* ActorDesc)
@@ -234,7 +256,7 @@ void FLandscapeModule::StartupModule()
 			ActorDesc->AddProperty(ALandscape::AffectsLandscapeActorDescProperty, *LandscapeGuid.ToString());
 		}
 	});
-#endif
+#endif // WITH_EDITOR
 }
 
 void FLandscapeModule::OnPostEngineInit()
@@ -251,8 +273,30 @@ void FLandscapeModule::OnEnginePreExit()
 
 void FLandscapeModule::ShutdownModule()
 {
+	// Unbind material events
+#if ALLOW_LANDSCAPE_COLLISION_COMPONENT_DEBUG_DISPLAY
+	{
+		Chaos::FPhysicalMaterialManager& MaterialManager = Chaos::FPhysicalMaterialManager::Get();
+		MaterialManager.OnMaterialUpdated.Remove(OnUpdatePhysicalMaterialHandle);
+	}
+#endif // ALLOW_LANDSCAPE_COLLISION_COMPONENT_DEBUG_DISPLAY
+
 	FCoreDelegates::OnEnginePreExit.RemoveAll(this);
 	FCoreDelegates::OnPostEngineInit.RemoveAll(this);
 }
 
+#if ALLOW_LANDSCAPE_COLLISION_COMPONENT_DEBUG_DISPLAY
+void FLandscapeModule::OnUpdatePhysicalMaterial(Chaos::FMaterialHandle InHandle)
+{
+	// Just rebuild all collision components' render states when the Chaos material changes, this is just for debug purposes anyway : 
+	for (ULandscapeHeightfieldCollisionComponent* LandscapeHeightfieldCollisionComponent 
+		: TObjectRange<ULandscapeHeightfieldCollisionComponent>(RF_ClassDefaultObject | RF_ArchetypeObject, /*bIncludeDerivedClasses = */true, EInternalObjectFlags::Garbage))
+	{
+		LandscapeHeightfieldCollisionComponent->MarkRenderStateDirty();
+	}
+}
+#endif // ALLOW_LANDSCAPE_COLLISION_COMPONENT_DEBUG_DISPLAY
+
 IMPLEMENT_MODULE(FLandscapeModule, Landscape);
+
+#undef ALLOW_LANDSCAPE_COLLISION_COMPONENT_DEBUG_DISPLAY
\ No newline at end of file
diff --git a/Engine/Source/Runtime/PhysicsCore/Private/PhysicalMaterial.cpp b/Engine/Source/Runtime/PhysicsCore/Private/PhysicalMaterial.cpp
--- a/Engine/Source/Runtime/PhysicsCore/Private/PhysicalMaterial.cpp
+++ b/Engine/Source/Runtime/PhysicsCore/Private/PhysicalMaterial.cpp
@@ -9,6 +9,7 @@
 #include "UObject/UObjectIterator.h"
 #include "Chaos/PhysicalMaterials.h"
 #include "HAL/IConsoleManager.h"
+#include "UObject/FortniteMainBranchObjectVersion.h"
 
 #include UE_INLINE_GENERATED_CPP_BY_NAME(PhysicalMaterial)
 
@@ -103,6 +104,12 @@ void UPhysicalMaterial::RebuildPhysicalMaterials()
 
 #endif // WITH_EDITOR
 
+void UPhysicalMaterial::Serialize(FArchive& Ar)
+{
+	Ar.UsingCustomVersion(FFortniteMainBranchObjectVersion::GUID);
+	Super::Serialize(Ar);
+}
+
 void UPhysicalMaterial::PostLoad()
 {
 	Super::PostLoad();
@@ -115,6 +122,12 @@ void UPhysicalMaterial::PostLoad()
 			SurfaceType = PhysicalMaterialProperty_DEPRECATED->ConvertToSurfaceType();
 		}
 	}
+
+	if (GetLinkerCustomVersion(FFortniteMainBranchObjectVersion::GUID) < FFortniteMainBranchObjectVersion::DebugColorForPhysicalMaterials)
+	{
+		// We don't have a color yet, generate one based on a stable hash generated by the asset's path : 
+		DebugColor = FLinearColor::MakeRandomSeededColor(GetTypeHash(GetPathName()));
+	}
 }
 
 void UPhysicalMaterial::FinishDestroy()
diff --git a/Engine/Source/Runtime/PhysicsCore/Public/PhysicalMaterials/PhysicalMaterial.h b/Engine/Source/Runtime/PhysicsCore/Public/PhysicalMaterials/PhysicalMaterial.h
--- a/Engine/Source/Runtime/PhysicsCore/Public/PhysicalMaterials/PhysicalMaterial.h
+++ b/Engine/Source/Runtime/PhysicsCore/Public/PhysicalMaterials/PhysicalMaterial.h
@@ -186,6 +186,10 @@ class UPhysicalMaterial : public UObject
 	UPROPERTY(EditAnywhere, BlueprintReadOnly, Category = PhysicalProperties)
 	FPhysicalMaterialDamageModifier DamageModifier;
 
+	/** Color used to represent this material in certain debug views (e.g. landscape collision) */
+	UPROPERTY(EditAnywhere, BlueprintReadOnly, Category="Debugging", AdvancedDisplay)
+	FLinearColor DebugColor;
+
 	//
 	// Experimental properties
 	//
@@ -225,6 +229,8 @@ public:
 	PHYSICSCORE_API virtual void FinishDestroy() override;
 	//~ End UObject Interface
 
+	virtual void Serialize(FArchive& Ar) override;
+
 	/** Get the physics-interface derived version of this material */
 	PHYSICSCORE_API FPhysicsMaterialHandle& GetPhysicsMaterial();
 
