Landscape - remove trace scope for UpdateComponentState as it's very spammy

#rnx

diff --git a/Engine/Source/Runtime/Landscape/Private/LandscapeGrassMapsBuilder.cpp b/Engine/Source/Runtime/Landscape/Private/LandscapeGrassMapsBuilder.cpp
--- a/Engine/Source/Runtime/Landscape/Private/LandscapeGrassMapsBuilder.cpp
+++ b/Engine/Source/Runtime/Landscape/Private/LandscapeGrassMapsBuilder.cpp
@@ -430,8 +430,6 @@ bool FLandscapeGrassMapsBuilder::UpdateTrackedComponents(const TArray<FVector>&
 	int32 ComponentStateIndex = 0;
 	for (auto It = ComponentStates.CreateIterator(); It; ++It, ++ComponentStateIndex)
 	{
-		TRACE_CPUPROFILER_EVENT_SCOPE(UpdateComponentState);
-
 		FComponentState* State = It.Value();
 		ULandscapeComponent* Component = State->Component;
 		State->TickCount++;
