[Landscape] Make landscape grass weight exporter internally accessible for use outside of the landscape plugin.

#rnx

diff --git a/Engine/Source/Runtime/Landscape/Classes/LandscapeComponent.h b/Engine/Source/Runtime/Landscape/Classes/LandscapeComponent.h
--- a/Engine/Source/Runtime/Landscape/Classes/LandscapeComponent.h
+++ b/Engine/Source/Runtime/Landscape/Classes/LandscapeComponent.h
@@ -929,7 +929,7 @@ public:
 	void RemoveGrassMap();
 
 	/* Could a grassmap currently be generated, disregarding whether our textures are streamed in? */
-	bool CanRenderGrassMap() const;
+	LANDSCAPE_API bool CanRenderGrassMap() const;
 
 #if WITH_EDITOR
 	/** Computes a hash representing the state of the material and grasstypes used by this component. */
diff --git a/Engine/Source/Runtime/Landscape/Internal/LandscapeGrassWeightExporter.h b/Engine/Source/Runtime/Landscape/Internal/LandscapeGrassWeightExporter.h
new file mode 100644
--- /dev/null
+++ b/Engine/Source/Runtime/Landscape/Internal/LandscapeGrassWeightExporter.h
@@ -0,0 +1,138 @@
+// Copyright Epic Games, Inc. All Rights Reserved.
+#pragma once
+
+#include "CoreMinimal.h"
+#include "LandscapeComponent.h"
+#include "MeshMaterialShader.h"
+
+class ALandscapeProxy;
+class FLandscapeAsyncTextureReadback;
+class FLandscapeComponentSceneProxy;
+class FRDGTexture;
+class FTextureRenderTarget2DResource;
+class UTextureRenderTarget2D;
+using FRDGTextureRef = FRDGTexture*;
+
+// Hacky base class to avoid 8 bytes of padding after the vtable
+class FLandscapeGrassWeightExporter_RenderThread_FixLayout
+{
+public:
+	virtual ~FLandscapeGrassWeightExporter_RenderThread_FixLayout() = default;
+};
+
+// data also accessible by render thread
+class FLandscapeGrassWeightExporter_RenderThread : public FLandscapeGrassWeightExporter_RenderThread_FixLayout
+{
+	friend class FLandscapeGrassWeightExporter;
+	friend class FLandscapeGrassMapsBuilder;
+
+private:
+	FLandscapeGrassWeightExporter_RenderThread(const TArray<int32>& InHeightMips, bool bInReadbackToCPU = true);
+
+public:
+	virtual ~FLandscapeGrassWeightExporter_RenderThread();
+
+	LANDSCAPE_API const FIntPoint& GetTargetSize() const { return TargetSize; }
+
+	/** Renders the component to the given texture. */
+	LANDSCAPE_API void RenderLandscapeComponentToTexture_RenderThread(FRDGBuilder& GraphBuilder, FRDGTextureRef OutputTexture);
+
+private:
+	struct FComponentInfo
+	{
+		TObjectPtr<ULandscapeComponent> Component = nullptr;
+		TArray<TObjectPtr<ULandscapeGrassType>> RequestedGrassTypes;
+		FVector2D ViewOffset = FVector2D::ZeroVector;
+		int32 PixelOffsetX = 0;
+		FLandscapeComponentSceneProxy* SceneProxy = nullptr;
+		int32 NumPasses = 0;
+		int32 FirstHeightMipsPassIndex = MAX_int32;
+
+		FComponentInfo(ULandscapeComponent* InComponent, bool bInNeedsGrassmap, bool bInNeedsHeightmap, const TArray<int32>& InHeightMips)
+			: Component(InComponent)
+			, SceneProxy((FLandscapeComponentSceneProxy*)InComponent->SceneProxy)
+		{
+			if (bInNeedsGrassmap)
+			{
+				RequestedGrassTypes = InComponent->GetGrassTypes();
+			}
+			int32 NumGrassMaps = RequestedGrassTypes.Num();
+			if (bInNeedsHeightmap || NumGrassMaps > 0)
+			{
+				// 2 channels for the heightmap, and one channel for each grass map, packed into 4 channel render targets
+				NumPasses += FMath::DivideAndRoundUp(2 /* heightmap */ + NumGrassMaps, 4);
+			}
+
+#if WITH_EDITORONLY_DATA
+			// since we don't read HeightMips unless we are in editor, there's no reason to add passes for it unless we are in editor
+			if (InHeightMips.Num() > 0)
+			{
+				FirstHeightMipsPassIndex = NumPasses;
+				NumPasses += InHeightMips.Num();
+			}
+#endif // WITH_EDITORONLY_DATA
+		}
+	};
+
+	FSceneInterface* SceneInterface = nullptr;
+	TArray<FComponentInfo, TInlineAllocator<1>> ComponentInfos;
+	FIntPoint TargetSize;
+	TArray<int32> HeightMips;
+	float PassOffsetX;
+	FVector ViewOrigin;
+
+	FLandscapeAsyncTextureReadback* AsyncReadbackPtr = nullptr;
+
+	FMatrix ViewRotationMatrix;
+	FMatrix ProjectionMatrix;
+
+	/** Creates a texture and renders the component to it, and then and triggers a readback of the texture. */
+	void RenderLandscapeComponentToTexture_RenderThread(FRHICommandListImmediate& RHICmdList);
+};
+
+class FLandscapeGrassWeightExporter : public FLandscapeGrassWeightExporter_RenderThread
+{
+	friend class FLandscapeGrassMapsBuilder;
+
+private:
+	TObjectPtr<ALandscapeProxy> LandscapeProxy;
+	int32 ComponentSizeVerts;
+	int32 SubsectionSizeQuads;
+	int32 NumSubsections;
+	TArray<TObjectPtr<ULandscapeGrassType>> GrassTypes;
+
+public:
+	LANDSCAPE_API FLandscapeGrassWeightExporter(ALandscapeProxy* InLandscapeProxy, TArrayView<ULandscapeComponent* const> InLandscapeComponents, bool bInNeedsGrassmap = true, bool bInNeedsHeightmap = true, const TArray<int32>& InHeightMips = {}, bool bInRenderImmediately = true, bool bInReadbackToCPU = true);
+
+	// If using the async readback path, check its status and update if needed. Return true when the AsyncReadbackResults are available.
+	// You must call this periodically, or the async readback may not complete.
+	// bInForceFinish will force the RenderThread to wait until GPU completes the readback, ensuring the readback is completed after the render thread executes the command.
+	// NOTE: you may still see false returned, this just means the render thread hasn't executed the command yet.
+	bool CheckAndUpdateAsyncReadback(bool& bOutRenderCommandsQueued, const bool bInForceFinish = false);
+
+	// return true if the async readback is complete.  (Does not update the readback state)
+	bool IsAsyncReadbackComplete();
+
+	// Fetches the results from the GPU texture and translates them into FLandscapeComponentGrassDatas.
+	// If using async readback, requires AsyncReadback to be complete before calling this.
+	// bFreeAsyncReadback if true will call FreeAsyncReadback() to free the readback resource (otherwise you must do it manually)
+	TMap<ULandscapeComponent*, TUniquePtr<FLandscapeComponentGrassData>, TInlineSetAllocator<1>> FetchResults(bool bFreeAsyncReadback);
+
+private:
+	void FreeAsyncReadback();
+
+	// Applies the results using pre-fetched data.
+	static void ApplyResults(TMap<ULandscapeComponent*, TUniquePtr<FLandscapeComponentGrassData>, TInlineSetAllocator<1>>& Results);
+
+	// Fetches the results and applies them to the landscape components
+	// If using async readback, requires AsyncReadback to be complete before calling this.
+	void ApplyResults();
+
+	void CancelAndSelfDestruct();
+};
+
+namespace UE::Landscape::Grass
+{
+	void AddGrassWeightShaderTypes(FMaterialShaderTypes& InOutShaderTypes);
+}
+
diff --git a/Engine/Source/Runtime/Landscape/Private/LandscapeGrassMapsBuilder.cpp b/Engine/Source/Runtime/Landscape/Private/LandscapeGrassMapsBuilder.cpp
--- a/Engine/Source/Runtime/Landscape/Private/LandscapeGrassMapsBuilder.cpp
+++ b/Engine/Source/Runtime/Landscape/Private/LandscapeGrassMapsBuilder.cpp
@@ -3,6 +3,7 @@
 #include "LandscapeGrassMapsBuilder.h"
 #include "Misc/ScopedSlowTask.h"
 #include "Landscape.h"
+#include "LandscapeAsyncTextureReadback.h"
 #include "LandscapePrivate.h"
 #include "LandscapeProxy.h"
 #include "LandscapeGrassWeightExporter.h"
diff --git a/Engine/Source/Runtime/Landscape/Private/LandscapeGrassWeightExporter.cpp b/Engine/Source/Runtime/Landscape/Private/LandscapeGrassWeightExporter.cpp
--- a/Engine/Source/Runtime/Landscape/Private/LandscapeGrassWeightExporter.cpp
+++ b/Engine/Source/Runtime/Landscape/Private/LandscapeGrassWeightExporter.cpp
@@ -15,6 +15,7 @@
 #include "RenderCaptureInterface.h"
 #include "ShaderPlatformCachedIniValue.h"
 #include "LandscapePrivate.h"
+#include "LandscapeAsyncTextureReadback.h"
 
 #if UE_BUILD_DEBUG
 #include "Misc/FileHelper.h"
@@ -394,10 +395,49 @@ IPSOCollector* CreateLandscapeGrassWeightPSOCollector(ERHIFeatureLevel::Type Fea
 }
 FRegisterPSOCollectorCreateFunction RegisterLandscapeGrassWeightPSOCollector(&CreateLandscapeGrassWeightPSOCollector, EShadingPath::Deferred, LandscapeGrassWeightMeshPassName);
 
+FLandscapeGrassWeightExporter_RenderThread::FLandscapeGrassWeightExporter_RenderThread(const TArray<int32>& InHeightMips, bool bInReadbackToCPU)
+	: HeightMips(InHeightMips)
+{
+	if (bInReadbackToCPU)
+	{
+		// even when doing a synchronous readback, we use the async readback structure
+		AsyncReadbackPtr = new FLandscapeAsyncTextureReadback();
+	}
+}
+
+FLandscapeGrassWeightExporter_RenderThread::~FLandscapeGrassWeightExporter_RenderThread()
+{
+	if (AsyncReadbackPtr != nullptr)
+	{
+		AsyncReadbackPtr->QueueDeletionFromGameThread();
+		AsyncReadbackPtr = nullptr;
+	}
+}
+
 void FLandscapeGrassWeightExporter_RenderThread::RenderLandscapeComponentToTexture_RenderThread(FRHICommandListImmediate& RHICmdList)
 {
 	FRDGBuilder GraphBuilder(RHICmdList);
 
+	FRDGTextureDesc TextureDesc = FRDGTextureDesc::Create2D(
+		TargetSize,
+		PF_B8G8R8A8,
+		FClearValueBinding(),
+		ETextureCreateFlags::RenderTargetable);
+
+	FRDGTextureRef OutputTexture = GraphBuilder.CreateTexture(TextureDesc, TEXT("LandscapeGrassMapRenderTarget"), ERDGTextureFlags::None);
+
+	RenderLandscapeComponentToTexture_RenderThread(GraphBuilder, OutputTexture);
+
+	if (AsyncReadbackPtr != nullptr)
+	{
+		AsyncReadbackPtr->StartReadback_RenderThread(GraphBuilder, OutputTexture);
+	}
+
+	GraphBuilder.Execute();
+}
+
+void FLandscapeGrassWeightExporter_RenderThread::RenderLandscapeComponentToTexture_RenderThread(FRDGBuilder& GraphBuilder, FRDGTextureRef OutputTexture)
+{
 	FSceneViewFamilyContext ViewFamily(FSceneViewFamily::ConstructionValues(nullptr, SceneInterface, FEngineShowFlags(ESFIM_Game))
 		.SetTime(FGameTime::GetTimeSinceAppStart()));
 
@@ -413,18 +453,9 @@ void FLandscapeGrassWeightExporter_RenderThread::RenderLandscapeComponentToTextu
 	ViewInitOptions.ProjectionMatrix = ProjectionMatrix;
 	ViewInitOptions.ViewFamily = &ViewFamily;
 
-	GetRendererModule().CreateAndInitSingleView(RHICmdList, &ViewFamily, &ViewInitOptions);
+	GetRendererModule().CreateAndInitSingleView(GraphBuilder.RHICmdList, &ViewFamily, &ViewInitOptions);
 
 	const FSceneView* View = ViewFamily.Views[0];
-
-	FRDGTextureDesc TextureDesc = FRDGTextureDesc::Create2D(
-		TargetSize,
-		PF_B8G8R8A8,
-		FClearValueBinding(),
-		ETextureCreateFlags::RenderTargetable);
-
-	FRDGTextureRef OutputTexture = GraphBuilder.CreateTexture(TextureDesc, TEXT("LandscapeGrassMapRenderTarget"), ERDGTextureFlags::None);
-
 	auto* PassParameters = GraphBuilder.AllocParameters<FLandscapeGrassPassParameters>();
 	PassParameters->View = View->ViewUniformBuffer;
 	PassParameters->Scene = GetSceneUniformBufferRef(GraphBuilder, *View);
@@ -452,18 +483,10 @@ void FLandscapeGrassWeightExporter_RenderThread::RenderLandscapeComponentToTextu
 				}
 			}
 		});
-
-	if (AsyncReadbackPtr != nullptr)
-	{
-		AsyncReadbackPtr->StartReadback_RenderThread(GraphBuilder, OutputTexture);
-	}
-
-	GraphBuilder.Execute();
 }
 
-
-FLandscapeGrassWeightExporter::FLandscapeGrassWeightExporter(ALandscapeProxy* InLandscapeProxy, TArrayView<ULandscapeComponent* const> InLandscapeComponents, bool bInNeedsGrassmap, bool bInNeedsHeightmap, const TArray<int32>& InHeightMips)
-	: FLandscapeGrassWeightExporter_RenderThread(InHeightMips)
+FLandscapeGrassWeightExporter::FLandscapeGrassWeightExporter(ALandscapeProxy* InLandscapeProxy, TArrayView<ULandscapeComponent* const> InLandscapeComponents, bool bInNeedsGrassmap, bool bInNeedsHeightmap, const TArray<int32>& InHeightMips, bool bInRenderImmediately, bool bInReadbackToCPU)
+	: FLandscapeGrassWeightExporter_RenderThread(InHeightMips, bInReadbackToCPU)
 	, LandscapeProxy(InLandscapeProxy)
 	, ComponentSizeVerts(InLandscapeProxy->ComponentSizeQuads + 1)
 	, SubsectionSizeQuads(InLandscapeProxy->SubsectionSizeQuads)
@@ -523,21 +546,35 @@ FLandscapeGrassWeightExporter::FLandscapeGrassWeightExporter(ALandscapeProxy* In
 		0.5f / ZOffset,
 		ZOffset);
 
-	UE::RenderCommandPipe::FSyncScope SyncScope;
+	if (bInRenderImmediately)
+	{
+		UE::RenderCommandPipe::FSyncScope SyncScope;
 
-	RenderCaptureInterface::FScopedCapture RenderCapture((GRenderCaptureNextGrassmapDraws != 0), TEXT("LandscapeGrassmapCapture"));
-	GRenderCaptureNextGrassmapDraws = FMath::Max(0, GRenderCaptureNextGrassmapDraws - 1);
+		RenderCaptureInterface::FScopedCapture RenderCapture((GRenderCaptureNextGrassmapDraws != 0), TEXT("LandscapeGrassmapCapture"));
+		GRenderCaptureNextGrassmapDraws = FMath::Max(0, GRenderCaptureNextGrassmapDraws - 1);
 
-	// render
-	FLandscapeGrassWeightExporter_RenderThread* Exporter = this;
-	ENQUEUE_RENDER_COMMAND(FDrawSceneCommand)(
-		[Exporter](FRHICommandListImmediate& RHICmdList)
-		{
-			Exporter->RenderLandscapeComponentToTexture_RenderThread(RHICmdList);
-			FlushPendingDeleteRHIResources_RenderThread();
-		});
+		// render
+		FLandscapeGrassWeightExporter_RenderThread* Exporter = this;
+		ENQUEUE_RENDER_COMMAND(FDrawSceneCommand)(
+			[Exporter](FRHICommandListImmediate& RHICmdList)
+			{
+				Exporter->RenderLandscapeComponentToTexture_RenderThread(RHICmdList);
+				FlushPendingDeleteRHIResources_RenderThread();
+			});
+	}
 }
 
+bool FLandscapeGrassWeightExporter::CheckAndUpdateAsyncReadback(bool& bOutRenderCommandsQueued, const bool bInForceFinish)
+{
+	check(AsyncReadbackPtr != nullptr);
+	return AsyncReadbackPtr->CheckAndUpdate(bOutRenderCommandsQueued, bInForceFinish);
+}
+
+bool FLandscapeGrassWeightExporter::IsAsyncReadbackComplete()
+{
+	check(AsyncReadbackPtr != nullptr);
+	return AsyncReadbackPtr->IsComplete();
+}
 
 struct FByteBuffer2DView : public IBuffer2DView<uint8>
 {
diff --git a/Engine/Source/Runtime/Landscape/Private/LandscapeGrassWeightExporter.h b/Engine/Source/Runtime/Landscape/Private/LandscapeGrassWeightExporter.h
deleted file mode 100644
