CIS fix.

#rnx

diff --git a/Engine/Source/Runtime/Landscape/Private/LandscapeEditLayers.cpp b/Engine/Source/Runtime/Landscape/Private/LandscapeEditLayers.cpp
--- a/Engine/Source/Runtime/Landscape/Private/LandscapeEditLayers.cpp
+++ b/Engine/Source/Runtime/Landscape/Private/LandscapeEditLayers.cpp
@@ -6895,7 +6895,7 @@ UE::Landscape::EditLayers::FMergeRenderContext ALandscape::PrepareEditLayersMerg
 						// TODO [jonathan.bard] : Change this to OOBB to OOBB test to cull more components : look at FOrientedBoxHelpers
 						const FOOBox2D& OOBox = OutputWorldArea.GetOOBox();
 						Info->GetOverlappedComponents(OOBox.Transform, FBox2D(-OOBox.Extents / 2.0, OOBox.Extents / 2.0), DependentComponents, ComponentIndicesBoundingRect);
-						for (auto It = DependentComponents.CreateConstIterator(); It; ++It && !bOutputRenderItem)
+						for (auto It = DependentComponents.CreateConstIterator(); It && !bOutputRenderItem; ++It)
 						{
 							int32 ComponentIndex = Component2DIndexer.GetValueIndexForKeyChecked(It.Key());
 							// Retain the render item if one of the components it renders to overlaps with one the components to render :
