[Landscape] Move ULandscapeComponent::CanRenderGrassMap outside of WITH_EDITOR guards to match definition in header.

#rnx

diff --git a/Engine/Source/Runtime/Landscape/Private/LandscapeGrass.cpp b/Engine/Source/Runtime/Landscape/Private/LandscapeGrass.cpp
--- a/Engine/Source/Runtime/Landscape/Private/LandscapeGrass.cpp
+++ b/Engine/Source/Runtime/Landscape/Private/LandscapeGrass.cpp
@@ -907,24 +907,6 @@ bool ULandscapeComponent::UpdateGrassTypes(bool bForceUpdate)
 	return bChanged;
 }
 
-#if WITH_EDITOR
-
-namespace UE::Landscape
-{
-	extern uint32 ComputeGrassMapGenerationHash(const ULandscapeComponent* Component, UMaterialInterface* Material);
-	extern void SubmitGPUCommands(bool bBlockUntilRTComplete, bool bBlockRTUntilGPUComplete);
-}
-
-uint32 ULandscapeComponent::ComputeGrassMapGenerationHash() const
-{
-	return UE::Landscape::ComputeGrassMapGenerationHash(this, GetLandscapeMaterial());
-}
-
-bool ULandscapeComponent::IsGrassMapOutdated() const
-{
-	return GrassData->HasValidData() && (ComputeGrassMapGenerationHash() != GrassData->GenerationHash);
-}
-
 bool ULandscapeComponent::CanRenderGrassMap() const
 {
 	// Check we can render
@@ -956,6 +938,24 @@ bool ULandscapeComponent::CanRenderGrassMap() const
 	return true;
 }
 
+#if WITH_EDITOR
+
+namespace UE::Landscape
+{
+	extern uint32 ComputeGrassMapGenerationHash(const ULandscapeComponent* Component, UMaterialInterface* Material);
+	extern void SubmitGPUCommands(bool bBlockUntilRTComplete, bool bBlockRTUntilGPUComplete);
+}
+
+uint32 ULandscapeComponent::ComputeGrassMapGenerationHash() const
+{
+	return UE::Landscape::ComputeGrassMapGenerationHash(this, GetLandscapeMaterial());
+}
+
+bool ULandscapeComponent::IsGrassMapOutdated() const
+{
+	return GrassData->HasValidData() && (ComputeGrassMapGenerationHash() != GrassData->GenerationHash);
+}
+
 TArray<uint16> ULandscapeComponent::RenderWPOHeightmap(int32 LOD)
 {
 	TArray<uint16> Results;
@@ -1809,7 +1809,7 @@ struct FLandscapeComponentGrassAccess
 	{
 		return LandscapeDataAccess::GetLocalHeight(HeightData[IdxX + Stride*IdxY]);
 	}
-	FORCEINLINE float GetWeight(int32 IdxX, int32 IdxY)
+	FORCEINLINE float GetWeight(int32 IdxX, int32 IdxY) // usages
 	{
 		return ((float)WeightData[IdxX + Stride*IdxY]) / 255.f;
 	}
