Small optimization for landscape sculpting.  Faster version of generating the material index array for the UseDefaultMaterialOnly case in WriteRuntimeData.  Saves 1 ms per frame in some cases.
Some refactoring to prepare for another, slightly more complicated, optimization.

diff --git a/Engine/Source/Runtime/Landscape/Classes/LandscapeHeightfieldCollisionComponent.h b/Engine/Source/Runtime/Landscape/Classes/LandscapeHeightfieldCollisionComponent.h
--- a/Engine/Source/Runtime/Landscape/Classes/LandscapeHeightfieldCollisionComponent.h
+++ b/Engine/Source/Runtime/Landscape/Classes/LandscapeHeightfieldCollisionComponent.h
@@ -261,6 +261,15 @@ public:
 	virtual uint32 ComputeCollisionHash() const;
 #endif
 
+	struct FCollisionSampleInfo
+	{
+		int32 CollisionSizeVerts = 0;
+		int32 SimpleCollisionSizeVerts = 0;
+		int32 NumSamples = 0;
+		int32 NumSimpleSamples = 0;
+	};
+	FCollisionSampleInfo GetCollisionSampleInfo() const;
+
 	struct FWriteRuntimeDataParams
 	{
 		bool bUseDefaultMaterialOnly = false;
@@ -277,11 +286,16 @@ public:
 		int32 VisibilityLayerIndex = INDEX_NONE;
 	};
 
-	void GetCollisionSampleInfo(int32& OutCollisionSizeVerts, int32& OutSimpleCollisionSizeVerts, int32& OutNumSamples, int32& OutNumSimpleSamples) const;
+#if WITH_EDITOR
+	// Set up params to call WriteRuntimeData using editor data members.
+	FWriteRuntimeDataParams MakeWriteRuntimeDataParams(bool bUseDefaultMaterialOnly) const;
+#endif
 
 	// Writes to a cooked data buffer using raw heightfield data
 	LANDSCAPE_API bool WriteRuntimeData(const FWriteRuntimeDataParams& Params, TArray<uint8>& OutHeightfieldData, TArray<UPhysicalMaterial*>& InOutMaterials) const;
 
+	bool GenerateCollisionData(const FWriteRuntimeDataParams& Params, Chaos::FHeightFieldPtr& OutHeightField, bool bGenerateSimpleCollision, Chaos::FHeightFieldPtr& OutSimpleHeightField, TArray<UPhysicalMaterial*>& InOutMaterials) const;
+
 	/** Gets the landscape info object for this landscape */
 	ULandscapeInfo* GetLandscapeInfo() const;
 
@@ -293,7 +307,7 @@ public:
 	/** Creates collision object from a cooked collision data */
 	virtual void CreateCollisionObject();
 
-	/** Creates collision object from raw runtime data.  Data is assumed at this point to contain valid physical material indices with visiblility layer set to 0xFF*/
+	/** Creates collision object from raw runtime data.  Data is assumed at this point to contain valid physical material indices with visibility layer set to 0xFF*/
 	LANDSCAPE_API void CreateCollisionObject(
 		bool bUseDefaultMaterialOnly,
 		TArrayView<const uint16> Heights, TArrayView<const uint16> SimpleHeights,
diff --git a/Engine/Source/Runtime/Landscape/Private/LandscapeCollision.cpp b/Engine/Source/Runtime/Landscape/Private/LandscapeCollision.cpp
--- a/Engine/Source/Runtime/Landscape/Private/LandscapeCollision.cpp
+++ b/Engine/Source/Runtime/Landscape/Private/LandscapeCollision.cpp
@@ -935,6 +935,7 @@ void ULandscapeHeightfieldCollisionComponent::CreateCollisionObject()
 
 				// Create heightfields
 				{
+					TRACE_CPUPROFILER_EVENT_SCOPE(CreateCollisionObject_ChaosStream);
 					FMemoryReader Reader(CookedCollisionData);
 					Chaos::FChaosArchive Ar(Reader);
 					bool bContainsSimple = false;
@@ -959,7 +960,7 @@ void ULandscapeHeightfieldCollisionComponent::CreateCollisionObject()
 				}
 			}
 
-			// Release cooked collison data
+			// Release cooked collision data
 			// In cooked builds created collision object will never be deleted while component is alive, so we don't need this data anymore
 			if(FPlatformProperties::RequiresCookedData() || (World && World->IsGameWorld()))
 			{
@@ -974,6 +975,7 @@ void ULandscapeHeightfieldCollisionComponent::CreateCollisionObject()
 				TArray<UPhysicalMaterial*> CookedMaterialsEd;
 				if(CookCollisionData(PhysicsFormatName, true, bCheckDDC, CookedCollisionDataEd, CookedMaterialsEd))
 				{
+					TRACE_CPUPROFILER_EVENT_SCOPE(CreateCollisionObject_ChaosStream);
 					FMemoryReader Reader(CookedCollisionDataEd);
 					Chaos::FChaosArchive Ar(Reader);
 
@@ -1005,19 +1007,15 @@ void ULandscapeHeightfieldCollisionComponent::CreateCollisionObject(
 	{
 		return;
 	}
+	 
+	const FCollisionSampleInfo Info = GetCollisionSampleInfo();
 
-	int32 CollisionSizeVerts;
-	int32 SimpleCollisionSizeVerts;
-	int32 NumSamples;
-	int32 NumSimpleSamples;
-	GetCollisionSampleInfo(CollisionSizeVerts, SimpleCollisionSizeVerts, NumSamples, NumSimpleSamples);
-
-	if (!ensure(Heights.Num() == NumSamples))
+	if (!ensure(Heights.Num() == Info.NumSamples))
 	{
 		return;
 	}
 
-	int32 NumQuads = (CollisionSizeVerts - 1) * (CollisionSizeVerts - 1);
+	int32 NumQuads = (Info.CollisionSizeVerts - 1) * (Info.CollisionSizeVerts - 1);
 	if (!ensure(PhysicalMaterialIds.Num() == NumQuads))
 	{
 		return;
@@ -1025,12 +1023,12 @@ void ULandscapeHeightfieldCollisionComponent::CreateCollisionObject(
 
 	if(bGenerateSimpleCollision)
 	{
-		if (!ensure(SimpleHeights.Num() == NumSimpleSamples))
+		if (!ensure(SimpleHeights.Num() == Info.NumSimpleSamples))
 		{
 			return;
 		}
 
-		int32 NumSimpleQuads = (SimpleCollisionSizeVerts - 1) * (SimpleCollisionSizeVerts - 1);
+		int32 NumSimpleQuads = (Info.SimpleCollisionSizeVerts - 1) * (Info.SimpleCollisionSizeVerts - 1);
 		if (!ensure(SimplePhysicalMaterialIds.Num() == NumSimpleQuads))
 		{
 			return;
@@ -1065,20 +1063,20 @@ void ULandscapeHeightfieldCollisionComponent::CreateCollisionObject(
 	HeightfieldGuid = FGuid::NewDeterministicGuid(GetPathName(), CollisionHash);
 
 	HeightfieldRef = GSharedHeightfieldRefs.Add(HeightfieldGuid, new FHeightfieldGeometryRef(HeightfieldGuid));
-	HeightfieldRef->HeightfieldGeometry = Chaos::FHeightFieldPtr(new Chaos::FHeightField(Heights, PhysicalMaterialIds, CollisionSizeVerts, CollisionSizeVerts, Chaos::FVec3(1)));
+	HeightfieldRef->HeightfieldGeometry = Chaos::FHeightFieldPtr(new Chaos::FHeightField(Heights, PhysicalMaterialIds, Info.CollisionSizeVerts, Info.CollisionSizeVerts, Chaos::FVec3(1)));
 	
 #if WITH_EDITOR
 	UWorld* World = GetWorld();
 	const bool bNeedsEditorHeightField = World && !World->IsGameWorld() && !GetOutermost()->bIsCookedForEditor;
 	if (bNeedsEditorHeightField)
 	{
-		HeightfieldRef->EditorHeightfieldGeometry = Chaos::FHeightFieldPtr(new Chaos::FHeightField(Heights, PhysicalMaterialIds, CollisionSizeVerts, CollisionSizeVerts, Chaos::FVec3(1)));
+		HeightfieldRef->EditorHeightfieldGeometry = Chaos::FHeightFieldPtr(new Chaos::FHeightField(Heights, PhysicalMaterialIds, Info.CollisionSizeVerts, Info.CollisionSizeVerts, Chaos::FVec3(1)));
 	}
 #endif // WITH_EDITOR
 
 	if (bGenerateSimpleCollision)
 	{
-		HeightfieldRef->HeightfieldSimpleGeometry = Chaos::FHeightFieldPtr(new Chaos::FHeightField(SimpleHeights, SimplePhysicalMaterialIds, SimpleCollisionSizeVerts, SimpleCollisionSizeVerts, Chaos::FVec3(1)));
+		HeightfieldRef->HeightfieldSimpleGeometry = Chaos::FHeightFieldPtr(new Chaos::FHeightField(SimpleHeights, SimplePhysicalMaterialIds, Info.SimpleCollisionSizeVerts, Info.SimpleCollisionSizeVerts, Chaos::FVec3(1)));
 	}
 
 	for (UPhysicalMaterial* PhysicalMaterial : PhysicalMaterialObjects)
@@ -1106,20 +1104,76 @@ void ULandscapeHeightfieldCollisionComponent::SpeculativelyLoadAsyncDDCCollsionD
 	}
 }
 
+ULandscapeHeightfieldCollisionComponent::FWriteRuntimeDataParams ULandscapeHeightfieldCollisionComponent::MakeWriteRuntimeDataParams(bool bUseDefaultMaterialOnly) const
+{
+	const FCollisionSampleInfo Info = GetCollisionSampleInfo();
+
+	const uint16* Heights = (const uint16*)CollisionHeightData.LockReadOnly();
+	check(CollisionHeightData.GetElementCount() == Info.NumSamples + Info.NumSimpleSamples);
+	const uint16* SimpleHeights = Heights + Info.NumSamples;
+
+	// Physical material data from layer system
+	const uint8* DominantLayers = nullptr;
+	const uint8* SimpleDominantLayers = nullptr;
+	if (DominantLayerData.GetElementCount())
+	{
+		DominantLayers = (const uint8*)DominantLayerData.LockReadOnly();
+		check(DominantLayerData.GetElementCount() == Info.NumSamples + Info.NumSimpleSamples);
+		SimpleDominantLayers = DominantLayers + Info.NumSamples;
+	}
+
+	// Physical material data from render material graph
+	const uint8* RenderPhysicalMaterialIds = nullptr;
+	const uint8* SimpleRenderPhysicalMaterialIds = nullptr;
+	if (PhysicalMaterialRenderData.GetElementCount())
+	{
+		RenderPhysicalMaterialIds = (const uint8*)PhysicalMaterialRenderData.LockReadOnly();
+		check(PhysicalMaterialRenderData.GetElementCount() == Info.NumSamples + Info.NumSimpleSamples);
+		SimpleRenderPhysicalMaterialIds = RenderPhysicalMaterialIds + Info.NumSamples;
+	}
+
+	auto MakeSafeArrayView =
+		[](auto* Data, int32 DataSamples)
+		{
+			return Data ? MakeArrayView(Data, DataSamples) : MakeArrayView(Data, 0);
+		};
+
+	FWriteRuntimeDataParams WriteParams;
+	WriteParams.bUseDefaultMaterialOnly = bUseDefaultMaterialOnly;
+	WriteParams.bProcessRenderIndices = true;
+	WriteParams.bProcessVisibilityLayer = true;
+	WriteParams.Heights = MakeSafeArrayView(Heights, Info.NumSamples);
+	WriteParams.SimpleHeights = MakeSafeArrayView(SimpleHeights, Info.NumSimpleSamples);
+	WriteParams.DominantLayers = MakeSafeArrayView(DominantLayers, Info.NumSamples);
+	WriteParams.SimpleDominantLayers = MakeSafeArrayView(SimpleDominantLayers, Info.NumSimpleSamples);
+	WriteParams.RenderPhysicalMaterialIds = MakeSafeArrayView(RenderPhysicalMaterialIds, Info.NumSamples);
+	WriteParams.SimpleRenderPhysicalMaterialIds = MakeSafeArrayView(SimpleRenderPhysicalMaterialIds, Info.NumSimpleSamples);
+	WriteParams.PhysicalMaterialRenderObjects = MakeSafeArrayView(PhysicalMaterialRenderObjects.GetData(), PhysicalMaterialRenderObjects.Num());
+	WriteParams.ComponentLayerInfos = MakeSafeArrayView(ComponentLayerInfos.GetData(), ComponentLayerInfos.Num());
+	WriteParams.VisibilityLayerIndex = ComponentLayerInfos.IndexOfByKey(ALandscapeProxy::VisibilityLayer);
+
+	return WriteParams;
+}
+
 #endif
 
-void ULandscapeHeightfieldCollisionComponent::GetCollisionSampleInfo(int32& OutCollisionSizeVerts, int32& OutSimpleCollisionSizeVerts, int32& OutNumSamples, int32& OutNumSimpleSamples) const
+ULandscapeHeightfieldCollisionComponent::FCollisionSampleInfo ULandscapeHeightfieldCollisionComponent::GetCollisionSampleInfo() const
 {
-	OutCollisionSizeVerts = CollisionSizeQuads + 1;
-	OutSimpleCollisionSizeVerts = SimpleCollisionSizeQuads > 0 ? SimpleCollisionSizeQuads + 1 : 0;
-	OutNumSamples = FMath::Square(OutCollisionSizeVerts);
-	OutNumSimpleSamples = FMath::Square(OutSimpleCollisionSizeVerts);
+	FCollisionSampleInfo Result;
+	Result.CollisionSizeVerts = CollisionSizeQuads + 1;
+	Result.SimpleCollisionSizeVerts = SimpleCollisionSizeQuads > 0 ? SimpleCollisionSizeQuads + 1 : 0;
+	Result.NumSamples = FMath::Square(Result.CollisionSizeVerts);
+	Result.NumSimpleSamples = FMath::Square(Result.SimpleCollisionSizeVerts);
+	return Result;
 }
 
 
-// Writes to a cooked data buffer using raw heightfield data
-bool ULandscapeHeightfieldCollisionComponent::WriteRuntimeData(const FWriteRuntimeDataParams& Params, TArray<uint8>& OutHeightfieldData, TArray<UPhysicalMaterial*>& InOutMaterials) const
+// Generate the heightfield and optional simple heightfield objects
+// The dominant materials for the collision object will be added to InOutMaterials if Params.bUseDefaultMaterialOnly is false.
+bool ULandscapeHeightfieldCollisionComponent::GenerateCollisionData(const FWriteRuntimeDataParams& Params, Chaos::FHeightFieldPtr& OutHeightField, bool bGenerateSimpleCollision, Chaos::FHeightFieldPtr& OutSimpleHeightField, TArray<UPhysicalMaterial*>& InOutMaterials) const
 {
+	TRACE_CPUPROFILER_EVENT_SCOPE(ULandscapeHeightfieldCollisionComponent::GenerateCollisionData);
+
 	ALandscapeProxy* Proxy = GetLandscapeProxy();
 	if (!Proxy || !Proxy->GetRootComponent())
 	{
@@ -1132,88 +1186,120 @@ bool ULandscapeHeightfieldCollisionComponent::WriteRuntimeData(const FWriteRunti
 	const FVector LandscapeScale = Proxy->GetRootComponent()->GetRelativeScale3D();
 	const bool bIsMirrored = (LandscapeScale.X * LandscapeScale.Y * LandscapeScale.Z) < 0.f;
 
-	const bool bGenerateSimpleCollision = SimpleCollisionSizeQuads > 0 && !Params.bUseDefaultMaterialOnly;
-
-	int32 CollisionSizeVerts;
-	int32 SimpleCollisionSizeVerts;
-	int32 NumSamples;
-	int32 NumSimpleSamples;
-	GetCollisionSampleInfo(CollisionSizeVerts, SimpleCollisionSizeVerts, NumSamples, NumSimpleSamples);
-
-	// List of materials which is actually used by heightfield
-	InOutMaterials.Empty();
+	const FCollisionSampleInfo Info = GetCollisionSampleInfo();
 
 	// Generate material indices
 	TArray<uint8> MaterialIndices;
-	MaterialIndices.Reserve(NumSamples + NumSimpleSamples);
 
-	auto ResolveMaterials = [&MaterialIndices, &bIsMirrored, &Params, &DefMaterial, &InOutMaterials](int32 InCollisionVertExtent, TArrayView<const uint8> InDominantLayers, TArrayView<const uint8> InRenderMaterialIds)
+	if (!Params.bUseDefaultMaterialOnly)
 	{
-		for (int32 RowIndex = 0; RowIndex < InCollisionVertExtent; RowIndex++)
-		{
-			for (int32 ColIndex = 0; ColIndex < InCollisionVertExtent; ColIndex++)
-			{
-				const int32 SrcSampleIndex = (RowIndex * InCollisionVertExtent) + (bIsMirrored ? (InCollisionVertExtent - ColIndex - 1) : ColIndex);
+		// List of materials which is actually used by heightfield
+		InOutMaterials.Empty();
 
-				// Materials are not relevant on the last row/column because they are per-triangle and the last row/column don't own any
-				if (RowIndex < InCollisionVertExtent - 1 &&
-					ColIndex < InCollisionVertExtent - 1)
+		MaterialIndices.Reserve(Info.NumSamples + Info.NumSimpleSamples);
+
+		auto ResolveMaterials = [&MaterialIndices, &bIsMirrored, &Params, &DefMaterial, &InOutMaterials](int32 InCollisionVertExtent, TArrayView<const uint8> InDominantLayers, TArrayView<const uint8> InRenderMaterialIds)
+			{
+				check(!Params.bUseDefaultMaterialOnly);
+				for (int32 RowIndex = 0; RowIndex < InCollisionVertExtent; RowIndex++)
 				{
-					int32 MaterialIndex = 0; // Default physical material.
-					if (!Params.bUseDefaultMaterialOnly)
+					for (int32 ColIndex = 0; ColIndex < InCollisionVertExtent; ColIndex++)
 					{
-						uint8 DominantLayerIdx = InDominantLayers.IsEmpty() ? -1 : InDominantLayers[SrcSampleIndex];
-						ULandscapeLayerInfoObject* Layer = Params.ComponentLayerInfos.IsValidIndex(DominantLayerIdx) ? ToRawPtr(Params.ComponentLayerInfos[DominantLayerIdx]) : nullptr;
+						const int32 SrcSampleIndex = (RowIndex * InCollisionVertExtent) + (bIsMirrored ? (InCollisionVertExtent - ColIndex - 1) : ColIndex);
 
-						if (Params.bProcessVisibilityLayer && DominantLayerIdx == Params.VisibilityLayerIndex)
-						{
-							// If it's a hole, use the final index
-							MaterialIndex = TNumericLimits<uint8>::Max();
-						}
-						else if (Params.bProcessRenderIndices && !InRenderMaterialIds.IsEmpty())
+						// Materials are not relevant on the last row/column because they are per-triangle and the last row/column don't own any
+						if (RowIndex < InCollisionVertExtent - 1 &&
+							ColIndex < InCollisionVertExtent - 1)
 						{
-							uint8 RenderIdx = InRenderMaterialIds[SrcSampleIndex];
-							UPhysicalMaterial* DominantMaterial = RenderIdx > 0 ? ToRawPtr(Params.PhysicalMaterialRenderObjects[RenderIdx - 1]) : DefMaterial;
-							MaterialIndex = InOutMaterials.AddUnique(DominantMaterial);
-						}
-						else
-						{
-							UPhysicalMaterial* DominantMaterial = Layer && Layer->PhysMaterial ? ToRawPtr(Layer->PhysMaterial) : DefMaterial;
-							MaterialIndex = InOutMaterials.AddUnique(DominantMaterial);
+							int32 MaterialIndex = 0; // Default physical material.
+							uint8 DominantLayerIdx = InDominantLayers.IsEmpty() ? -1 : InDominantLayers[SrcSampleIndex];
+							ULandscapeLayerInfoObject* Layer = Params.ComponentLayerInfos.IsValidIndex(DominantLayerIdx) ? ToRawPtr(Params.ComponentLayerInfos[DominantLayerIdx]) : nullptr;
+
+							if (Params.bProcessVisibilityLayer && DominantLayerIdx == Params.VisibilityLayerIndex)
+							{
+								// If it's a hole, use the final index
+								MaterialIndex = TNumericLimits<uint8>::Max();
+							}
+							else if (Params.bProcessRenderIndices && !InRenderMaterialIds.IsEmpty())
+							{
+								uint8 RenderIdx = InRenderMaterialIds[SrcSampleIndex];
+								UPhysicalMaterial* DominantMaterial = RenderIdx > 0 ? ToRawPtr(Params.PhysicalMaterialRenderObjects[RenderIdx - 1]) : DefMaterial;
+								MaterialIndex = InOutMaterials.AddUnique(DominantMaterial);
+							}
+							else
+							{
+								UPhysicalMaterial* DominantMaterial = Layer && Layer->PhysMaterial ? ToRawPtr(Layer->PhysMaterial) : DefMaterial;
+								MaterialIndex = InOutMaterials.AddUnique(DominantMaterial);
+							}
+							MaterialIndices.Add(IntCastChecked<uint8>(MaterialIndex));
 						}
 					}
-					MaterialIndices.Add(IntCastChecked<uint8>(MaterialIndex));
 				}
-			}
+			};
+
+		{
+			TRACE_CPUPROFILER_EVENT_SCOPE(ResolveMaterials);
+			ResolveMaterials(Info.CollisionSizeVerts, Params.DominantLayers, Params.RenderPhysicalMaterialIds);
+			ResolveMaterials(Info.SimpleCollisionSizeVerts, Params.SimpleDominantLayers, Params.SimpleRenderPhysicalMaterialIds);
 		}
-	};
+	}
+	else
+	{
+		// for bUseDefaultMaterialOnly==true, much faster to just make the array of zeros directly
+		int32 NumMatIndices = FMath::Square(CollisionSizeQuads);
+		if (bGenerateSimpleCollision)
+		{
+			NumMatIndices += FMath::Square(SimpleCollisionSizeQuads);
+		}
+		MaterialIndices.SetNumZeroed(NumMatIndices);
+	}
 
-	ResolveMaterials(CollisionSizeVerts, Params.DominantLayers, Params.RenderPhysicalMaterialIds);
-	ResolveMaterials(SimpleCollisionSizeVerts, Params.SimpleDominantLayers, Params.SimpleRenderPhysicalMaterialIds);
+	{
+		TRACE_CPUPROFILER_EVENT_SCOPE(CreateHeightField);
+		const int32 NumCollisionCells = FMath::Square(CollisionSizeQuads);
+		TArrayView<uint8> ComplexMaterialIndicesView(MaterialIndices.GetData(), NumCollisionCells);
+		OutHeightField = Chaos::FHeightFieldPtr(new Chaos::FHeightField(Params.Heights, ComplexMaterialIndicesView, Info.CollisionSizeVerts, Info.CollisionSizeVerts, Chaos::FVec3(1)));
+
+		if (bGenerateSimpleCollision)
+		{
+			const int32 NumSimpleCollisionCells = FMath::Square(SimpleCollisionSizeQuads);
+			TArrayView<uint8> SimpleMaterialIndicesView(MaterialIndices.GetData() + NumCollisionCells, NumSimpleCollisionCells);
+			OutSimpleHeightField = Chaos::FHeightFieldPtr(new Chaos::FHeightField(Params.SimpleHeights, SimpleMaterialIndicesView, Info.SimpleCollisionSizeVerts, Info.SimpleCollisionSizeVerts, Chaos::FVec3(1)));
+		}
+	}
+
+	return true;
+}
+
+// Generate the heightfield and optional simple heightfield objects and serialize them into a byte array.
+bool ULandscapeHeightfieldCollisionComponent::WriteRuntimeData(const FWriteRuntimeDataParams& Params, TArray<uint8>& OutHeightfieldData, TArray<UPhysicalMaterial*>& InOutMaterials) const
+{
+	TRACE_CPUPROFILER_EVENT_SCOPE(ULandscapeHeightfieldCollisionComponent::WriteRuntimeData);
 
 	Chaos::FHeightFieldPtr Heightfield = nullptr;
 	Chaos::FHeightFieldPtr HeightfieldSimple = nullptr;
+	const bool bGenerateSimpleCollision = SimpleCollisionSizeQuads > 0 && !Params.bUseDefaultMaterialOnly;
+	if (GenerateCollisionData(Params, Heightfield, bGenerateSimpleCollision, HeightfieldSimple, InOutMaterials))
+	{
+		TRACE_CPUPROFILER_EVENT_SCOPE(Chaos_stream);
+		check(Heightfield);
+		check(!bGenerateSimpleCollision || HeightfieldSimple);
 
-	FMemoryWriter Writer(OutHeightfieldData);
-	Chaos::FChaosArchive Ar(Writer);
-
-	bool bSerializeGenerateSimpleCollision = bGenerateSimpleCollision;
-	Ar << bSerializeGenerateSimpleCollision;
+		FMemoryWriter Writer(OutHeightfieldData);
+		Chaos::FChaosArchive Ar(Writer);
 
-	const int32 NumCollisionCells = FMath::Square(CollisionSizeQuads);
-	const int32 NumSimpleCollisionCells = FMath::Square(SimpleCollisionSizeQuads);
+		bool bSerializeGenerateSimpleCollision = bGenerateSimpleCollision;
+		Ar << bSerializeGenerateSimpleCollision;
 
-	TArrayView<uint8> ComplexMaterialIndicesView(MaterialIndices.GetData(), NumCollisionCells);
-	Heightfield = Chaos::FHeightFieldPtr( new Chaos::FHeightField(Params.Heights, ComplexMaterialIndicesView, CollisionSizeVerts, CollisionSizeVerts, Chaos::FVec3(1)));
-	Ar << Heightfield;
-	if (bGenerateSimpleCollision)
-	{
-		TArrayView<uint8> SimpleMaterialIndicesView(MaterialIndices.GetData() + NumCollisionCells, NumSimpleCollisionCells);
-		HeightfieldSimple = Chaos::FHeightFieldPtr( new Chaos::FHeightField(Params.SimpleHeights, SimpleMaterialIndicesView, SimpleCollisionSizeVerts, SimpleCollisionSizeVerts, Chaos::FVec3(1)));
-		Ar << HeightfieldSimple;
+		Ar << Heightfield;
+		if (bGenerateSimpleCollision)
+		{
+			Ar << HeightfieldSimple;
+		}
+		return true;
 	}
 
-	return true;
+	return false;
 }
 
 #if WITH_EDITOR
@@ -1278,55 +1364,7 @@ bool ULandscapeHeightfieldCollisionComponent::CookCollisionData(const FName& For
 		}
 	}
 
-	int32 CollisionSizeVerts;
-	int32 SimpleCollisionSizeVerts;
-	int32 NumSamples;
-	int32 NumSimpleSamples;
-	GetCollisionSampleInfo(CollisionSizeVerts, SimpleCollisionSizeVerts, NumSamples, NumSimpleSamples);
-
-	const uint16* Heights = (const uint16*)CollisionHeightData.LockReadOnly();
-	check(CollisionHeightData.GetElementCount() == NumSamples + NumSimpleSamples);
-	const uint16* SimpleHeights = Heights + NumSamples;
-
-	// Physical material data from layer system
-	const uint8* DominantLayers = nullptr;
-	const uint8* SimpleDominantLayers = nullptr;
-	if (DominantLayerData.GetElementCount())
-	{
-		DominantLayers = (const uint8*)DominantLayerData.LockReadOnly();
-		check(DominantLayerData.GetElementCount() == NumSamples + NumSimpleSamples);
-		SimpleDominantLayers = DominantLayers + NumSamples;
-	}
-
-	// Physical material data from render material graph
-	const uint8* RenderPhysicalMaterialIds = nullptr;
-	const uint8* SimpleRenderPhysicalMaterialIds = nullptr;
-	if (PhysicalMaterialRenderData.GetElementCount())
-	{
-		RenderPhysicalMaterialIds = (const uint8*)PhysicalMaterialRenderData.LockReadOnly();
-		check(PhysicalMaterialRenderData.GetElementCount() == NumSamples + NumSimpleSamples);
-		SimpleRenderPhysicalMaterialIds = RenderPhysicalMaterialIds + NumSamples;
-	}
-
-	auto MakeSafeArrayView = 
-		[](auto* Data, int32 DataSamples)
-		{
-			return Data ? MakeArrayView(Data, DataSamples) : MakeArrayView(Data, 0);
-		};
-
-	FWriteRuntimeDataParams WriteParams;
-	WriteParams.bUseDefaultMaterialOnly = bUseDefaultMaterialOnly;
-	WriteParams.bProcessRenderIndices = true;
-	WriteParams.bProcessVisibilityLayer = true;
-	WriteParams.Heights = MakeSafeArrayView(Heights, NumSamples);
-	WriteParams.SimpleHeights = MakeSafeArrayView(SimpleHeights, NumSimpleSamples);
-	WriteParams.DominantLayers = MakeSafeArrayView(DominantLayers, NumSamples);
-	WriteParams.SimpleDominantLayers = MakeSafeArrayView(SimpleDominantLayers, NumSimpleSamples);
-	WriteParams.RenderPhysicalMaterialIds = MakeSafeArrayView(RenderPhysicalMaterialIds, NumSamples);
-	WriteParams.SimpleRenderPhysicalMaterialIds = MakeSafeArrayView(SimpleRenderPhysicalMaterialIds, NumSimpleSamples);
-	WriteParams.PhysicalMaterialRenderObjects = MakeSafeArrayView(PhysicalMaterialRenderObjects.GetData(), PhysicalMaterialRenderObjects.Num());
-	WriteParams.ComponentLayerInfos = MakeSafeArrayView(ComponentLayerInfos.GetData(), ComponentLayerInfos.Num());
-	WriteParams.VisibilityLayerIndex = ComponentLayerInfos.IndexOfByKey(ALandscapeProxy::VisibilityLayer);
+	FWriteRuntimeDataParams WriteParams = MakeWriteRuntimeDataParams(bUseDefaultMaterialOnly);
 
 	TArray<uint8> OutData;
 	bool Succeeded = WriteRuntimeData(WriteParams, OutData, InOutMaterials);
@@ -1360,7 +1398,6 @@ bool ULandscapeHeightfieldCollisionComponent::CookCollisionData(const FName& For
 		GetDerivedDataCacheRef().Put(*GetHFDDCKeyString(Format, bUseDefaultMaterialOnly, HeightfieldGuid, InOutMaterials), OutCookedData, GetPathName());
 		bShouldSaveCookedDataToDDC[CookedDataIndex] = false;
 	}
-	
 
 	return Succeeded;
 }
