Optimization for landscape sculpting.  Add ParallelFor to ULandscapeComponent::UpdateCachedBounds and ComputeMipToMipMaxDeltas.

diff --git a/Engine/Source/Runtime/Landscape/Private/LandscapeEdit.cpp b/Engine/Source/Runtime/Landscape/Private/LandscapeEdit.cpp
--- a/Engine/Source/Runtime/Landscape/Private/LandscapeEdit.cpp
+++ b/Engine/Source/Runtime/Landscape/Private/LandscapeEdit.cpp
@@ -100,6 +100,7 @@ LandscapeEdit.cpp: Landscape editing
 #include "Engine/Canvas.h"
 #include "Spatial/PointHashGrid3.h"
 #include "Engine/Texture2DArray.h"
+#include "Async/ParallelFor.h"
 
 DEFINE_LOG_CATEGORY(LogLandscape);
 DEFINE_LOG_CATEGORY(LogLandscapeBP);
@@ -225,43 +226,86 @@ namespace UE::Landscape::Private
 		TRACE_CPUPROFILER_EVENT_SCOPE(ComputeMipToMipMaxDeltas);
 		const int32 TextureSize = 1 << (InNumTextureMips - 1);
 
-		TArray<double> MipToMipMaxDeltas;
-		MipToMipMaxDeltas.AddZeroed(ComputeMipToMipMaxDeltasCount(InNumRelevantMips));
+		bool bForceSingleThread = FApp::ShouldUseThreadingForPerformance();
+		const EParallelForFlags ParallelFlags = bForceSingleThread ? EParallelForFlags::None : EParallelForFlags::ForceSingleThread;
+
+		constexpr int32 MaxBatchSize = 512; // Tested faster than 1024, about the same as 256.  Take 512 in favor of smaller task context allocation size.
+		int32 MaxTaskContexts = 1;
+		if (!bForceSingleThread)  // Keep the arena size small if not using threads.
+		{
+			const TArrayView<const FQuadHeightInfo> LargestMipInfo = GetMipQuadInfosForMipConst(InAllMipsQuadInfos, 0, TextureSize);
+			int32 LargestJobSize = LargestMipInfo.Num();
+			MaxTaskContexts = (LargestJobSize / MaxBatchSize) + 1;
+			MaxTaskContexts = FMath::Min(MaxTaskContexts, FTaskGraphInterface::Get().GetNumWorkerThreads() + 1);  // Similar logic to ParallelForImpl::GetNumberOfThreadTasks, but doesn't need to exactly match.
+		}
+	
+		constexpr int32 ByteAlignment = PLATFORM_CACHE_LINE_SIZE;
+		constexpr int32 WordAlignment = ByteAlignment / sizeof(double);
+
+		// Task contexts will be blocks of ResultSize doubles, with the block size rounded up to be a multiple of the cache line size, to prevent false-sharing.
+		int32 ResultSize = ComputeMipToMipMaxDeltasCount(InNumRelevantMips);
+		int32 ResultSizeAlignedWords = ((ResultSize + WordAlignment - 1) / WordAlignment) * WordAlignment;
+		TArray<double> TaskContextArena;
+		TaskContextArena.AddZeroed(ResultSizeAlignedWords * MaxTaskContexts);
+
+		using TaskContext = TArrayView<double>;
+		TArray<TaskContext, TInlineAllocator<65>> TaskContexts;  // Task contexts the threads use are TArrayViews into the area allocation.
+		TaskContexts.AddUninitialized(MaxTaskContexts);
+		for (int32 Idx = 0; Idx < MaxTaskContexts; ++Idx)
+		{
+			TaskContexts[Idx] = MakeArrayView(&TaskContextArena[Idx * ResultSizeAlignedWords], ResultSizeAlignedWords);
+		};
 
 		for (int32 SourceMipIndex = 0; SourceMipIndex < InNumRelevantMips - 1; ++SourceMipIndex)
 		{
-			TArrayView<const FQuadHeightInfo> SourceMipQuadInfos = GetMipQuadInfosForMipConst(InAllMipsQuadInfos, SourceMipIndex, TextureSize);
+			const TArrayView<const FQuadHeightInfo> SourceMipQuadInfos = GetMipQuadInfosForMipConst(InAllMipsQuadInfos, SourceMipIndex, TextureSize);
 			const int32 NumQuadsForSourceMip = SourceMipQuadInfos.Num();
 			const int32 SourceMipQuadsStride = (TextureSize >> SourceMipIndex) / 2;
 			check(SourceMipQuadsStride * SourceMipQuadsStride == NumQuadsForSourceMip);
 			const int32 NumMaxDeltasForSourceMip = ComputeMaxDeltasCountForMip(SourceMipIndex, InNumRelevantMips);
 			const int32 SourceMipMaxDeltasOffset = ComputeMaxDeltasOffsetForMip(SourceMipIndex, InNumRelevantMips);
-			check(SourceMipMaxDeltasOffset + NumMaxDeltasForSourceMip <= MipToMipMaxDeltas.Num());
-			TArrayView<double> SourceMipToDestinationMipMaxDeltas = MakeArrayView(MipToMipMaxDeltas.GetData() + SourceMipMaxDeltasOffset, NumMaxDeltasForSourceMip);
-			// Iterate through all remaining mips and find the max delta between the source mip and destination mips :
-			for (int32 SourceMipQuadIndex = 0; SourceMipQuadIndex < NumQuadsForSourceMip; ++SourceMipQuadIndex)
-			{
-				const FIntPoint SourceMipQuadCoords(SourceMipQuadIndex % SourceMipQuadsStride, SourceMipQuadIndex / SourceMipQuadsStride);
-				const FQuadHeightInfo& SourceMipQuadInfo = SourceMipQuadInfos[SourceMipQuadIndex];
+			check(SourceMipMaxDeltasOffset + NumMaxDeltasForSourceMip <= ResultSize);
 
-				// Special case for N to N+1 because all the info is located within the row of InAllMipsQuadInfos already so we don't need to fetch from another mip: FQuadHeightInfo's Average is the next mip's value:
-				int32 DestinationMipQuadsStride = SourceMipQuadsStride / 2;
-				FIntPoint DestinationMipQuadCoords = SourceMipQuadCoords / 2;
-				ComputeMaxDelta(SourceMipQuadInfo, SourceMipQuadInfo, SourceMipToDestinationMipMaxDeltas[0]);
-
-				for (int32 DestinationMipIndex = SourceMipIndex + 2; DestinationMipIndex < InNumRelevantMips; ++DestinationMipIndex)
+			ParallelForWithExistingTaskContext(MakeArrayView(TaskContexts), NumQuadsForSourceMip, MaxBatchSize,
+				[SourceMipQuadsStride, &SourceMipQuadInfos, SourceMipMaxDeltasOffset, NumMaxDeltasForSourceMip, InNumRelevantMips, SourceMipIndex, &InAllMipsQuadInfos, TextureSize]
+				(TaskContext& Context, int32 SourceMipQuadIndex)
 				{
-					DestinationMipQuadsStride /= 2;
-					DestinationMipQuadCoords /= 2;
-					const int32 DestinationMipQuadIndex = DestinationMipQuadCoords.X + DestinationMipQuadCoords.Y * DestinationMipQuadsStride;
+					// Accumulate results in the Context of the worker thread
+					TArrayView<double> SourceMipToDestinationMipMaxDeltas = MakeArrayView(Context.GetData() + SourceMipMaxDeltasOffset, NumMaxDeltasForSourceMip);
 
-					TArrayView<const FQuadHeightInfo> DestinationMipQuadInfos = GetMipQuadInfosForMipConst(InAllMipsQuadInfos, DestinationMipIndex, TextureSize);
-					checkSlow(DestinationMipQuadInfos.Num() == DestinationMipQuadsStride * DestinationMipQuadsStride);
+					const FIntPoint SourceMipQuadCoords(SourceMipQuadIndex % SourceMipQuadsStride, SourceMipQuadIndex / SourceMipQuadsStride);
+					const FQuadHeightInfo& SourceMipQuadInfo = SourceMipQuadInfos[SourceMipQuadIndex];
 
-					const FQuadHeightInfo& DestinationMipQuadInfo = DestinationMipQuadInfos[DestinationMipQuadIndex];
-					int32 DestinationMipRelativeIndex = DestinationMipIndex - SourceMipIndex - 1;
-					ComputeMaxDelta(SourceMipQuadInfo, DestinationMipQuadInfo, SourceMipToDestinationMipMaxDeltas[DestinationMipRelativeIndex]);
-				}
+					// Special case for N to N+1 because all the info is located within the row of InAllMipsQuadInfos already so we don't need to fetch from another mip: FQuadHeightInfo's Average is the next mip's value:
+					int32 DestinationMipQuadsStride = SourceMipQuadsStride / 2;
+					FIntPoint DestinationMipQuadCoords = SourceMipQuadCoords / 2;
+					ComputeMaxDelta(SourceMipQuadInfo, SourceMipQuadInfo, SourceMipToDestinationMipMaxDeltas[0]);
+
+					for (int32 DestinationMipIndex = SourceMipIndex + 2; DestinationMipIndex < InNumRelevantMips; ++DestinationMipIndex)
+					{
+						DestinationMipQuadsStride /= 2;
+						DestinationMipQuadCoords /= 2;
+						const int32 DestinationMipQuadIndex = DestinationMipQuadCoords.X + DestinationMipQuadCoords.Y * DestinationMipQuadsStride;
+
+						TArrayView<const FQuadHeightInfo> DestinationMipQuadInfos = GetMipQuadInfosForMipConst(InAllMipsQuadInfos, DestinationMipIndex, TextureSize);
+						checkSlow(DestinationMipQuadInfos.Num() == DestinationMipQuadsStride * DestinationMipQuadsStride);
+
+						const FQuadHeightInfo& DestinationMipQuadInfo = DestinationMipQuadInfos[DestinationMipQuadIndex];
+						int32 DestinationMipRelativeIndex = DestinationMipIndex - SourceMipIndex - 1;
+						ComputeMaxDelta(SourceMipQuadInfo, DestinationMipQuadInfo, SourceMipToDestinationMipMaxDeltas[DestinationMipRelativeIndex]);
+					}
+				},
+				ParallelFlags);
+		}
+
+		// Collapse completed task contexts to the final result.
+		TArray<double> MipToMipMaxDeltas;
+		MipToMipMaxDeltas.AddZeroed(ResultSize);  //ComputeMaxDelta uses absolute values, so starting with 0 is ok.
+		for (int32 ContextIdx = 0; ContextIdx < MaxTaskContexts; ++ContextIdx)
+		{
+			for (int32 ValueIdx = 0; ValueIdx < ResultSize; ++ValueIdx)
+			{
+				MipToMipMaxDeltas[ValueIdx] = FMath::Max(MipToMipMaxDeltas[ValueIdx], TaskContexts[ContextIdx][ValueIdx]);
 			}
 		}
 
@@ -298,7 +342,6 @@ bool ULandscapeComponent::UpdateCachedBounds(bool bInApproximateBounds)
 		TArray<double> PreviousMipToMipMaxDeltas;
 		Swap(MipToMipMaxDeltas, PreviousMipToMipMaxDeltas); 
 
-		// TODO [jonathan.bard] : multithread and make this async : 
 		using namespace UE::Landscape::Private;
 		TArray<FQuadHeightInfo> AllMipsQuadInfos;
 		const int32 TextureSize = (SubsectionSizeQuads + 1) * NumSubsections;
@@ -312,9 +355,28 @@ bool ULandscapeComponent::UpdateCachedBounds(bool bInApproximateBounds)
 		{
 			TRACE_CPUPROFILER_EVENT_SCOPE(FetchMipQuads);
 
-			double LocalMin = DBL_MAX;
-			double LocalMax = -DBL_MAX;
 			AllMipsQuadInfos.AddZeroed(ComputeQuadInfosCount(NumRelevantMips, TextureSize));
+
+			struct alignas(PLATFORM_CACHE_LINE_SIZE) FMinMaxType  //pad to cache line size to avoid false-sharing
+			{
+				double Min = TNumericLimits<double>::Max();
+				double Max = TNumericLimits<double>::Lowest();
+			};
+			TArray<FMinMaxType, TInlineAllocator<17>> ContextValues;
+			// 1024 batch size limit is a little arbitrary.  Each task isn't very complicated and the job size isn't very large so one expects a low ceiling on effective parallelism
+			// due to memory bandwidth and overhead.  Quick test shows this is significantly faster than 4096 or 128 on a workstation with more cores than needed.
+			// With TextureSize=256, MaxBatchSize=1024, MaxTaskContexts is 17.
+			const int32 MaxBatchSize = 1024;
+
+			TArrayView<FQuadHeightInfo> LargestMipInfo = GetMipQuadInfosForMip(AllMipsQuadInfos, 0, TextureSize);
+			int32 LargestJobSize = LargestMipInfo.Num();
+			int32 MaxTaskContexts = (LargestJobSize / MaxBatchSize) + 1;
+			MaxTaskContexts = FMath::Min(MaxTaskContexts, FTaskGraphInterface::Get().GetNumWorkerThreads() + 1);  // Similar logic to ParallelForImpl::GetNumberOfThreadTasks, but doesn't need to exactly match.
+			ContextValues.SetNum(MaxTaskContexts);
+
+			const EParallelForFlags ParallelFlags = FApp::ShouldUseThreadingForPerformance() ? EParallelForFlags::None : EParallelForFlags::ForceSingleThread;
+
+			// Outer loop over mip levels.
 			for (int32 MipIndex = 0; MipIndex < NumRelevantMips; ++MipIndex)
 			{
 				FLandscapeComponentDataInterface CDI(this, MipIndex, /*bWorkOnEditingLayer = */false);
@@ -323,33 +385,45 @@ bool ULandscapeComponent::UpdateCachedBounds(bool bInApproximateBounds)
 				const int32 MipTextureSize = TextureSize >> MipIndex;
 				const int32 MipTextureSubSectionSize = MipTextureSize / NumSubsections;
 				const int32 MipQuadsStride = MipTextureSize / 2;
-				for (int32 MipQuadIndex = 0; MipQuadIndex < NumQuadsForMip; ++MipQuadIndex)
-				{
-					int32 QuadY = MipQuadIndex / MipQuadsStride;
-					int32 QuadX = MipQuadIndex - QuadY * MipQuadsStride;
-					int32 X = QuadX * 2;
-					int32 Y = QuadY * 2;
-					// 2x2 Subsections have a duplicate pixel in the middle so we have to subtract one for the pixels in those subsections, since FLandscapeComponentDataInterface provides a view into 
-					//  pixels as if there was no subsection (hence the pixel indices for a 128 heightmap is in the range [0, 127] when using a single (1x1) subsection, but [0, 126] when using 2x2 subsections) :
-					X -= (X >= MipTextureSubSectionSize) ? 1 : 0;
-					Y -= (Y >= MipTextureSubSectionSize) ? 1 : 0;
-					double QuadVertices[4] = {
-						CDI.GetLocalHeight(X + 0, Y + 0),
-						CDI.GetLocalHeight(X + 1, Y + 0),
-						CDI.GetLocalHeight(X + 0, Y + 1),
-						CDI.GetLocalHeight(X + 1, Y + 1) };
-					FQuadHeightInfo QuadInfo(QuadVertices);
-					MipQuadInfos[MipQuadIndex] = QuadInfo;
-
-					if (MipIndex == 0)
+
+				// Inner loop, over quads in the mip
+				ParallelForWithExistingTaskContext(MakeArrayView(ContextValues), NumQuadsForMip, MaxBatchSize,
+					[&MipQuadInfos, &CDI, MipQuadsStride, MipTextureSubSectionSize, MipIndex](FMinMaxType& Context, int32 MipQuadIndex)
 					{
-						LocalMin = FMath::Min(LocalMin, QuadInfo.Min);
-						LocalMax = FMath::Max(LocalMax, QuadInfo.Max);
-					}
-				}
+						int32 QuadY = MipQuadIndex / MipQuadsStride;
+						int32 QuadX = MipQuadIndex - QuadY * MipQuadsStride;
+						int32 X = QuadX * 2;
+						int32 Y = QuadY * 2;
+						// 2x2 Subsections have a duplicate pixel in the middle so we have to subtract one for the pixels in those subsections, since FLandscapeComponentDataInterface provides a view into 
+						//  pixels as if there was no subsection (hence the pixel indices for a 128 heightmap is in the range [0, 127] when using a single (1x1) subsection, but [0, 126] when using 2x2 subsections) :
+						X -= (X >= MipTextureSubSectionSize) ? 1 : 0;
+						Y -= (Y >= MipTextureSubSectionSize) ? 1 : 0;
+						double QuadVertices[4] = {
+							CDI.GetLocalHeight(X + 0, Y + 0),
+							CDI.GetLocalHeight(X + 1, Y + 0),
+							CDI.GetLocalHeight(X + 0, Y + 1),
+							CDI.GetLocalHeight(X + 1, Y + 1) };
+						FQuadHeightInfo QuadInfo(QuadVertices);
+						MipQuadInfos[MipQuadIndex] = QuadInfo;
+
+						if (MipIndex == 0)
+						{
+							Context.Min = FMath::Min(Context.Min, QuadInfo.Min);
+							Context.Max = FMath::Max(Context.Max, QuadInfo.Max);
+						}
+					},
+					ParallelFlags);
 			}
 
-			NewLocalBox = FBox(FVector(0.0, 0.0, LocalMin), FVector(ComponentSizeQuads, ComponentSizeQuads, LocalMax));
+			FMinMaxType TotalMinMax;
+			for (const FMinMaxType& Context : ContextValues)
+			{
+				TotalMinMax.Min = FMath::Min(Context.Min, TotalMinMax.Min);
+				TotalMinMax.Max = FMath::Max(Context.Max, TotalMinMax.Max);
+			}
+			check(TotalMinMax.Max >= TotalMinMax.Min);
+
+			NewLocalBox = FBox(FVector(0.0, 0.0, TotalMinMax.Min), FVector(ComponentSizeQuads, ComponentSizeQuads, TotalMinMax.Max));
 		}
 
 		MipToMipMaxDeltas = ComputeMipToMipMaxDeltas(MakeArrayView<const FQuadHeightInfo>(AllMipsQuadInfos.GetData(), AllMipsQuadInfos.Num()), NumTextureMips, NumRelevantMips);
@@ -387,6 +461,7 @@ bool ULandscapeComponent::UpdateCachedBounds(bool bInApproximateBounds)
 
 	return bChanged;
 }
+
 #endif //WITH_EDITOR
 
 void ULandscapeComponent::UpdateNavigationRelevance()
