Fix crash when showing collision mesh after a material asset fails to load.

diff --git a/Engine/Source/Runtime/Landscape/Private/LandscapeCollision.cpp b/Engine/Source/Runtime/Landscape/Private/LandscapeCollision.cpp
--- a/Engine/Source/Runtime/Landscape/Private/LandscapeCollision.cpp
+++ b/Engine/Source/Runtime/Landscape/Private/LandscapeCollision.cpp
@@ -592,7 +592,7 @@ FPrimitiveSceneProxy* ULandscapeHeightfieldCollisionComponent::CreateSceneProxy(
 			{
 				Chaos::FChaosPhysicsMaterial* ChaosMaterial = MaterialHandle.Get();
 				UPhysicalMaterial* PhysicalMaterial = (ChaosMaterial != nullptr) ? FChaosUserData::Get<UPhysicalMaterial>(ChaosMaterial->UserData) : nullptr;
-				MaterialIndexColors.Emplace(PhysicalMaterial->DebugColor.ToFColor(/*bSRGB = */false));
+				MaterialIndexColors.Emplace(PhysicalMaterial ? PhysicalMaterial->DebugColor.ToFColor(/*bSRGB = */false) : FColor::Black);
 			}
 
 			for (int32 I = 0; I < NumVerts; I++)
@@ -952,12 +952,18 @@ void ULandscapeHeightfieldCollisionComponent::CreateCollisionObject()
 			// Register materials
 			for(UPhysicalMaterial* PhysicalMaterial : CookedPhysicalMaterials)
 			{
-				//TODO: Figure out why we are getting into a state like this (PhysicalMaterial == nullptr) in the first place. Potentially a loading issue
 				if (PhysicalMaterial)
 				{
 					//todo: total hack until we get landscape fully converted to chaos
 					HeightfieldRef->UsedChaosMaterials.Add(PhysicalMaterial->GetPhysicsMaterial());
 				}
+				else
+				{
+					// If a material fails to load, we have a null value in PhysicalMaterialRenderObjects and end up here.  Substitute the default material.
+					ALandscapeProxy* Proxy = GetLandscapeProxy();
+					UPhysicalMaterial* DefMaterial = (Proxy && Proxy->DefaultPhysMaterial) ? Proxy->DefaultPhysMaterial : GEngine->DefaultPhysMaterial;
+					HeightfieldRef->UsedChaosMaterials.Add(DefMaterial->GetPhysicsMaterial());
+				}
 			}
 
 			// Release cooked collision data
