Optimization for landscape sculpting.  Bypass expensive serialization steps for collision data in the editor.

diff --git a/Engine/Source/Runtime/Landscape/Classes/LandscapeHeightfieldCollisionComponent.h b/Engine/Source/Runtime/Landscape/Classes/LandscapeHeightfieldCollisionComponent.h
--- a/Engine/Source/Runtime/Landscape/Classes/LandscapeHeightfieldCollisionComponent.h
+++ b/Engine/Source/Runtime/Landscape/Classes/LandscapeHeightfieldCollisionComponent.h
@@ -252,6 +252,7 @@ public:
 	/** 
 	 * Cooks raw height data into collision object binary stream
 	 */
+	bool GenerateCollisionObjects(const FName& Format, bool bUseDefaultMaterialOnly, Chaos::FHeightFieldPtr& OutHeightField, bool bGenerateSimpleCollision, Chaos::FHeightFieldPtr& OutSimpleHeightField, TArray<UPhysicalMaterial*>& InOutMaterials) const;
 	virtual bool CookCollisionData(const FName& Format, bool bUseDefaultMaterialOnly, bool bCheckDDC, TArray<uint8>& OutCookedData, TArray<UPhysicalMaterial*>& InOutMaterials) const;
 
 	/** Modify a sub-region of the physics heightfield. Note that this does not update the physical material */
diff --git a/Engine/Source/Runtime/Landscape/Private/LandscapeCollision.cpp b/Engine/Source/Runtime/Landscape/Private/LandscapeCollision.cpp
--- a/Engine/Source/Runtime/Landscape/Private/LandscapeCollision.cpp
+++ b/Engine/Source/Runtime/Landscape/Private/LandscapeCollision.cpp
@@ -852,6 +852,25 @@ void ULandscapeHeightfieldCollisionComponent::CreateCollisionObject()
 {
 	LLM_SCOPE(ELLMTag::ChaosLandscape);
 
+	auto RegisterMaterials = [this](TArray<Chaos::FMaterialHandle>& DestMaterialHandles, const TArray<TObjectPtr<UPhysicalMaterial>>& SrcMaterials)
+		{
+			for (UPhysicalMaterial* PhysicalMaterial : SrcMaterials)
+			{
+				if (PhysicalMaterial)
+				{
+					//todo: total hack until we get landscape fully converted to chaos
+					DestMaterialHandles.Add(PhysicalMaterial->GetPhysicsMaterial());
+				}
+				else
+				{
+					// If a material fails to load, we have a null value in PhysicalMaterialRenderObjects and end up here.  Substitute the default material.
+					ALandscapeProxy* Proxy = GetLandscapeProxy();
+					UPhysicalMaterial* DefMaterial = (Proxy && Proxy->DefaultPhysMaterial) ? Proxy->DefaultPhysMaterial : GEngine->DefaultPhysMaterial;
+					DestMaterialHandles.Add(DefMaterial->GetPhysicsMaterial());
+				}
+			}
+		};
+
 	// If we have not created a heightfield yet - do it now.
 	if (!IsValidRef(HeightfieldRef))
 	{
@@ -861,7 +880,7 @@ void ULandscapeHeightfieldCollisionComponent::CreateCollisionObject()
 		const bool bNeedsEditorHeightField = World && !World->IsGameWorld() && !GetOutermost()->bIsCookedForEditor;
 #endif // WITH_EDITOR
 		FHeightfieldGeometryRef* ExistingHeightfieldRef = nullptr;
-		bool bCheckDDC = true;
+		bool bReuseIsValid = true;  // Are pre-existing copies of CookedCollisionData valid to re-use.
 
 		if (!HeightfieldGuid.IsValid())
 		{
@@ -869,7 +888,7 @@ void ULandscapeHeightfieldCollisionComponent::CreateCollisionObject()
 			uint32 CollisionHash = 0;
 #endif
 			HeightfieldGuid = FGuid::NewDeterministicGuid(GetPathName(), CollisionHash);
-			bCheckDDC = false;
+			bReuseIsValid = false;
 		}
 		else
 		{
@@ -893,15 +912,40 @@ void ULandscapeHeightfieldCollisionComponent::CreateCollisionObject()
 			// was resaved using a commandlet and not saved in the editor, or if a PhysicalMaterial asset was deleted.
 			if (CookedPhysicalMaterials.Num() == 0 || CookedPhysicalMaterials.Contains(nullptr))
 			{
-				bCheckDDC = false;
+				bReuseIsValid = false;
 			}
 
-			// Prepare heightfield data
 			static FName PhysicsFormatName(FPlatformProperties::GetPhysicsFormat());
-			CookCollisionData(PhysicsFormatName, false, bCheckDDC, CookedCollisionData, MutableView(CookedPhysicalMaterials));
 
-			// The World will clean up any speculatively-loaded data we didn't end up using.
-			SpeculativeDDCRequest.Reset();
+			// Use the fast bypass path when the data isn't already available and DDC is disabled for landscape collision.
+			bool bBypassCookingStep = GLandscapeCollisionSkipDDC && !GetOutermost()->bIsCookedForEditor &&
+				!(bReuseIsValid && CookedCollisionData.Num() > 0);
+			if (bBypassCookingStep)
+			{
+				HeightfieldRef = GSharedHeightfieldRefs.Add(HeightfieldGuid, new FHeightfieldGeometryRef(HeightfieldGuid));
+
+				const bool bGenerateSimpleCollision = SimpleCollisionSizeQuads > 0;
+				bool bSuccess = GenerateCollisionObjects(PhysicsFormatName, false, HeightfieldRef->HeightfieldGeometry, bGenerateSimpleCollision, HeightfieldRef->HeightfieldSimpleGeometry, MutableView(CookedPhysicalMaterials));
+
+				if (bNeedsEditorHeightField)
+				{
+					Chaos::FHeightFieldPtr DummySimpleRef;
+					TArray<UPhysicalMaterial*> CookedMaterialsEd;
+					bSuccess &= GenerateCollisionObjects(PhysicsFormatName, true, HeightfieldRef->EditorHeightfieldGeometry, false, DummySimpleRef, CookedMaterialsEd);
+				}
+
+				if (bSuccess)
+				{
+					RegisterMaterials(HeightfieldRef->UsedChaosMaterials, CookedPhysicalMaterials);
+				}
+
+				// Return unconditionally even if GenerateCollisionObjects failed.  CookCollisionData succeeds or fails in the same conditions, so trying that
+				// as a fallback doesn't gain anything.
+				return;
+			}
+
+			CookCollisionData(PhysicsFormatName, false, bReuseIsValid, CookedCollisionData, MutableView(CookedPhysicalMaterials));
+
 #endif //WITH_EDITOR
 
 			if (CookedCollisionData.IsEmpty())
@@ -926,7 +970,7 @@ void ULandscapeHeightfieldCollisionComponent::CreateCollisionObject()
 					}
 					return;
 				}
-				
+
 				// Fallthrough to the shared register materials code below
 			}
 			else
@@ -942,33 +986,18 @@ void ULandscapeHeightfieldCollisionComponent::CreateCollisionObject()
 					Ar << bContainsSimple;
 					Ar << HeightfieldRef->HeightfieldGeometry;
 
-					if(bContainsSimple)
+					if (bContainsSimple)
 					{
 						Ar << HeightfieldRef->HeightfieldSimpleGeometry;
 					}
 				}
 			}
 
-			// Register materials
-			for(UPhysicalMaterial* PhysicalMaterial : CookedPhysicalMaterials)
-			{
-				if (PhysicalMaterial)
-				{
-					//todo: total hack until we get landscape fully converted to chaos
-					HeightfieldRef->UsedChaosMaterials.Add(PhysicalMaterial->GetPhysicsMaterial());
-				}
-				else
-				{
-					// If a material fails to load, we have a null value in PhysicalMaterialRenderObjects and end up here.  Substitute the default material.
-					ALandscapeProxy* Proxy = GetLandscapeProxy();
-					UPhysicalMaterial* DefMaterial = (Proxy && Proxy->DefaultPhysMaterial) ? Proxy->DefaultPhysMaterial : GEngine->DefaultPhysMaterial;
-					HeightfieldRef->UsedChaosMaterials.Add(DefMaterial->GetPhysicsMaterial());
-				}
-			}
+			RegisterMaterials(HeightfieldRef->UsedChaosMaterials, CookedPhysicalMaterials);
 
 			// Release cooked collision data
 			// In cooked builds created collision object will never be deleted while component is alive, so we don't need this data anymore
-			if(FPlatformProperties::RequiresCookedData() || (World && World->IsGameWorld()))
+			if (FPlatformProperties::RequiresCookedData() || (World && World->IsGameWorld()))
 			{
 				CookedCollisionData.Empty();
 				bCookedCollisionDataWasDeleted = true;
@@ -976,10 +1005,10 @@ void ULandscapeHeightfieldCollisionComponent::CreateCollisionObject()
 
 #if WITH_EDITOR
 			// Create heightfield for the landscape editor (no holes in it)
-			if(bNeedsEditorHeightField)
+			if (bNeedsEditorHeightField)
 			{
 				TArray<UPhysicalMaterial*> CookedMaterialsEd;
-				if(CookCollisionData(PhysicsFormatName, true, bCheckDDC, CookedCollisionDataEd, CookedMaterialsEd))
+				if (CookCollisionData(PhysicsFormatName, true, bReuseIsValid, CookedCollisionDataEd, CookedMaterialsEd))
 				{
 					TRACE_CPUPROFILER_EVENT_SCOPE(CreateCollisionObject_ChaosStream);
 					FMemoryReader Reader(CookedCollisionDataEd);
@@ -1309,6 +1338,35 @@ bool ULandscapeHeightfieldCollisionComponent::WriteRuntimeData(const FWriteRunti
 }
 
 #if WITH_EDITOR
+
+// Create the collision object for the component.  Similar to CookCollisionData, but bypasses the buffer serialization step.
+bool ULandscapeHeightfieldCollisionComponent::GenerateCollisionObjects(const FName& Format, bool bUseDefaultMaterialOnly, Chaos::FHeightFieldPtr& OutHeightField, bool bGenerateSimpleCollision, Chaos::FHeightFieldPtr& OutSimpleHeightField, TArray<UPhysicalMaterial*>& InOutMaterials) const
+{
+	TRACE_CPUPROFILER_EVENT_SCOPE(ULandscapeHeightfieldCollisionComponent::GenerateCollisionObjects);
+
+	if (GetOutermost()->bIsCookedForEditor)
+	{
+		return true;
+	}
+
+	FWriteRuntimeDataParams WriteParams = MakeWriteRuntimeDataParams(bUseDefaultMaterialOnly);
+	bool bSucceeded = GenerateCollisionData(WriteParams, OutHeightField, bGenerateSimpleCollision, OutSimpleHeightField, InOutMaterials);
+	if (CollisionHeightData.IsLocked())
+	{
+		CollisionHeightData.Unlock();
+	}
+	if (DominantLayerData.IsLocked())
+	{
+		DominantLayerData.Unlock();
+	}
+	if (PhysicalMaterialRenderData.IsLocked())
+	{
+		PhysicalMaterialRenderData.Unlock();
+	}
+
+	return bSucceeded;
+}
+
 bool ULandscapeHeightfieldCollisionComponent::CookCollisionData(const FName& Format, bool bUseDefaultMaterialOnly, bool bCheckDDC, TArray<uint8>& OutCookedData, TArray<UPhysicalMaterial*>& InOutMaterials) const
 {
 	TRACE_CPUPROFILER_EVENT_SCOPE(ULandscapeHeightfieldCollisionComponent::CookCollisionData);
