Correction to single-thread control.

diff --git a/Engine/Source/Runtime/Landscape/Private/LandscapeEdit.cpp b/Engine/Source/Runtime/Landscape/Private/LandscapeEdit.cpp
--- a/Engine/Source/Runtime/Landscape/Private/LandscapeEdit.cpp
+++ b/Engine/Source/Runtime/Landscape/Private/LandscapeEdit.cpp
@@ -226,12 +226,12 @@ namespace UE::Landscape::Private
 		TRACE_CPUPROFILER_EVENT_SCOPE(ComputeMipToMipMaxDeltas);
 		const int32 TextureSize = 1 << (InNumTextureMips - 1);
 
-		bool bForceSingleThread = FApp::ShouldUseThreadingForPerformance();
-		const EParallelForFlags ParallelFlags = bForceSingleThread ? EParallelForFlags::None : EParallelForFlags::ForceSingleThread;
+		bool bForceSingleThread = !FApp::ShouldUseThreadingForPerformance();
+		const EParallelForFlags ParallelFlags = bForceSingleThread ? EParallelForFlags::ForceSingleThread : EParallelForFlags::None;
 
 		constexpr int32 MaxBatchSize = 512; // Tested faster than 1024, about the same as 256.  Take 512 in favor of smaller task context allocation size.
-		int32 MaxTaskContexts = 1;
-		if (!bForceSingleThread)  // Keep the arena size small if not using threads.
+		int32 MaxTaskContexts = 1;  // Keep the arena size small if not using threads.
+		if (!bForceSingleThread)
 		{
 			const TArrayView<const FQuadHeightInfo> LargestMipInfo = GetMipQuadInfosForMipConst(InAllMipsQuadInfos, 0, TextureSize);
 			int32 LargestJobSize = LargestMipInfo.Num();
