Fix error with FMath::Max
#rnx

diff --git a/Engine/Source/Runtime/Landscape/Private/LandscapeEdit.cpp b/Engine/Source/Runtime/Landscape/Private/LandscapeEdit.cpp
--- a/Engine/Source/Runtime/Landscape/Private/LandscapeEdit.cpp
+++ b/Engine/Source/Runtime/Landscape/Private/LandscapeEdit.cpp
@@ -3995,7 +3995,7 @@ int32 ALandscapeProxy::FRawMeshExportParams::GetNumUVChannelsNeeded() const
 	{
 		for (const FUVConfiguration& ComponentUVConfiguration : *ComponentsUVConfiguration)
 		{
-			FMath::Max(Result, ComponentUVConfiguration.GetNumUVChannelsNeeded());
+			Result = FMath::Max(Result, ComponentUVConfiguration.GetNumUVChannelsNeeded());
 		}
 	}
 	return Result;
