Fixed assert with batched merge : BP brushes could end up allocating an invalid size for their temporary render target

diff --git a/Engine/Source/Runtime/Landscape/Private/LandscapeBlueprintBrushBase.cpp b/Engine/Source/Runtime/Landscape/Private/LandscapeBlueprintBrushBase.cpp
--- a/Engine/Source/Runtime/Landscape/Private/LandscapeBlueprintBrushBase.cpp
+++ b/Engine/Source/Runtime/Landscape/Private/LandscapeBlueprintBrushBase.cpp
@@ -322,7 +322,7 @@ void ALandscapeBlueprintBrushBase::RenderLayer(ILandscapeEditLayerRenderer::FRen
 
 		check(CurrentLayerReadRT->IsTexture2DArray() && WriteRT->IsTexture2DArray());
 		FScratchRenderTargetParams ScratchRenderTargetParams(TEXT("BPBrushScratchRT"), /*bInExactDimensions = */false, /*bInUseUAV = */false, /*bInTargetArraySlicesIndependently = */false,
-			CurrentLayerReadRT->GetResolution(), /*InNumSlices = */0, CurrentLayerReadRT->GetFormat(), CurrentLayerReadRT->GetClearColor(), ERHIAccess::CopyDest);
+			InRenderParams.RenderAreaSectionRect.Size(), /*InNumSlices = */0, CurrentLayerReadRT->GetFormat(), CurrentLayerReadRT->GetClearColor(), ERHIAccess::CopyDest);
 		// We need a new scratch 2D texture in order to copy the merged result of a single target layer, so that the BP brush can use that as input (as it operates on texture 2D render targets only): 
 		AdditionalScratchTexture = MakeUnique<FScratchRenderTargetScope>(ScratchRenderTargetParams);
 		CurrentLayerReadRT2D = AdditionalScratchTexture->RenderTarget;
@@ -343,6 +343,8 @@ void ALandscapeBlueprintBrushBase::RenderLayer(ILandscapeEditLayerRenderer::FRen
 			RHI_BREADCRUMB_EVENT_GAMETHREAD("Copy Source (slice %i) -> %s", TargetLayerIndex, CurrentLayerReadRT2D->GetDebugName());
 
 			ULandscapeScratchRenderTarget::FCopyFromScratchRenderTargetParams CopyParams(CurrentLayerReadRT);
+			// It's important to copy only the needed size, because CurrentLayerReadRT2D can come from AdditionalScratchTexture, which is allocated with bInExactDimensions = false)
+			CopyParams.CopySize = InRenderParams.RenderAreaSectionRect.Size();
 			// Copy from the proper slice in the texture array : 
 			CopyParams.SourceSliceIndex = TargetLayerIndex;
 			CurrentLayerReadRT2D->CopyFrom(CopyParams);
@@ -383,6 +385,9 @@ void ALandscapeBlueprintBrushBase::RenderLayer(ILandscapeEditLayerRenderer::FRen
 			});
 
 			ULandscapeScratchRenderTarget::FCopyFromTextureParams CopyParams(OutputRT2D);
+			// It's important to copy only the needed size, because OutputRT2D is currently sized after ReadRT2D, which can be a render target larger than it needs to be (because 
+			//  it can come from AdditionalScratchTexture, which is allocated with bInExactDimensions = false)
+			CopyParams.CopySize = InRenderParams.RenderAreaSectionRect.Size();
 			CopyParams.DestSliceIndex = TargetLayerIndex;
 			WriteRT->CopyFrom(CopyParams);
 
