Removed unused extern declaration

#rnx

diff --git a/Engine/Source/Runtime/Landscape/Private/LandscapeEditLayerRenderer.cpp b/Engine/Source/Runtime/Landscape/Private/LandscapeEditLayerRenderer.cpp
--- a/Engine/Source/Runtime/Landscape/Private/LandscapeEditLayerRenderer.cpp
+++ b/Engine/Source/Runtime/Landscape/Private/LandscapeEditLayerRenderer.cpp
@@ -23,7 +23,6 @@
 
 #define LOCTEXT_NAMESPACE "LandscapeEditLayerRenderer"
 
-extern TAutoConsoleVariable<int32> CVarLandscapeEditLayersMaxResolutionPerRenderBatch;
 extern TAutoConsoleVariable<float> CVarLandscapeBatchedMergeVisualLogOffsetIncrement;
 extern TAutoConsoleVariable<int32> CVarLandscapeEditLayersClearBeforeEachWriteToScratch;
 extern TAutoConsoleVariable<int32> CVarLandscapeBatchedMergeVisualLogShowMergeType;
