Added support for asynchronous physics creation of ULandscapeHeightfieldCollisionComponent done at component incremental pre-registration.
Requires p.Chaos.EnableAsyncInitBody and LevelStreaming.AllowIncrementalPreRegisterComponents to be enabled.

diff --git a/Engine/Source/Runtime/Landscape/Classes/LandscapeHeightfieldCollisionComponent.h b/Engine/Source/Runtime/Landscape/Classes/LandscapeHeightfieldCollisionComponent.h
--- a/Engine/Source/Runtime/Landscape/Classes/LandscapeHeightfieldCollisionComponent.h
+++ b/Engine/Source/Runtime/Landscape/Classes/LandscapeHeightfieldCollisionComponent.h
@@ -99,7 +99,7 @@ public:
 	UFUNCTION(BlueprintCallable, Category = "Landscape")
 	LANDSCAPE_API ULandscapeComponent* GetRenderComponent() const;
 
-	struct FHeightfieldGeometryRef : public FRefCountedObject
+	struct FHeightfieldGeometryRef : public FThreadSafeRefCountedObject
 	{
 		FGuid Guid;
 
@@ -217,6 +217,7 @@ public:
 	virtual const FCollisionResponseContainer& GetCollisionResponseToChannels() const override;
 	virtual void OnRegister() override;
 	virtual void OnUnregister() override;
+	virtual bool AllowsAsyncPhysicsStateCreation() const override;
 	//~ End USceneComponent Interface.
 
 	//~ Begin UPrimitiveComponent Interface
diff --git a/Engine/Source/Runtime/Landscape/Private/LandscapeCollision.cpp b/Engine/Source/Runtime/Landscape/Private/LandscapeCollision.cpp
--- a/Engine/Source/Runtime/Landscape/Private/LandscapeCollision.cpp
+++ b/Engine/Source/Runtime/Landscape/Private/LandscapeCollision.cpp
@@ -9,6 +9,7 @@
 #include "Stats/Stats.h"
 #include "Serialization/BufferArchive.h"
 #include "Misc/FeedbackContext.h"
+#include "Misc/ScopeLock.h"
 #include "UObject/ObjectSaveContext.h"
 #include "UObject/Package.h"
 #include "UObject/PropertyPortFlags.h"
@@ -124,6 +125,15 @@ namespace LandscapeCollisionCookStats
 }
 #endif
 
+// Used for MT access to GSharedHeightfieldRefs.
+// This is necessary when p.Chaos.EnableAsyncInitBody = true and ULandscapeHeightfieldCollisionComponent::AllowsAsyncPhysicsStateCreation returns true.
+// Since ULandscapeHeightfieldCollisionComponent::AllowsAsyncPhysicsStateCreation only returns true when !WITH_EDITOR, there's no point to lock when WITH_EDITOR.
+#if WITH_EDITOR
+#define UE_SCOPELOCK_SHARED_HEIGHTFIELD_REFS()
+#else
+FCriticalSection GSharedHeightfieldRefsCriticalSection;
+#define UE_SCOPELOCK_SHARED_HEIGHTFIELD_REFS() FScopeLock Lock(&GSharedHeightfieldRefsCriticalSection);
+#endif
 TMap<FGuid, ULandscapeHeightfieldCollisionComponent::FHeightfieldGeometryRef* > GSharedHeightfieldRefs;
 
 ULandscapeHeightfieldCollisionComponent::FHeightfieldGeometryRef::FHeightfieldGeometryRef(FGuid& InGuid)
@@ -135,6 +145,7 @@ PRAGMA_DISABLE_DEPRECATION_WARNINGS
 ULandscapeHeightfieldCollisionComponent::FHeightfieldGeometryRef::~FHeightfieldGeometryRef()
 {
 	// Remove ourselves from the shared map.
+	UE_SCOPELOCK_SHARED_HEIGHTFIELD_REFS();
 	GSharedHeightfieldRefs.Remove(Guid);
 }
 PRAGMA_ENABLE_DEPRECATION_WARNINGS
@@ -336,8 +347,18 @@ const FCollisionResponseContainer& ULandscapeHeightfieldCollisionComponent::GetC
 	return Proxy->BodyInstance.GetResponseToChannels();
 }
 
+bool ULandscapeHeightfieldCollisionComponent::AllowsAsyncPhysicsStateCreation() const
+{
+#if WITH_EDITOR
+	return Super::AllowsAsyncPhysicsStateCreation();
+#else
+	return true;
+#endif
+}
+
 void ULandscapeHeightfieldCollisionComponent::OnCreatePhysicsState()
 {
+	TRACE_CPUPROFILER_EVENT_SCOPE(ULandscapeHeightfieldCollisionComponent::OnCreatePhysicsState);
 	USceneComponent::OnCreatePhysicsState(); // route OnCreatePhysicsState, skip PrimitiveComponent implementation
 
 	if (!BodyInstance.IsValidBodyInstance())
@@ -349,26 +370,19 @@ void ULandscapeHeightfieldCollisionComponent::OnCreatePhysicsState()
 
 		if (IsValidRef(HeightfieldRef))
 		{
-			// Make transform for this landscape component PxActor
-			FTransform LandscapeComponentTransform = GetComponentToWorld();
-			FMatrix LandscapeComponentMatrix = LandscapeComponentTransform.ToMatrixWithScale();
-			FVector LandscapeScale = LandscapeComponentMatrix.ExtractScaling();
+			auto CreateShapesAndActors = [this]()
+			{
+				// Make transform for this landscape component PxActor
+				FTransform LandscapeComponentTransform = GetComponentToWorld();
+				FMatrix LandscapeComponentMatrix = LandscapeComponentTransform.ToMatrixWithScale();
+				FVector LandscapeScale = LandscapeComponentMatrix.ExtractScaling();
 
-			// Reorder the axes
-			FVector TerrainX = LandscapeComponentMatrix.GetScaledAxis(EAxis::X);
-			FVector TerrainY = LandscapeComponentMatrix.GetScaledAxis(EAxis::Y);
-			FVector TerrainZ = LandscapeComponentMatrix.GetScaledAxis(EAxis::Z);
-			LandscapeComponentMatrix.SetAxis(0, TerrainX);
-			LandscapeComponentMatrix.SetAxis(2, TerrainY);
-			LandscapeComponentMatrix.SetAxis(1, TerrainZ);
-			
-			const bool bCreateSimpleCollision = SimpleCollisionSizeQuads > 0;
-			const float SimpleCollisionScale = bCreateSimpleCollision ? CollisionScale * CollisionSizeQuads / SimpleCollisionSizeQuads : 0;
+				const bool bCreateSimpleCollision = SimpleCollisionSizeQuads > 0;
+				const float SimpleCollisionScale = bCreateSimpleCollision ? CollisionScale * CollisionSizeQuads / SimpleCollisionSizeQuads : 0;
 
-			// Create the geometry
-			FVector FinalScale(LandscapeScale.X * CollisionScale, LandscapeScale.Y * CollisionScale, LandscapeScale.Z * LANDSCAPE_ZSCALE);
+				// Create the geometry
+				FVector FinalScale(LandscapeScale.X * CollisionScale, LandscapeScale.Y * CollisionScale, LandscapeScale.Z * LANDSCAPE_ZSCALE);
 
-			{
 				FActorCreationParams Params;
 				Params.InitialTM = LandscapeComponentTransform;
 				Params.InitialTM.SetScale3D(FVector(0));
@@ -402,8 +416,8 @@ void ULandscapeHeightfieldCollisionComponent::OnCreatePhysicsState()
 
 				// Setup filtering
 				FCollisionFilterData QueryFilterData, SimFilterData;
-				CreateShapeFilterData(static_cast<uint8>(GetCollisionObjectType()), FMaskFilter(0), GetOwner()->GetUniqueID(), GetCollisionResponseToChannels(), 
-				GetUniqueID(), 0, QueryFilterData, SimFilterData, true, false, true);
+				CreateShapeFilterData(static_cast<uint8>(GetCollisionObjectType()), FMaskFilter(0), GetOwner()->GetUniqueID(), GetCollisionResponseToChannels(),
+					GetUniqueID(), 0, QueryFilterData, SimFilterData, true, false, true);
 
 				// Heightfield is used for simple and complex collision
 				QueryFilterData.Word3 |= bCreateSimpleCollision ? EPDF_ComplexCollision : (EPDF_SimpleCollision | EPDF_ComplexCollision);
@@ -416,7 +430,7 @@ void ULandscapeHeightfieldCollisionComponent::OnCreatePhysicsState()
 				const bool bNewQueryShape = CollisionFlags.bEnableQueryCollision;
 				const bool bNewSimShape = CollisionFlags.bEnableSimCollisionSimple;
 				const bool bNewProbeShape = CollisionFlags.bEnableProbeCollision;
-				
+
 				NewShape->SetIsProbe(bNewProbeShape);
 				NewShape->SetSimEnabled(bNewSimShape);
 				NewShape->SetQueryEnabled(bNewQueryShape);
@@ -425,9 +439,9 @@ void ULandscapeHeightfieldCollisionComponent::OnCreatePhysicsState()
 				ShapeArray.Emplace(MoveTemp(NewShape));
 
 				// Add simple geometry if necessary
-				if(bCreateSimpleCollision)
+				if (bCreateSimpleCollision)
 				{
-					FVector FinalSimpleCollisionScale(LandscapeScale.X* SimpleCollisionScale, LandscapeScale.Y* SimpleCollisionScale, LandscapeScale.Z* LANDSCAPE_ZSCALE);
+					FVector FinalSimpleCollisionScale(LandscapeScale.X * SimpleCollisionScale, LandscapeScale.Y * SimpleCollisionScale, LandscapeScale.Z * LANDSCAPE_ZSCALE);
 					HeightfieldRef->HeightfieldSimpleGeometry->SetScale(FinalSimpleCollisionScale);
 					Chaos::FImplicitObjectPtr ImplicitHeightFieldSimple(HeightfieldRef->HeightfieldSimpleGeometry);
 					Chaos::FImplicitObjectPtr ChaosSimpleHeightFieldFromCooked = MakeImplicitObjectPtr<Chaos::TImplicitObjectTransformed<Chaos::FReal, 3>>(ImplicitHeightFieldSimple, Chaos::FRigidTransform3(FTransform::Identity));
@@ -453,7 +467,7 @@ void ULandscapeHeightfieldCollisionComponent::OnCreatePhysicsState()
 
 #if WITH_EDITOR
 				// Create a shape for a heightfield which is used only by the landscape editor
-				if(!GetWorld()->IsGameWorld() && !GetOutermost()->bIsCookedForEditor)
+				if (!GetWorld()->IsGameWorld() && !GetOutermost()->bIsCookedForEditor)
 				{
 					HeightfieldRef->EditorHeightfieldGeometry->SetScale(FinalScale * LandscapeComponentTransform.GetScale3D().GetSignVector());
 					Chaos::FImplicitObjectPtr ImplicitEditorHeightField(HeightfieldRef->EditorHeightfieldGeometry);
@@ -482,7 +496,7 @@ void ULandscapeHeightfieldCollisionComponent::OnCreatePhysicsState()
 				}
 #endif // WITH_EDITOR
 				// Push the shapes to the actor
-				if(Geoms.Num() == 1)
+				if (Geoms.Num() == 1)
 				{
 					Body_External.SetGeometry(Geoms[0]);
 				}
@@ -499,9 +513,6 @@ void ULandscapeHeightfieldCollisionComponent::OnCreatePhysicsState()
 				}
 				Body_External.MergeShapesArray(MoveTemp(ShapeArray));
 
-				// Push the actor to the scene
-				FPhysScene* PhysScene = GetWorld()->GetPhysicsScene();
-
 				// Set body instance data
 				BodyInstance.PhysicsUserData = FPhysicsUserData(&BodyInstance);
 				BodyInstance.OwnerComponent = this;
@@ -509,21 +520,27 @@ void ULandscapeHeightfieldCollisionComponent::OnCreatePhysicsState()
 
 				Body_External.SetUserData(&BodyInstance.PhysicsUserData);
 
-				TArray<FPhysicsActorHandle> Actors;
-				Actors.Add(PhysHandle);
-
-				FPhysicsCommand::ExecuteWrite(PhysScene, [&]()
-				{
-					bool bImmediateAccelStructureInsertion = true;
-					PhysScene->AddActorsToScene_AssumesLocked(Actors, bImmediateAccelStructureInsertion);
-				});
+				return PhysHandle;
+			};
 
-				PhysScene->AddToComponentMaps(this, PhysHandle);
-				if (BodyInstance.bNotifyRigidBodyCollision)
-				{
-					PhysScene->RegisterForCollisionEvents(this);
-				}
+			// Push the actor to the scene
+			FPhysScene* PhysScene = GetWorld()->GetPhysicsScene();
 
+			const bool bIsInGameThread = IsInGameThread();
+			check(Chaos::CVars::bEnableAsyncInitBody || bIsInGameThread);
+			FPhysicsActorHandle PhysHandle = bIsInGameThread ? CreateShapesAndActors() : FPhysicsActorHandle();
+			FPhysicsCommand::ExecuteWrite(PhysScene, [&PhysHandle, &CreateShapesAndActors, PhysScene, bIsInGameThread]()
+			{
+				PhysHandle = !bIsInGameThread ? CreateShapesAndActors() : PhysHandle;
+				TArray<FPhysicsActorHandle> Actors = { PhysHandle };
+				const bool bImmediateAccelStructureInsertion = true;
+				PhysScene->AddActorsToScene_AssumesLocked(Actors, bImmediateAccelStructureInsertion);
+			});
+
+			PhysScene->AddToComponentMaps(this, PhysHandle);
+			if (BodyInstance.bNotifyRigidBodyCollision)
+			{
+				PhysScene->RegisterForCollisionEvents(this);
 			}
 		}
 	}
@@ -850,6 +867,7 @@ FPrimitiveSceneProxy* ULandscapeHeightfieldCollisionComponent::CreateSceneProxy(
 
 void ULandscapeHeightfieldCollisionComponent::CreateCollisionObject()
 {
+	TRACE_CPUPROFILER_EVENT_SCOPE(ULandscapeHeightfieldCollisionComponent::CreateCollisionObject);
 	LLM_SCOPE(ELLMTag::ChaosLandscape);
 
 	auto RegisterMaterials = [this](TArray<Chaos::FMaterialHandle>& DestMaterialHandles, const TArray<TObjectPtr<UPhysicalMaterial>>& SrcMaterials)
@@ -874,6 +892,7 @@ void ULandscapeHeightfieldCollisionComponent::CreateCollisionObject()
 	// If we have not created a heightfield yet - do it now.
 	if (!IsValidRef(HeightfieldRef))
 	{
+		UE_SCOPELOCK_SHARED_HEIGHTFIELD_REFS();
 		UWorld* World = GetWorld();
 
 #if WITH_EDITOR
@@ -1097,6 +1116,7 @@ void ULandscapeHeightfieldCollisionComponent::CreateCollisionObject(
 #endif
 	HeightfieldGuid = FGuid::NewDeterministicGuid(GetPathName(), CollisionHash);
 
+	UE_SCOPELOCK_SHARED_HEIGHTFIELD_REFS();
 	HeightfieldRef = GSharedHeightfieldRefs.Add(HeightfieldGuid, new FHeightfieldGeometryRef(HeightfieldGuid));
 	HeightfieldRef->HeightfieldGeometry = Chaos::FHeightFieldPtr(new Chaos::FHeightField(Heights, PhysicalMaterialIds, Info.CollisionSizeVerts, Info.CollisionSizeVerts, Chaos::FVec3(1)));
 	
