Unshelved from pending changelist '38324758':

Cleanup Landscape Build commands :
* Added log at the end of each command, indicating how many landscape proxies were effectively processed by each
* Added additional log at the end of the actionable message-issued build command to indicate how many landscape proxy packages still need to be saved after the user has clicked on the "update/build" button
* Removed Build Physical Materials / Grass Maps / Nanite Only commands since they are all covered by [REDACTED_PERSON] and we should not need to run them individually

diff --git a/Engine/Source/Editor/DetailCustomizations/Private/LandscapeProxyUIDetails.cpp b/Engine/Source/Editor/DetailCustomizations/Private/LandscapeProxyUIDetails.cpp
--- a/Engine/Source/Editor/DetailCustomizations/Private/LandscapeProxyUIDetails.cpp
+++ b/Engine/Source/Editor/DetailCustomizations/Private/LandscapeProxyUIDetails.cpp
@@ -19,6 +19,7 @@
 #include "Internationalization/Internationalization.h"
 #include "Internationalization/Text.h"
 #include "Landscape.h"
+#include "LandscapeEditTypes.h"
 #include "LandscapeInfo.h"
 #include "LandscapeProxy.h"
 #include "LandscapeStreamingProxy.h"
@@ -329,7 +330,7 @@ void FLandscapeProxyUIDetails::CustomizeDetails( IDetailLayoutBuilder& DetailBui
 	// Add Nanite buttons :
 	if (!EditingProxies.IsEmpty())
 	{
-		auto BuildNaniteData = [EditingProxies](bool bInForceRebuild) -> FReply
+		auto BuildNaniteData = [EditingProxies](UE::Landscape::EBuildFlags InBuildFlags) -> FReply
 		{
 			TArray<ALandscapeProxy*> ProxiesToBuild;
 			Algo::TransformIf(EditingProxies, ProxiesToBuild, [](const TWeakObjectPtr<ALandscapeProxy>& InProxy) { return InProxy.IsValid(); }, [](const TWeakObjectPtr<ALandscapeProxy>& InProxy) { return InProxy.Get(); });
@@ -342,15 +343,16 @@ void FLandscapeProxyUIDetails::CustomizeDetails( IDetailLayoutBuilder& DetailBui
 			{
 				if (ULandscapeSubsystem* LandscapeSubsystem = World->GetSubsystem<ULandscapeSubsystem>())
 				{
-					LandscapeSubsystem->BuildNanite(MakeArrayView(ProxiesToBuild), bInForceRebuild);
+					LandscapeSubsystem->BuildNanite(InBuildFlags, MakeArrayView(ProxiesToBuild));
 				}
 			}
 
 			return FReply::Handled();
 		};
 
-		auto GetNumProxiesNeedingRebuild = [EditingProxies](bool bInForceRebuild) -> int32
+		auto GetNumProxiesNeedingRebuild = [EditingProxies](UE::Landscape::EBuildFlags InBuildFlags) -> int32
 		{
+			const bool bForceRebuild = EnumHasAnyFlags(InBuildFlags, UE::Landscape::EBuildFlags::ForceRebuild);
 			TSet<TWeakObjectPtr<ALandscapeProxy>> ProxiesToBuild;
 			for (TWeakObjectPtr<ALandscapeProxy> EditingProxy : EditingProxies)
 			{
@@ -369,7 +371,7 @@ void FLandscapeProxyUIDetails::CustomizeDetails( IDetailLayoutBuilder& DetailBui
 				}
 			}
 
-			return Algo::CountIf(ProxiesToBuild, [bInForceRebuild](TWeakObjectPtr<ALandscapeProxy> InProxy) { return (InProxy.IsValid() && (bInForceRebuild || !InProxy->IsNaniteMeshUpToDate())); });
+			return Algo::CountIf(ProxiesToBuild, [bForceRebuild](TWeakObjectPtr<ALandscapeProxy> InProxy) { return (InProxy.IsValid() && (bForceRebuild || !InProxy->IsNaniteMeshUpToDate())); });
 		};
 
 		auto HasAtLeastOneNaniteLandscape = [LandscapeActors]() -> bool
@@ -392,10 +394,10 @@ void FLandscapeProxyUIDetails::CustomizeDetails( IDetailLayoutBuilder& DetailBui
 				.VAlign(VAlign_Center)
 				.ToolTipText_Lambda([GetNumProxiesNeedingRebuild]()
 				{ 
-					return FText::Format(LOCTEXT("BuildNaniteDataTooltip", "Builds the Nanite mesh representation from the Landscape data if it's not up to date ({0} {0}|plural(one=actors, other=actors) to build)"), GetNumProxiesNeedingRebuild(/*bInForceRebuild = */false));
+					return FText::Format(LOCTEXT("BuildNaniteDataTooltip", "Builds the Nanite mesh representation from the Landscape data if it's not up to date ({0} {0}|plural(one=actors, other=actors) to build)"), GetNumProxiesNeedingRebuild(UE::Landscape::EBuildFlags::None));
 				})
-				.OnClicked_Lambda([BuildNaniteData]() { return BuildNaniteData(/*bInForceRebuild = */false); })
-				.IsEnabled_Lambda([GetNumProxiesNeedingRebuild]() { return (GetNumProxiesNeedingRebuild(/*bInForceRebuild = */false) > 0); })
+				.OnClicked_Lambda([BuildNaniteData]() { return BuildNaniteData(UE::Landscape::EBuildFlags::WriteFinalLog); })
+				.IsEnabled_Lambda([GetNumProxiesNeedingRebuild]() { return (GetNumProxiesNeedingRebuild(UE::Landscape::EBuildFlags::None) > 0); })
 			]
 			+ SHorizontalBox::Slot()
 			.FillWidth(1)
@@ -406,9 +408,9 @@ void FLandscapeProxyUIDetails::CustomizeDetails( IDetailLayoutBuilder& DetailBui
 				.VAlign(VAlign_Center)
 				.ToolTipText_Lambda([GetNumProxiesNeedingRebuild]()
 				{ 
-					return FText::Format(LOCTEXT("RebuildNaniteDataTooltip", "Rebuilds the Nanite mesh representation from the Landscape data ({0} {0}|plural(one=actors, other=actors) to build)"), GetNumProxiesNeedingRebuild(/*bInForceRebuild = */true));
+					return FText::Format(LOCTEXT("RebuildNaniteDataTooltip", "Rebuilds the Nanite mesh representation from the Landscape data ({0} {0}|plural(one=actors, other=actors) to build)"), GetNumProxiesNeedingRebuild(UE::Landscape::EBuildFlags::ForceRebuild));
 				})
-				.OnClicked_Lambda([BuildNaniteData]() { return BuildNaniteData(/*bInForceRebuild = */true); })
+				.OnClicked_Lambda([BuildNaniteData]() { return BuildNaniteData(UE::Landscape::EBuildFlags::WriteFinalLog | UE::Landscape::EBuildFlags::ForceRebuild); })
 			]
 		];
 	}
diff --git a/Engine/Source/Editor/LandscapeEditor/Private/LandscapeEditorModule.cpp b/Engine/Source/Editor/LandscapeEditor/Private/LandscapeEditorModule.cpp
--- a/Engine/Source/Editor/LandscapeEditor/Private/LandscapeEditorModule.cpp
+++ b/Engine/Source/Editor/LandscapeEditor/Private/LandscapeEditorModule.cpp
@@ -20,6 +20,7 @@
 #include "LandscapeFileFormatPng.h"
 #include "LandscapeFileFormatRaw.h"
 #include "LandscapeEditorServices.h"
+#include "LandscapeEditTypes.h"
 #include "LandscapeImageFileCache.h"
 #include "LandscapeSettings.h"
 #include "SLandscapeLayerListDialog.h"
@@ -100,7 +101,7 @@ public:
 		GlobalUICommandList->MapAction(LandscapeActions.ViewModeWireframeOnTop, FExecuteAction::CreateStatic(&ChangeLandscapeViewMode, ELandscapeViewMode::WireframeOnTop), FCanExecuteAction(), FIsActionChecked::CreateStatic(&IsLandscapeViewModeSelected, ELandscapeViewMode::WireframeOnTop));
 		GlobalUICommandList->MapAction(LandscapeActions.ViewModeLayerUsage, FExecuteAction::CreateStatic(&ChangeLandscapeViewMode, ELandscapeViewMode::LayerUsage), FCanExecuteAction(), FIsActionChecked::CreateStatic(&IsLandscapeViewModeSelected, ELandscapeViewMode::LayerUsage));
 		GlobalUICommandList->MapAction(LandscapeActions.ViewModeLayerContribution, FExecuteAction::CreateStatic(&ChangeLandscapeViewMode, ELandscapeViewMode::LayerContribution), FCanExecuteAction(), FIsActionChecked::CreateStatic(&IsLandscapeViewModeSelected, ELandscapeViewMode::LayerContribution));
-		
+
 		// add menu extension
 		UToolMenu* ViewportMenu = UToolMenus::Get()->ExtendMenu("UnrealEd.ViewportToolbar.View");
 		FToolMenuSection& LandscapeSection = ViewportMenu->FindOrAddSection("LevelViewportLandscape");
@@ -142,16 +143,7 @@ public:
 		{
 			FToolMenuSection& Section = BuildMenu->FindOrAddSection("LevelEditorLandscape");
 
-			FUIAction ActionBuildGrassMaps(FExecuteAction::CreateStatic(&UE::Landscape::BuildGrassMaps), FCanExecuteAction());
-			Section.AddMenuEntry(TEXT("BuildGrassMapsOnly"), LOCTEXT("BuildGrassMapsOnly", "Build Grass Maps Only"), LOCTEXT("BuildLandscapeGrassMaps", "Build landscape grass maps"), TAttribute<FSlateIcon>(), ActionBuildGrassMaps, EUserInterfaceActionType::Button);
-
-			FUIAction ActionBuildPhysicalMaterial(FExecuteAction::CreateStatic(&UE::Landscape::BuildPhysicalMaterial), FCanExecuteAction());
-			Section.AddMenuEntry(TEXT("BuildPhysicalMaterialOnly"), LOCTEXT("BuildPhysicalMaterialOnly", "Build Physical Material Only"), LOCTEXT("BuildLandscapePhysicalMaterial", "Build landscape physical material"), TAttribute<FSlateIcon>(), ActionBuildPhysicalMaterial, EUserInterfaceActionType::Button);
-		
-			FUIAction ActionBuildNanite(FExecuteAction::CreateStatic(&UE::Landscape::BuildNanite), FCanExecuteAction());
-			Section.AddMenuEntry(NAME_None, LOCTEXT("BuildNaniteOnly", "Build Nanite Only"), LOCTEXT("BuildLandscapeNanite", "Build Nanite representation"), TAttribute<FSlateIcon>(), ActionBuildNanite, EUserInterfaceActionType::Button);
-
-			FUIAction ActionSaveModifiedLandscapes(FExecuteAction::CreateStatic(&UE::Landscape::SaveModifiedLandscapes), FCanExecuteAction::CreateStatic(&UE::Landscape::HasModifiedLandscapes));
+			FUIAction ActionSaveModifiedLandscapes(FExecuteAction::CreateStatic(&UE::Landscape::SaveModifiedLandscapes, UE::Landscape::EBuildFlags::WriteFinalLog), FCanExecuteAction::CreateStatic(&UE::Landscape::HasModifiedLandscapes));
 			Section.AddMenuEntry(NAME_None,
 				LOCTEXT("SaveModifiedLandscapes", "Save Modified Landscapes"), LOCTEXT("SaveModifiedLandscapesToolTip", "Save landscapes that were modified outside of the editor mode"),
 				TAttribute<FSlateIcon>(), ActionSaveModifiedLandscapes, EUserInterfaceActionType::Button);
diff --git a/Engine/Source/Editor/UnrealEd/Private/EditorBuildUtils.cpp b/Engine/Source/Editor/UnrealEd/Private/EditorBuildUtils.cpp
--- a/Engine/Source/Editor/UnrealEd/Private/EditorBuildUtils.cpp
+++ b/Engine/Source/Editor/UnrealEd/Private/EditorBuildUtils.cpp
@@ -48,6 +48,7 @@
 #include "UObject/UObjectIterator.h"
 #include "VirtualTexturingEditorModule.h"
 #include "Components/RuntimeVirtualTextureComponent.h"
+#include "LandscapeEditTypes.h"
 #include "LandscapeSubsystem.h"
 #include "ShaderCompilerCore.h"
 #include "AssetRegistry/AssetRegistryModule.h"
@@ -1865,7 +1866,7 @@ void FEditorBuildUtils::EditorBuildAllLandscape(UWorld* InWorld)
 	{
 		if (ULandscapeSubsystem* LandscapeSubsystem = InWorld->GetSubsystem<ULandscapeSubsystem>())
 		{
-			LandscapeSubsystem->BuildAll();
+			LandscapeSubsystem->BuildAll(UE::Landscape::EBuildFlags::WriteFinalLog);
 		}
 	}
 }
diff --git a/Engine/Source/Runtime/Landscape/Classes/LandscapeInfo.h b/Engine/Source/Runtime/Landscape/Classes/LandscapeInfo.h
--- a/Engine/Source/Runtime/Landscape/Classes/LandscapeInfo.h
+++ b/Engine/Source/Runtime/Landscape/Classes/LandscapeInfo.h
@@ -268,7 +268,12 @@ public:
 	LANDSCAPE_API int32 GetModifiedPackageCount() const;
 	LANDSCAPE_API TArray<UPackage*> GetModifiedPackages() const;
 	LANDSCAPE_API bool IsPackageModified(UPackage* InPackage) const;
-	LANDSCAPE_API void MarkModifiedPackagesAsDirty();
+	/**
+	 * Goes through the list of modified packages (as in IsPackageModified) and marks them as dirty
+	 * 
+	 * @return the number of packages effectively dirtied
+	 */
+	LANDSCAPE_API int32 MarkModifiedPackagesAsDirty();
 
 	/** Landscapes are a bit special in that they contain derived data within their source data, which can be updated on-the-fly depending on many external factors (e.g. a procedural edit layer can invalidate some 
 		data on load). Depending on project preference, the modifications can be either automatically taken into account (leading to landscape actors becoming dirty on load) or manually by the user (i.e. invalidated
diff --git a/Engine/Source/Runtime/Landscape/Classes/LandscapeProxy.h b/Engine/Source/Runtime/Landscape/Classes/LandscapeProxy.h
--- a/Engine/Source/Runtime/Landscape/Classes/LandscapeProxy.h
+++ b/Engine/Source/Runtime/Landscape/Classes/LandscapeProxy.h
@@ -53,6 +53,8 @@ enum class ENavDataGatheringMode : uint8;
 namespace UE::Landscape
 {
 	enum class EOutdatedDataFlags : uint8;
+	enum class EBuildFlags : uint8;
+
 	namespace Nanite
 	{
 		struct FAsyncBuildData;
@@ -1199,8 +1201,13 @@ PRAGMA_ENABLE_DEPRECATION_WARNINGS
 	UE_DEPRECATED(5.4, "This version of RenderGrassMaps is deprecated.  Use BuildGrassMaps() instead.")
 	void RenderGrassMaps(const TArray<ULandscapeComponent*>& InLandscapeComponents, const TArray<ULandscapeGrassType*>& InGrassTypes) {}
 
-	/** Update the landscape physical material render tasks */
-	void UpdatePhysicalMaterialTasks(bool bInShouldMarkDirty = false);
+	/**
+	 * Update the landscape physical material render tasks
+	 * @param bInShouldMarkDirty indicates whether the proxy should be marked as dirty if anything was built
+	 * 
+	 * @return true if anything was built on this proxy
+	 */
+	bool UpdatePhysicalMaterialTasks(bool bInShouldMarkDirty = false);
 	void UpdatePhysicalMaterialTasksStatus(TSet<ULandscapeComponent*>* OutdatedComponents, int32* OutdatedComponentsCount) const;
 
 	/** Editor notification when changing feature level */
@@ -1293,8 +1300,17 @@ PRAGMA_ENABLE_DEPRECATION_WARNINGS
 	virtual bool SupportsForeignSplineMesh() const override { return true; }
 
 	LANDSCAPE_API int32 GetOutdatedGrassMapCount() const;
-	LANDSCAPE_API void BuildGrassMaps(struct FScopedSlowTask* InSlowTask = nullptr);
-	LANDSCAPE_API void BuildPhysicalMaterial(struct FScopedSlowTask* InSlowTask = nullptr);
+	UE_DEPRECATED(5.6, "InSlowTask was unused, please use the parameterless function")
+	LANDSCAPE_API void BuildGrassMaps(struct FScopedSlowTask* InSlowTask) { BuildGrassMaps(); }
+	LANDSCAPE_API void BuildGrassMaps();
+	UE_DEPRECATED(5.6, "InSlowTask was unused, please use the parameterless function")
+	LANDSCAPE_API void BuildPhysicalMaterial(struct FScopedSlowTask* InSlowTask) { BuildPhysicalMaterial(); }
+	/**
+	 * Builds the physical materials on the components of this proxy
+	 * 
+	 * @return true if anything has been built
+	 */
+	LANDSCAPE_API bool BuildPhysicalMaterial();
 	LANDSCAPE_API void InvalidatePhysicalMaterial();
 	LANDSCAPE_API int32 GetOudatedPhysicalMaterialComponentsCount() const;
 	LANDSCAPE_API virtual void CreateSplineComponent() override;
@@ -1685,7 +1701,10 @@ class FLandscapePhysicalMaterialBuilder
 {
 public:
 	LANDSCAPE_API FLandscapePhysicalMaterialBuilder(UWorld* InWorld);
+	UE_DEPRECATED(5.6, "Use the function with the EBuildFlags param")
 	LANDSCAPE_API void Build();
+	LANDSCAPE_API void Build(UE::Landscape::EBuildFlags InBuildFlags);
+	UE_DEPRECATED(5.6, "Use the Build function with UE::Landscape::EBuildFlags::ForceRebuild")
 	LANDSCAPE_API void Rebuild();
 	LANDSCAPE_API int32 GetOudatedPhysicalMaterialComponentsCount();
 private:
diff --git a/Engine/Source/Runtime/Landscape/Private/Landscape.cpp b/Engine/Source/Runtime/Landscape/Private/Landscape.cpp
--- a/Engine/Source/Runtime/Landscape/Private/Landscape.cpp
+++ b/Engine/Source/Runtime/Landscape/Private/Landscape.cpp
@@ -31,6 +31,7 @@ Landscape.cpp: Terrain rendering
 #include "LandscapeNaniteComponent.h"
 #include "LandscapeRender.h"
 #include "LandscapePrivate.h"
+#include "Logging/StructuredLog.h"
 #include "Logging/TokenizedMessage.h"
 #include "Logging/MessageLog.h"
 #include "Misc/UObjectToken.h"
@@ -3669,8 +3670,8 @@ FBox ALandscape::GetCompleteBounds() const
 
 void ALandscape::SetUseGeneratedLandscapeSplineMeshesActors(bool bInEnabled)
 {
-	bUseGeneratedLandscapeSplineMeshesActors = bInEnabled;
-}
+		bUseGeneratedLandscapeSplineMeshesActors = bInEnabled;
+	}
 
 bool ALandscape::GetUseGeneratedLandscapeSplineMeshesActors() const
 {
@@ -3845,8 +3846,8 @@ void ALandscapeProxy::PreSave(FObjectPreSaveContext ObjectSaveContext)
 		const bool bIsStreamingProxy = this->IsA<ALandscapeStreamingProxy>();
 		const bool bUseEdgeFixups = bIsStreamingProxy;
 		const bool bUpdateSnapshotNow = true;
-		for (ULandscapeComponent* LandscapeComponent : LandscapeComponents)
-		{
+			for (ULandscapeComponent* LandscapeComponent : LandscapeComponents)
+			{
 			LandscapeComponent->InstallOrUpdateTextureUserDatas(bUseEdgeFixups, bShouldCompressHeightmap, bUpdateSnapshotNow);
 		}
 	}
@@ -4593,8 +4594,8 @@ namespace UE::Landscape::Private
 			->AddToken(FUObjectToken::Create(&InSynchronizedProxy, FText::FromString(InSynchronizedProxy.GetActorNameOrLabel())))
 			->AddToken(FTextToken::Create(LOCTEXT("MapCheck_Message_LandscapeProxy_FixupSharedData", "had some shared properties not in sync with its parent landscape actor. This has been fixed but the proxy needs to be saved in order to ensure cooking behaves as expected. ")))
 			->AddToken(FActionToken::Create(LOCTEXT("MapCheck_SaveFixedUpData", "Save Modified Landscapes"), LOCTEXT("MapCheck_SaveFixedUpData_Desc", "Saves the modified landscape proxy actors"),
-				FOnActionTokenExecuted::CreateStatic(&UE::Landscape::SaveModifiedLandscapes),
-				FCanExecuteActionToken::CreateStatic(&UE::Landscape::HasModifiedLandscapes),
+				FOnActionTokenExecuted::CreateUObject(LandscapeSubsystem, &ULandscapeSubsystem::SaveModifiedLandscapes, UE::Landscape::EBuildFlags::WriteFinalLog),
+				FCanExecuteActionToken::CreateUObject(LandscapeSubsystem, &ULandscapeSubsystem::HasModifiedLandscapes),
 				/*bInSingleUse = */false))
 			->AddToken(FTextToken::Create(FText::Format(LOCTEXT("MapCheck_Message_LandscapeProxy_FixupSharedData_SharedProperties", "The following properties were synchronized: {0}."), FText::FromString(SynchronizedPropertiesStringBuilder.ToString()))));
 
@@ -4924,7 +4925,7 @@ void ALandscapeProxy::FixupSharedData(ALandscape* Landscape, const bool bMapChec
 			// In cases where LandscapeInfo is not fully ready yet, we forward the provided ALandscape. If ALandscape is available in LandscapeInfo, we let the object function naturally.
 			const ALandscape* LandscapeActor = (LandscapeInfo->LandscapeActor == nullptr) ? Landscape : nullptr;
 
-			// Force resave the proxy through the modified landscape system, so that the user can then use the Save Modified Landscapes menu and therefore manually trigger the re-save of all modified proxies. * /
+			// Force resave the proxy through the modified landscape system, so that the user can then use the Build > Save Modified Landscapes (or Build > Build Landscape) button and therefore manually trigger the re-save of all modified proxies. * /
 			bool bNeedsManualResave = LandscapeInfo->MarkObjectDirty(/*InObject = */this, /*bInForceResave = */true, LandscapeActor);
 
 			if (bMapCheck && bNeedsManualResave && !SynchronizedProperties.IsEmpty())
@@ -5053,16 +5054,21 @@ bool ULandscapeInfo::IsPackageModified(UPackage* InPackage) const
 	return ModifiedPackages.Contains(InPackage);
 }
 
-void ULandscapeInfo::MarkModifiedPackagesAsDirty()
+int32 ULandscapeInfo::MarkModifiedPackagesAsDirty()
 {
+	int32 NumDirtied = 0;
 	for (TWeakObjectPtr<UPackage> WeakPackagePtr : ModifiedPackages)
 	{
 		if (UPackage* Package = WeakPackagePtr.Get())
 		{
-			Package->MarkPackageDirty();
+			const bool bWasDirty = Package->IsDirty();
+			const bool bIsDirty = Package->MarkPackageDirty();
+			NumDirtied += (!bWasDirty && bIsDirty) ? 1 : 0;
 		}
 	}
 	ModifiedPackages.Empty();
+
+	return NumDirtied;
 }
 
 int32 ULandscapeInfo::GetModifiedPackageCount() const
@@ -5113,7 +5119,7 @@ bool ULandscapeInfo::MarkObjectDirty(UObject* InObject, bool bInForceResave, con
 		if (!InObject->MarkPackageDirty())
 		{
 			// When force-resaving (e.g. when syncing must-sync properties on load), unconditionally add the package to the list of packages to save if we couldn't mark it dirty already, so that 
-			//  the user can manually resave all that needs to be saved with the Save Modified Landscapes button :
+			//  the user can manually resave all that needs to be saved with the Build > Save Modified Landscapes (or Build > Build Landscape) button :
 			bWasAddedToModifiedPackages = TryAddToModifiedPackages(InObject->GetPackage(), InLandscapeOverride);
 		}
 	}
@@ -6442,27 +6448,37 @@ FLandscapePhysicalMaterialBuilder::FLandscapePhysicalMaterialBuilder(UWorld* InW
 {
 }
 
+// Deprecated
 void FLandscapePhysicalMaterialBuilder::Build()
+{
+	Build(UE::Landscape::EBuildFlags::None);
+}
+
+void FLandscapePhysicalMaterialBuilder::Build(UE::Landscape::EBuildFlags InBuildFlags)
 {
 	if (World)
 	{
+		int32 NumBuilt = 0;
 		for (TActorIterator<ALandscapeProxy> ProxyIt(World); ProxyIt; ++ProxyIt)
 		{
-			ProxyIt->BuildPhysicalMaterial();
+			if (EnumHasAnyFlags(InBuildFlags, UE::Landscape::EBuildFlags::ForceRebuild))
+			{
+				ProxyIt->InvalidatePhysicalMaterial();
+			}
+			NumBuilt += ProxyIt->BuildPhysicalMaterial() ? 1 : 0;
+		}
+
+		if (EnumHasAnyFlags(InBuildFlags, UE::Landscape::EBuildFlags::WriteFinalLog))
+		{
+			UE_LOGFMT_LOC(LogLandscape, Log, "BuildPhysicalMaterialFinalLog", "Build Physical Materials: {NumProxies} landscape {NumProxies}|plural(one=proxy,other=proxies) built", ("NumProxies", NumBuilt));
 		}
 	}
 }
 
+// Deprecated
 void FLandscapePhysicalMaterialBuilder::Rebuild()
 {
-	if (World)
-	{
-		for (TActorIterator<ALandscapeProxy> ProxyIt(World); ProxyIt; ++ProxyIt)
-		{
-			ProxyIt->InvalidatePhysicalMaterial();
-			ProxyIt->BuildPhysicalMaterial();
-		}
-	}
+	Build(UE::Landscape::EBuildFlags::ForceRebuild);
 }
 
 int32 FLandscapePhysicalMaterialBuilder::GetOudatedPhysicalMaterialComponentsCount()
@@ -6545,13 +6561,15 @@ void ALandscapeProxy::InvalidatePhysicalMaterial()
 	}
 }
 
-void ALandscapeProxy::BuildPhysicalMaterial(struct FScopedSlowTask* InSlowTask)
+bool ALandscapeProxy::BuildPhysicalMaterial()
 {
 	if (!HasAnyFlags(RF_ClassDefaultObject))
 	{
 		const bool bShouldMarkDirty = true;
-		UpdatePhysicalMaterialTasks(bShouldMarkDirty);
+		return UpdatePhysicalMaterialTasks(bShouldMarkDirty);
 	}
+
+	return false;
 }
 
 void ALandscapeProxy::UpdatePhysicalMaterialTasksStatus(TSet<ULandscapeComponent*>* OutdatedComponents, int32* OutdatedComponentsCount) const
@@ -6588,7 +6606,7 @@ void ALandscapeProxy::UpdatePhysicalMaterialTasksStatus(TSet<ULandscapeComponent
 	}
 }
 
-void ALandscapeProxy::UpdatePhysicalMaterialTasks(bool bInShouldMarkDirty)
+bool ALandscapeProxy::UpdatePhysicalMaterialTasks(bool bInShouldMarkDirty)
 {
 	TSet<ULandscapeComponent*> OutdatedComponents;
 	int32 PendingComponentsToBeSaved = 0;
@@ -6601,6 +6619,8 @@ void ALandscapeProxy::UpdatePhysicalMaterialTasks(bool bInShouldMarkDirty)
 	{
 		MarkPackageDirty();
 	}
+
+	return (PendingComponentsToBeSaved > 0);
 }
 
 void ALandscapeProxy::RemoveNaniteComponents()
diff --git a/Engine/Source/Runtime/Landscape/Private/LandscapeEditLayers.cpp b/Engine/Source/Runtime/Landscape/Private/LandscapeEditLayers.cpp
--- a/Engine/Source/Runtime/Landscape/Private/LandscapeEditLayers.cpp
+++ b/Engine/Source/Runtime/Landscape/Private/LandscapeEditLayers.cpp
@@ -81,7 +81,6 @@ LandscapeEditLayers.cpp: Landscape editing layers mode
 #include "TextureCompiler.h"
 #include "Editor.h"
 #include "LandscapeNotification.h"
-#include "LandscapeSubsystem.h"
 #include "ObjectCacheContext.h"
 #include "Components/RuntimeVirtualTextureComponent.h"
 #include "UnrealEdGlobals.h" // GUnrealEd
@@ -12062,6 +12061,9 @@ bool ALandscapeProxy::RemoveObsoleteLayers(const TSet<FGuid>& InExistingLayers)
 
 	bool bModified = false;
 
+	ULandscapeSubsystem* LandscapeSubsystem = GetWorld()->GetSubsystem<ULandscapeSubsystem>();
+	checkf(LandscapeSubsystem != nullptr, TEXT("RemoveObsoleteLayers can only be called when a subsystem is available"));
+
 	for (const TPair<FGuid, FName>& LayerGuidAndName : ComponentLayers)
 	{
 		if (!InExistingLayers.Contains(LayerGuidAndName.Key))
@@ -12074,8 +12076,8 @@ bool ALandscapeProxy::RemoveObsoleteLayers(const TSet<FGuid>& InExistingLayers)
 				->AddToken(FUObjectToken::Create(this, FText::FromString(GetActorNameOrLabel())))
 				->AddToken(FTextToken::Create(FText::Format(LOCTEXT("MapCheck_Message_LandscapeProxyObsoleteLayer","Layer '{LayerName}' ('{LayerGuid}') was removed from LandscapeProxy because it doesn't match any of the Target Layers. Saving will remove this data for good."), Arguments)))
 				->AddToken(FActionToken::Create(LOCTEXT("MapCheck_RemoveObsoleteLayers", "Save Modified Landscapes"), LOCTEXT("MapCheck_RemoveObsoleteLayers_Desc", "Saves the modified landscape proxy actors"),
-					FOnActionTokenExecuted::CreateStatic(&UE::Landscape::SaveModifiedLandscapes),
-					FCanExecuteActionToken::CreateStatic(&UE::Landscape::HasModifiedLandscapes),
+					FOnActionTokenExecuted::CreateUObject(LandscapeSubsystem, &ULandscapeSubsystem::SaveModifiedLandscapes, UE::Landscape::EBuildFlags::WriteFinalLog),
+					FCanExecuteActionToken::CreateUObject(LandscapeSubsystem, &ULandscapeSubsystem::HasModifiedLandscapes),
 					/*bInSingleUse = */false))
 				->AddToken(FMapErrorToken::Create(FMapErrors::LandscapeComponentPostLoad_Warning));
 
diff --git a/Engine/Source/Runtime/Landscape/Private/LandscapeGrass.cpp b/Engine/Source/Runtime/Landscape/Private/LandscapeGrass.cpp
--- a/Engine/Source/Runtime/Landscape/Private/LandscapeGrass.cpp
+++ b/Engine/Source/Runtime/Landscape/Private/LandscapeGrass.cpp
@@ -2497,7 +2497,7 @@ void ALandscapeProxy::DebugDrawExclusionBoxes(const UWorld* InWorld)
 
 #if WITH_EDITOR
 
-void ALandscapeProxy::BuildGrassMaps(FScopedSlowTask* InSlowTask)
+void ALandscapeProxy::BuildGrassMaps()
 {
 	check(!HasAnyFlags(RF_ClassDefaultObject));
 
@@ -2516,7 +2516,7 @@ void ALandscapeProxy::BuildGrassMaps(FScopedSlowTask* InSlowTask)
 	if (Subsystem)
 	{
 		const bool bMarkDirty = true;
-		Subsystem->GetGrassMapBuilder()->BuildGrassMapsNowForComponents(LandscapeComponents, InSlowTask, bMarkDirty);
+		Subsystem->GetGrassMapBuilder()->BuildGrassMapsNowForComponents(LandscapeComponents, /*InSlowTask = */nullptr, bMarkDirty);
 	}
 }
 
diff --git a/Engine/Source/Runtime/Landscape/Private/LandscapeGrassMapsBuilder.cpp b/Engine/Source/Runtime/Landscape/Private/LandscapeGrassMapsBuilder.cpp
--- a/Engine/Source/Runtime/Landscape/Private/LandscapeGrassMapsBuilder.cpp
+++ b/Engine/Source/Runtime/Landscape/Private/LandscapeGrassMapsBuilder.cpp
@@ -8,8 +8,10 @@
 #include "LandscapeProxy.h"
 #include "LandscapeGrassWeightExporter.h"
 #include "MaterialCachedData.h"
+#include "LandscapeEditTypes.h"
 #include "LandscapeGrassType.h"
 #include "LandscapeSubsystem.h"
+#include "Logging/StructuredLog.h"
 #include "EngineUtils.h"
 #include "Materials/MaterialInstance.h"
 #include "Stats/Stats.h"
@@ -1551,7 +1553,13 @@ bool FLandscapeGrassMapsBuilder::FComponentState::IsBeyondEvictionRange(const TA
 
 #if WITH_EDITOR
 
+// Deprecated
 void FLandscapeGrassMapsBuilder::Build()
+{
+	Build(UE::Landscape::EBuildFlags::None);
+}
+
+void FLandscapeGrassMapsBuilder::Build(UE::Landscape::EBuildFlags InBuildFlags)
 {
 	if (World)
 	{
@@ -1560,6 +1568,7 @@ void FLandscapeGrassMapsBuilder::Build()
 
 		// iterate proxies, update grass types and create the list of components to build
 		TArray<TObjectPtr<ULandscapeComponent>> LandscapeComponentsToBuild;
+		TSet<ALandscapeProxy*> LandscapeProxiesToBuild;
 		for (TActorIterator<ALandscapeProxy> ProxyIt(World); ProxyIt; ++ProxyIt)
 		{
 			ALandscapeProxy* Proxy = *ProxyIt;
@@ -1580,6 +1589,7 @@ void FLandscapeGrassMapsBuilder::Build()
 				{
 					Component->MarkPackageDirty();
 					LandscapeComponentsToBuild.Add(Component);
+					LandscapeProxiesToBuild.Add(Proxy);
 				}
 			}
 
@@ -1598,6 +1608,11 @@ void FLandscapeGrassMapsBuilder::Build()
 			constexpr bool bMarkDirty = true;
 			BuildGrassMapsNowForComponents(LandscapeComponentsToBuild, &SlowTask, bMarkDirty);
 		}
+
+		if (EnumHasAnyFlags(InBuildFlags, UE::Landscape::EBuildFlags::WriteFinalLog))
+		{
+			UE_LOGFMT_LOC(LogLandscape, Log, "BuildGrassFinalLog", "Build Grass: {NumProxies} landscape {NumProxies}|plural(one=proxy,other=proxies) built.", ("NumProxies", LandscapeProxiesToBuild.Num()));
+		}
 	}
 }
 
diff --git a/Engine/Source/Runtime/Landscape/Private/LandscapeGrassMapsBuilder.h b/Engine/Source/Runtime/Landscape/Private/LandscapeGrassMapsBuilder.h
--- a/Engine/Source/Runtime/Landscape/Private/LandscapeGrassMapsBuilder.h
+++ b/Engine/Source/Runtime/Landscape/Private/LandscapeGrassMapsBuilder.h
@@ -12,6 +12,11 @@
 class FLandscapeGrassWeightExporter;
 struct FScopedSlowTask;
 
+namespace UE::Landscape
+{
+	enum class EBuildFlags : uint8;
+} // namespace UE::Landscape
+
 class FAsyncFetchTask : public FNonAbandonableTask
 {
 public:
@@ -43,7 +48,9 @@ public:
 
 #if WITH_EDITOR
 	/** Synchronously build all grass maps on all proxies in the world */
+	UE_DEPRECATED(5.6, "Use the function with the EBuildFlags param")
 	LANDSCAPE_API void Build();
+	LANDSCAPE_API void Build(UE::Landscape::EBuildFlags InBuildFlags);
 	
 	/** Return the number of landscape components in the world that have grass map data that is not up to date (NOT including components missing grass map data altogether) */
 	LANDSCAPE_API int32 GetOutdatedGrassMapCount(bool bInForceUpdate = true) const;
diff --git a/Engine/Source/Runtime/Landscape/Private/LandscapeSubsystem.cpp b/Engine/Source/Runtime/Landscape/Private/LandscapeSubsystem.cpp
--- a/Engine/Source/Runtime/Landscape/Private/LandscapeSubsystem.cpp
+++ b/Engine/Source/Runtime/Landscape/Private/LandscapeSubsystem.cpp
@@ -22,6 +22,7 @@
 #include "LandscapeTextureStorageProvider.h"
 #include "LandscapeComponent.h"
 #include "LandscapeGroup.h"
+#include "Logging/StructuredLog.h"
 #include "ProfilingDebugging/CsvProfiler.h"
 #include "WorldPartition/WorldPartitionSubsystem.h"
 #include "ActorPartition/ActorPartitionSubsystem.h"
@@ -79,57 +80,132 @@ DECLARE_CYCLE_STAT(TEXT("LandscapeSubsystem Tick"), STAT_LandscapeSubsystemTick,
 namespace UE::Landscape
 {
 #if WITH_EDITOR
+	bool HasModifiedLandscapes()
+	{
+		if (GEditor)
+		{
+			if (UWorld* World = GEditor->GetEditorWorldContext().World())
+			{
+				if (ULandscapeSubsystem* LandscapeSubsystem = World->GetSubsystem<ULandscapeSubsystem>())
+				{
+					return LandscapeSubsystem->HasModifiedLandscapes();
+				}
+			}
+		}
+
+		return false;
+	}
+
+	void SaveModifiedLandscapes(UE::Landscape::EBuildFlags InBuildFlags)
+	{
+		if (GEditor)
+		{
+			if (UWorld* World = GEditor->GetEditorWorldContext().World())
+			{
+				if (ULandscapeSubsystem* LandscapeSubsystem = World->GetSubsystem<ULandscapeSubsystem>())
+				{
+					LandscapeSubsystem->SaveModifiedLandscapes(InBuildFlags);
+				}
+			}
+		}
+	}
+
+	// Deprecated
 	void MarkModifiedLandscapesAsDirty()
 	{
-		if (UWorld* World = GEditor->GetEditorWorldContext().World())
+		MarkModifiedLandscapesAsDirty(EBuildFlags::None);
+	}
+
+	void MarkModifiedLandscapesAsDirty(EBuildFlags InBuildFlags)
+	{
+		if (GEditor)
 		{
-			if (ULandscapeSubsystem* LandscapeSubsystem = World->GetSubsystem<ULandscapeSubsystem>())
+			if (UWorld* World = GEditor->GetEditorWorldContext().World())
 			{
-				LandscapeSubsystem->MarkModifiedLandscapesAsDirty();
+				if (ULandscapeSubsystem* LandscapeSubsystem = World->GetSubsystem<ULandscapeSubsystem>())
+				{
+					LandscapeSubsystem->MarkModifiedLandscapesAsDirty(InBuildFlags);
+				}
 			}
 		}
 	}
 	
+	// Deprecated
 	void BuildGrassMaps()
 	{
-		if (UWorld* World = GEditor->GetEditorWorldContext().World())
+		BuildGrassMaps(EBuildFlags::None);
+	}
+
+	void BuildGrassMaps(EBuildFlags InBuildFlags)
+	{
+		if (GEditor)
 		{
-			if (ULandscapeSubsystem* LandscapeSubsystem = World->GetSubsystem<ULandscapeSubsystem>())
+			if (UWorld* World = GEditor->GetEditorWorldContext().World())
 			{
-				LandscapeSubsystem->BuildGrassMaps();
+				if (ULandscapeSubsystem* LandscapeSubsystem = World->GetSubsystem<ULandscapeSubsystem>())
+				{
+					LandscapeSubsystem->BuildGrassMaps(InBuildFlags);
+				}
 			}
 		}
 	}
 
+	// Deprecated
 	void BuildPhysicalMaterial()
 	{
-		if (UWorld* World = GEditor->GetEditorWorldContext().World())
+		BuildPhysicalMaterial(EBuildFlags::None);
+	}
+
+	void BuildPhysicalMaterial(EBuildFlags InBuildFlags)
+	{
+		if (GEditor)
 		{
-			if (ULandscapeSubsystem* LandscapeSubsystem = World->GetSubsystem<ULandscapeSubsystem>())
+			if (UWorld* World = GEditor->GetEditorWorldContext().World())
 			{
-				LandscapeSubsystem->BuildPhysicalMaterial();
+				if (ULandscapeSubsystem* LandscapeSubsystem = World->GetSubsystem<ULandscapeSubsystem>())
+				{
+					LandscapeSubsystem->BuildPhysicalMaterial(InBuildFlags);
+				}
 			}
 		}
 	}
 
+	// Deprecated
 	void BuildNanite()
 	{
-		if (UWorld* World = GEditor->GetEditorWorldContext().World())
+		BuildNanite(EBuildFlags::None);
+	}
+
+	void BuildNanite(EBuildFlags InBuildFlags)
+	{
+		if (GEditor)
 		{
-			if (ULandscapeSubsystem* LandscapeSubsystem = World->GetSubsystem<ULandscapeSubsystem>())
+			if (UWorld* World = GEditor->GetEditorWorldContext().World())
 			{
-				LandscapeSubsystem->BuildNanite();
+				if (ULandscapeSubsystem* LandscapeSubsystem = World->GetSubsystem<ULandscapeSubsystem>())
+				{
+					LandscapeSubsystem->BuildNanite(InBuildFlags);
+				}
 			}
 		}
 	}
 
+	// Deprecated
 	void BuildAll()
 	{
-		if (UWorld* World = GEditor->GetEditorWorldContext().World())
+		BuildAll(EBuildFlags::None);
+	}
+
+	void BuildAll(EBuildFlags InBuildFlags)
+	{
+		if (GEditor)
 		{
-			if (ULandscapeSubsystem* LandscapeSubsystem = World->GetSubsystem<ULandscapeSubsystem>())
+			if (UWorld* World = GEditor->GetEditorWorldContext().World())
 			{
-				LandscapeSubsystem->BuildAll();
+				if (ULandscapeSubsystem* LandscapeSubsystem = World->GetSubsystem<ULandscapeSubsystem>())
+				{
+					LandscapeSubsystem->BuildAll(InBuildFlags);
+				}
 			}
 		}
 	}
@@ -150,7 +226,7 @@ void ULandscapeSubsystem::RegisterActor(ALandscapeProxy* Proxy)
 	SUBSYSTEM_DEBUG_LOG_REGISTER(TEXT("ULandscapeSubsystem::RegisterActor %p %s (%s) to subsystem %p in world %p (%s)"),
 		Proxy, *Proxy->GetFullName(), *Proxy->GetClass()->GetName(), this, GetWorld(), *GetWorld()->GetName());
 
-	TObjectPtr<ALandscapeProxy> ProxyPtr(Proxy);
+ 	TObjectPtr<ALandscapeProxy> ProxyPtr(Proxy);
 
 	// in editor we can get multiple registration calls, so ensure we don't register more than once
 	if (Proxy->bIsRegisteredWithSubsystem)
@@ -311,7 +387,7 @@ void ULandscapeSubsystem::Deinitialize()
 	Scalability::OnScalabilitySettingsChanged.Remove(OnScalabilityChangedHandle);
 	
 #if WITH_EDITOR
-
+	
 	while (NaniteBuildsInFlight != 0)
 	{
 		ENamedThreads::Type CurrentThread = FTaskGraphInterface::Get().GetCurrentThreadIfKnown();
@@ -501,7 +577,7 @@ void ULandscapeSubsystem::Tick(float DeltaTime)
 		const bool bNewShowCollisions = NumViewsWithShowCollision > 0;
 		const bool bShowCollisionChanged = (bNewShowCollisions != bAnyViewShowCollisions);
 		bAnyViewShowCollisions = bNewShowCollisions;
-
+        
 		if (bShowCollisionChanged)
 		{
 			for (ULandscapeHeightfieldCollisionComponent* LandscapeHeightfieldCollisionComponent : TObjectRange<ULandscapeHeightfieldCollisionComponent>(RF_ClassDefaultObject | RF_ArchetypeObject, true, EInternalObjectFlags::Garbage))
@@ -798,24 +874,70 @@ void ULandscapeSubsystem::TickEdgeFixup()
 }
 
 #if WITH_EDITOR
+
+// Deprecated
 void ULandscapeSubsystem::BuildAll()
 {
+	BuildAll(UE::Landscape::EBuildFlags::None);
+}
+
+void ULandscapeSubsystem::BuildAll(UE::Landscape::EBuildFlags InBuildFlags)
+{
+	TRACE_CPUPROFILER_EVENT_SCOPE(ULandscapeSubsystem::BuildAll);
+
 	// This is a deliberate action, make sure to flush all packages that are 'pending dirty' :
-	MarkModifiedLandscapesAsDirty();
+	MarkModifiedLandscapesAsDirty(InBuildFlags);
 
-	BuildGrassMaps();
-	BuildPhysicalMaterial();
-	BuildNanite();
+	BuildGrassMaps(InBuildFlags);
+	BuildPhysicalMaterial(InBuildFlags);
+	BuildNanite(InBuildFlags);
+
+	if (EnumHasAnyFlags(InBuildFlags, UE::Landscape::EBuildFlags::WriteFinalLog))
+	{
+		UE_LOGFMT_LOC(LogLandscape, Log, "BuildAllFinalLog", "Build All operation done.");
+	}
 }
 
+// Deprecated
 void ULandscapeSubsystem::BuildGrassMaps()
 {
-	GrassMapsBuilder->Build();
+	BuildGrassMaps(UE::Landscape::EBuildFlags::None);
 }
 
+void ULandscapeSubsystem::BuildGrassMaps(UE::Landscape::EBuildFlags InBuildFlags)
+{
+	GrassMapsBuilder->Build(InBuildFlags);
+}
+
+// Deprecated
 void ULandscapeSubsystem::BuildPhysicalMaterial()
 {
-	PhysicalMaterialBuilder->Rebuild();
+	BuildPhysicalMaterial(UE::Landscape::EBuildFlags::None);
+}
+
+void ULandscapeSubsystem::BuildPhysicalMaterial(UE::Landscape::EBuildFlags InBuildFlags)
+{
+	PhysicalMaterialBuilder->Build(InBuildFlags);
+}
+
+TSet<UPackage*> ULandscapeSubsystem::GetDirtyLandscapeProxyPackages() const
+{
+	TSet<UPackage*> DirtyPackages;
+	ForEachLandscapeInfo([&](ULandscapeInfo* LandscapeInfo)
+		{
+			LandscapeInfo->ForEachLandscapeProxy([&](const ALandscapeProxy* LandscapeProxy)
+				{
+					UPackage* Package = LandscapeProxy->GetPackage();
+					if (Package->IsDirty())
+					{
+						DirtyPackages.Add(Package);
+					}
+					return true;
+				});
+			return true;
+		});
+
+	return DirtyPackages;
 }
 
 // Deprecated :
@@ -877,10 +999,18 @@ TArray<TTuple<ALandscapeProxy*, UE::Landscape::EOutdatedDataFlags>> ULandscapeSu
 	return OutdatedProxies;
 }
 
+// Deprecated
 void ULandscapeSubsystem::BuildNanite(TArrayView<ALandscapeProxy*> InProxiesToBuild, bool bForceRebuild)
+{
+	BuildNanite(bForceRebuild ? UE::Landscape::EBuildFlags::ForceRebuild : UE::Landscape::EBuildFlags::None, InProxiesToBuild);
+}
+
+void ULandscapeSubsystem::BuildNanite(UE::Landscape::EBuildFlags InBuildFlags, TArrayView<ALandscapeProxy*> InProxiesToBuild)
 {
 	TRACE_CPUPROFILER_EVENT_SCOPE(ULandscapeSubsystem::BuildNanite);
 
+	const bool bForceRebuild = EnumHasAnyFlags(InBuildFlags, UE::Landscape::EBuildFlags::ForceRebuild);
+
 	UWorld* World = GetWorld();
 	if (!World || World->IsGameWorld())
 	{
@@ -946,7 +1076,7 @@ void ULandscapeSubsystem::BuildNanite(TArrayView<ALandscapeProxy*> InProxiesToBu
 	int32 LastRemainingMeshes = NaniteBuildsInFlight.load();
 	int32 TotalMeshes = LastRemainingMeshes;
 	FScopedSlowTask SlowTask(TotalMeshes, (LOCTEXT("Landscape_BuildNanite", "Building Nanite Landscape Meshes")));
-	// todo [don.boogert] - need mechanism to abort all the current work here
+	// todo [don.boogert] - need mechanism to abort all the current work here : the cancellation should also lead to a different final log message...
 	const bool bShowCancelButton = false;
 	SlowTask.MakeDialog(bShowCancelButton);
 
@@ -962,6 +1092,11 @@ void ULandscapeSubsystem::BuildNanite(TArrayView<ALandscapeProxy*> InProxiesToBu
 		FAssetCompilingManager::Get().ProcessAsyncTasks();
 		SlowTask.EnterProgressFrame(MeshesProcessed, FText::Format(LOCTEXT("Landscape_BuildNaniteProgress", "Building Nanite Landscape Mesh ({0} of {1})"), FText::AsNumber(TotalMeshes - LastRemainingMeshes), FText::AsNumber(SlowTask.TotalAmountOfWork)));
 	}
+
+	if (EnumHasAnyFlags(InBuildFlags, UE::Landscape::EBuildFlags::WriteFinalLog))
+	{
+		UE_LOGFMT_LOC(LogLandscape, Log, "BuildNaniteFinalLog", "Build Nanite: {NumProxies} landscape {NumProxies}|plural(one=proxy,other=proxies) built.", ("NumProxies", FinalProxiesToBuild.Num()));
+	}
 }
 
 bool ULandscapeSubsystem::GetDirtyOnlyInMode() const
@@ -971,7 +1106,13 @@ bool ULandscapeSubsystem::GetDirtyOnlyInMode() const
 		|| (Settings->LandscapeDirtyingMode == ELandscapeDirtyingMode::InLandscapeModeAndUserTriggeredChanges);
 }
 
+// Deprecated
 void ULandscapeSubsystem::SaveModifiedLandscapes()
+{	
+	SaveModifiedLandscapes(UE::Landscape::EBuildFlags::None);
+}
+
+void ULandscapeSubsystem::SaveModifiedLandscapes(UE::Landscape::EBuildFlags InBuildFlags)
 {	
 	TSet<UPackage*> SetDirtyPackages;
 	TSet<FName> PackagesToSave;
@@ -1011,6 +1152,7 @@ void ULandscapeSubsystem::SaveModifiedLandscapes()
 		});
 
 	// If Package wasn't saved it is still in the LandscapeInfo ModifiedPackage list, set its dirty flag back to false.
+	int32 NumPackagesNotProcessed = 0;
 	ForEachLandscapeInfo([&](ULandscapeInfo* LandscapeInfo)
 	{
 		for (UPackage* ModifiedPackage : LandscapeInfo->GetModifiedPackages())
@@ -1018,20 +1160,40 @@ void ULandscapeSubsystem::SaveModifiedLandscapes()
 			if (SetDirtyPackages.Contains(ModifiedPackage))
 			{
 				ModifiedPackage->SetDirtyFlag(false);
+				++NumPackagesNotProcessed;
 			}
 		}
 		return true;
 	});
+
+	if (EnumHasAnyFlags(InBuildFlags, UE::Landscape::EBuildFlags::WriteFinalLog))
+	{
+		check(NumPackagesNotProcessed <= PackagesToSave.Num());
+		UE_LOGFMT_LOC(LogLandscape, Log, "SaveModifiedLandscapesFinalLog", "Save Modified Landscapes : {NumPackagesToSave} landscape {NumPackagesToSave}|plural(one=proxy,other=proxies) considered : {NumPackagesSaved} saved ({NumPackagesNotProcessed} skipped).", 
+			("NumPackagesToSave", PackagesToSave.Num()), ("NumPackagesSaved", (PackagesToSave.Num() - NumPackagesNotProcessed)), ("NumPackagesNotProcessed", NumPackagesNotProcessed));
+	}
 }
 
+// Deprecated
 void ULandscapeSubsystem::MarkModifiedLandscapesAsDirty()
 {
+	MarkModifiedLandscapesAsDirty(UE::Landscape::EBuildFlags::None);
+}
+
+void ULandscapeSubsystem::MarkModifiedLandscapesAsDirty(UE::Landscape::EBuildFlags InBuildFlags)
+{
+	int32 NumDirtied = 0;
 	// Flush all packages that are pending mark for dirty : 
-	ForEachLandscapeInfo([&](ULandscapeInfo* LandscapeInfo)
+	ForEachLandscapeInfo([InBuildFlags, &NumDirtied](ULandscapeInfo* LandscapeInfo)
 	{
-		LandscapeInfo->MarkModifiedPackagesAsDirty();
+		NumDirtied += LandscapeInfo->MarkModifiedPackagesAsDirty();
 		return true;
 	});
+
+	if (EnumHasAnyFlags(InBuildFlags, UE::Landscape::EBuildFlags::WriteFinalLog))
+	{
+		UE_LOGFMT_LOC(LogLandscape, Log, "MarkModifiedLandscapesAsDirtyFinalLog", "Mark Modified Landscapes Dirty : {NumProxies} {NumProxies}|plural(one=proxy,other=proxies) made dirty.", ("NumProxies", NumDirtied));
+	}
 }
 
 bool ULandscapeSubsystem::HasModifiedLandscapes() const
@@ -1118,13 +1280,23 @@ bool ULandscapeSubsystem::GetActionableMessage(FActionableMessage& OutActionable
 	}
 	});
 	
+	auto LogFinalNumberOfDirtyProxies = [this]() 
+		{
+			const int32 NumDirtyProxies = GetDirtyLandscapeProxyPackages().Num();
+			UE_LOGFMT_LOC(LogLandscape, Log, "BuildFinalNumberOfDirtyProxies", "{NumProxies} landscape {NumProxies}|plural(one=proxy,other=proxies) now need to be saved.", ("NumProxies", NumDirtyProxies));
+		};
+
 	// If more than 1 action is required, go with a BuildAll action
 	if (FMath::CountBits(static_cast<uint64>(OutdatedFlagsUnion)) > 1)
 	{
 		OutActionableMessage.Message = FText::Format(LOCTEXT("SeveralLandscapeDataOutdated.Message", "{0} Landscape {0}|plural(one=actor,other=actors) {0}|plural(one=is,other=are) out of date and {0}|plural(one=needs,other=need) to be rebuilt"), NumTotalOutdatedProxy);;
 		OutActionableMessage.Tooltip = DefaultTooltip;
 		OutActionableMessage.ActionMessage = DefaultActionMessage;
-		OutActionableMessage.ActionCallback = UE::Landscape::BuildAll;
+		OutActionableMessage.ActionCallback = [LogFinalNumberOfDirtyProxies]()
+			{
+				UE::Landscape::BuildAll(UE::Landscape::EBuildFlags::WriteFinalLog);
+				LogFinalNumberOfDirtyProxies();
+			};
 
 		return true;
 	}
@@ -1134,7 +1306,11 @@ bool ULandscapeSubsystem::GetActionableMessage(FActionableMessage& OutActionable
 		OutActionableMessage.Message = FText::Format(LOCTEXT("GRASS_MAPS_NEED_TO_BE_REBUILT_FMT", "{0} Landscape {0}|plural(one=actor,other=actors) with grass maps {0}|plural(one=needs,other=need) to be rebuilt"), OutdatedProxiesCount);
 		OutActionableMessage.Tooltip = DefaultTooltip;
 		OutActionableMessage.ActionMessage = DefaultActionMessage;
-		OutActionableMessage.ActionCallback = UE::Landscape::BuildGrassMaps;
+		OutActionableMessage.ActionCallback = [LogFinalNumberOfDirtyProxies]()
+			{
+				UE::Landscape::BuildGrassMaps(UE::Landscape::EBuildFlags::WriteFinalLog);
+				LogFinalNumberOfDirtyProxies();
+			};
 
 		return true;
 	}
@@ -1144,7 +1320,11 @@ bool ULandscapeSubsystem::GetActionableMessage(FActionableMessage& OutActionable
 		OutActionableMessage.Message = FText::Format(LOCTEXT("LANDSCAPE_PHYSICALMATERIAL_NEED_TO_BE_REBUILT_FMT", "{0} Landscape {0}|plural(one=actor,other=actors) with physical materials {0}|plural(one=needs,other=need) to be rebuilt"), OutdatedProxiesCount);
 		OutActionableMessage.Tooltip = DefaultTooltip;
 		OutActionableMessage.ActionMessage = DefaultActionMessage;
-		OutActionableMessage.ActionCallback = UE::Landscape::BuildPhysicalMaterial;
+		OutActionableMessage.ActionCallback = [LogFinalNumberOfDirtyProxies]()
+			{
+				UE::Landscape::BuildPhysicalMaterial(UE::Landscape::EBuildFlags::WriteFinalLog);
+				LogFinalNumberOfDirtyProxies();
+			};
 
 		return true;
 	}
@@ -1154,7 +1334,11 @@ bool ULandscapeSubsystem::GetActionableMessage(FActionableMessage& OutActionable
 		OutActionableMessage.Message = FText::Format(LOCTEXT("LANDSCAPE_NANITE_MESHES_NEED_TO_BE_REBUILT_FMT", "{0} Landscape {0}|plural(one=actor,other=actors) with Nanite meshes {0}|plural(one=needs,other=need) to be rebuilt"), OutdatedProxiesCount);
 		OutActionableMessage.Tooltip = DefaultTooltip;
 		OutActionableMessage.ActionMessage = DefaultActionMessage;
-		OutActionableMessage.ActionCallback = UE::Landscape::BuildNanite;
+		OutActionableMessage.ActionCallback = [LogFinalNumberOfDirtyProxies]()
+			{ 
+				UE::Landscape::BuildNanite(UE::Landscape::EBuildFlags::WriteFinalLog); 
+				LogFinalNumberOfDirtyProxies();
+			};
 
 		return true;
 	}
@@ -1164,7 +1348,11 @@ bool ULandscapeSubsystem::GetActionableMessage(FActionableMessage& OutActionable
 		OutActionableMessage.Message = FText::Format(LOCTEXT("LandscapeModified.Message", "{0} Landscape {0}|plural(one=actor,other=actors) {0}|plural(one=is,other=are) out of date and {0}|plural(one=needs,other=need) to be rebuilt"), OutdatedProxiesCount);
 		OutActionableMessage.Tooltip = LOCTEXT("LandscapeModified.Tooltip", "The Landscape actors visible in your level have been modified as a result of changes to other assets.\nThese changes need to be applied to the Landscape assets.");
 		OutActionableMessage.ActionMessage = LOCTEXT("LandscapeModified.Action", "Update");
-		OutActionableMessage.ActionCallback = UE::Landscape::MarkModifiedLandscapesAsDirty;
+		OutActionableMessage.ActionCallback = [LogFinalNumberOfDirtyProxies]()
+			{ 
+				UE::Landscape::MarkModifiedLandscapesAsDirty(UE::Landscape::EBuildFlags::WriteFinalLog); 
+				LogFinalNumberOfDirtyProxies();
+			};
 
 		return true;
 	}
diff --git a/Engine/Source/Runtime/Landscape/Private/LandscapeUtils.cpp b/Engine/Source/Runtime/Landscape/Private/LandscapeUtils.cpp
--- a/Engine/Source/Runtime/Landscape/Private/LandscapeUtils.cpp
+++ b/Engine/Source/Runtime/Landscape/Private/LandscapeUtils.cpp
@@ -7,7 +7,6 @@
 #include "LandscapeLayerInfoObject.h"
 #include "LandscapeProxy.h"
 #include "LandscapeEditTypes.h"
-#include "LandscapeSubsystem.h"
 #include "Algo/Transform.h"
 
 #if WITH_EDITOR
@@ -290,35 +289,6 @@ ULandscapeLayerInfoObject* FLayerInfoFinder::Find(const FName& LayerName) const
 	return nullptr;
 }
 
-bool HasModifiedLandscapes()
-{
-	if (GEditor)
-	{
-		if (UWorld* World = GEditor->GetEditorWorldContext().World())
-		{
-			if (ULandscapeSubsystem* LandscapeSubsystem = World->GetSubsystem<ULandscapeSubsystem>())
-			{
-				return LandscapeSubsystem->HasModifiedLandscapes();
-			}
-		}
-	}
-
-	return false;
-}
-
-void SaveModifiedLandscapes()
-{
-	if (GEditor)
-	{
-		if (UWorld* World = GEditor->GetEditorWorldContext().World())
-		{
-			if (ULandscapeSubsystem* LandscapeSubsystem = World->GetSubsystem<ULandscapeSubsystem>())
-			{
-				LandscapeSubsystem->SaveModifiedLandscapes();
-			}
-		}
-	}
-}
 #endif // WITH_EDITOR
 
 } // end namespace UE::Landscape
\ No newline at end of file
diff --git a/Engine/Source/Runtime/Landscape/Public/LandscapeEditTypes.h b/Engine/Source/Runtime/Landscape/Public/LandscapeEditTypes.h
--- a/Engine/Source/Runtime/Landscape/Public/LandscapeEditTypes.h
+++ b/Engine/Source/Runtime/Landscape/Public/LandscapeEditTypes.h
@@ -64,4 +64,13 @@ inline uint32 GetOutdatedDataFlagIndex(EOutdatedDataFlags InFlag)
 	return FMath::CountTrailingZeros(InFlagAsUInt32);
 }
 
+enum class EBuildFlags : uint8
+{
+	None = 0, 
+
+	WriteFinalLog = (1 << 0), // Log a message at the end of the build operation
+	ForceRebuild = (1 << 1), // Forces a rebuild even if not needed
+};
+ENUM_CLASS_FLAGS(EBuildFlags)
+
 } // namespace UE::Landscape
\ No newline at end of file
diff --git a/Engine/Source/Runtime/Landscape/Public/LandscapeSubsystem.h b/Engine/Source/Runtime/Landscape/Public/LandscapeSubsystem.h
--- a/Engine/Source/Runtime/Landscape/Public/LandscapeSubsystem.h
+++ b/Engine/Source/Runtime/Landscape/Public/LandscapeSubsystem.h
@@ -27,13 +27,34 @@ enum class EUpdateTransformFlags : int32;
 namespace UE::Landscape
 {
 	enum class EOutdatedDataFlags : uint8;
+	enum class EBuildFlags : uint8;
 
 #if WITH_EDITOR
+	/** Returns true if there are some landscapes in the editor world that have been automatically modified and are in need of being saved (see LandscapeDirtyingMode) */
+	bool LANDSCAPE_API HasModifiedLandscapes();
+
+	/** Dirties and saves the landscapes in the editor world that have been automatically modified and are in need of being saved (see LandscapeDirtyingMode) */
+	void LANDSCAPE_API SaveModifiedLandscapes(EBuildFlags InBuildFlags);
+
+	UE_DEPRECATED(5.6, "Use the function with the EBuildFlags param")
 	void LANDSCAPE_API MarkModifiedLandscapesAsDirty();
+	void LANDSCAPE_API MarkModifiedLandscapesAsDirty(EBuildFlags InBuildFlags);
+
+	UE_DEPRECATED(5.6, "Use the function with the EBuildFlags param")
 	void LANDSCAPE_API BuildGrassMaps();
+	void LANDSCAPE_API BuildGrassMaps(EBuildFlags InBuildFlags);
+
+	UE_DEPRECATED(5.6, "Use the function with the EBuildFlags param")
 	void LANDSCAPE_API BuildPhysicalMaterial();
+	void LANDSCAPE_API BuildPhysicalMaterial(EBuildFlags InBuildFlags);
+
+	UE_DEPRECATED(5.6, "Use the function with the EBuildFlags param")
 	void LANDSCAPE_API BuildNanite();
+	void LANDSCAPE_API BuildNanite(EBuildFlags InBuildFlags);
+
+	UE_DEPRECATED(5.6, "Use the function with the EBuildFlags param")
 	void LANDSCAPE_API BuildAll();
+	void LANDSCAPE_API BuildAll(EBuildFlags InBuildFlags);
 #endif // WITH_EDITOR
 } // end of namespace UE::Landscape
 
@@ -100,18 +121,27 @@ public:
 	void UnregisterComponent(ULandscapeComponent* Component);
 
 #if WITH_EDITOR
+	UE_DEPRECATED(5.6, "Use the function with the EBuildFlags param")
 	LANDSCAPE_API void BuildAll();
+	LANDSCAPE_API void BuildAll(UE::Landscape::EBuildFlags InBuildFlags);
 
 	// Synchronously build grass maps for all components
+	UE_DEPRECATED(5.6, "Use the function with the EBuildFlags param")
 	LANDSCAPE_API void BuildGrassMaps();
+	LANDSCAPE_API void BuildGrassMaps(UE::Landscape::EBuildFlags InBuildFlags);
+
+	UE_DEPRECATED(5.6, "Use the function with the EBuildFlags param")
 	LANDSCAPE_API void BuildPhysicalMaterial();
+	LANDSCAPE_API void BuildPhysicalMaterial(UE::Landscape::EBuildFlags InBuildFlags);
 
+	UE_DEPRECATED(5.6, "Use the function with the EBuildFlags param")
+	LANDSCAPE_API void BuildNanite(TArrayView<ALandscapeProxy*> InProxiesToBuild = TArrayView<ALandscapeProxy*>(), bool bForceRebuild = false);
 	/**
 	 * Updates the Nanite mesh on all landscape actors whose mesh is not up to date.
 	 * @param InProxiesToBuild - If specified, only the Nanite meshes of the specified landscape actors (recursively for all streaming proxies, in the case of a 1 ALandscape / N ALandscapeStreamingProxy setup) will be built
 	 * @param bForceRebuild - If true, forces the Nanite meshes to be rebuilt, no matter if they're up to date or not
 	 */
-	LANDSCAPE_API void BuildNanite(TArrayView<ALandscapeProxy*> InProxiesToBuild = TArrayView<ALandscapeProxy*>(), bool bForceRebuild = false);
+	LANDSCAPE_API void BuildNanite(UE::Landscape::EBuildFlags InBuildFlags, TArrayView<ALandscapeProxy*> InProxiesToBuild = TArrayView<ALandscapeProxy*>());
 	
 	UE_DEPRECATED(5.5, "Use GetOutdatedProxyDetails")
 	LANDSCAPE_API TArray<ALandscapeProxy*> GetOutdatedProxies(UE::Landscape::EOutdatedDataFlags InMatchingOutdatedDataFlags, bool bInMustMatchAllFlags) const;
@@ -125,8 +155,14 @@ public:
 	LANDSCAPE_API void DisplayMessages(class FCanvas* Canvas, float& XPos, float& YPos);
 
 	LANDSCAPE_API bool GetActionableMessage(FActionableMessage& OutActionableMessage);
+	UE_DEPRECATED(5.6, "Use the function with the EBuildFlags param")
 	LANDSCAPE_API void MarkModifiedLandscapesAsDirty();
+	LANDSCAPE_API void MarkModifiedLandscapesAsDirty(UE::Landscape::EBuildFlags InBuildFlags);
+	UE_DEPRECATED(5.6, "Use the function with the EBuildFlags param")
 	LANDSCAPE_API void SaveModifiedLandscapes();
+	/** Dirties and saves the landscapes in the editor world that have been automatically modified and are in need of being saved (see LandscapeDirtyingMode) */
+	LANDSCAPE_API void SaveModifiedLandscapes(UE::Landscape::EBuildFlags InBuildFlags);
+	/** Returns true if there are some landscapes in the editor world that have been automatically modified and are in need of being saved (see LandscapeDirtyingMode) */
 	LANDSCAPE_API bool HasModifiedLandscapes() const;
 	LANDSCAPE_API bool GetDirtyOnlyInMode() const;
 	FLandscapeNotificationManager* GetNotificationManager() { return NotificationManager; }
@@ -167,7 +203,9 @@ private:
 	void OnNaniteEnabledChanged(IConsoleVariable*);
 
 	void HandlePostGarbageCollect();
-	
+
+	TSet<UPackage*> GetDirtyLandscapeProxyPackages() const;
+
 	// When proxies move, we may need to update their position in the landscape group
 	void OnProxyMoved(USceneComponent*, EUpdateTransformFlags, ETeleportType);
 
diff --git a/Engine/Source/Runtime/Landscape/Public/LandscapeUtils.h b/Engine/Source/Runtime/Landscape/Public/LandscapeUtils.h
--- a/Engine/Source/Runtime/Landscape/Public/LandscapeUtils.h
+++ b/Engine/Source/Runtime/Landscape/Public/LandscapeUtils.h
@@ -18,6 +18,7 @@ enum class ELandscapeToolTargetTypeFlags : uint8;
 
 namespace UE::Landscape
 {
+enum class EBuildFlags : uint8;
 
 /**
 * Returns true if edit layers (GPU landscape tools) are enabled on this platform :
@@ -113,13 +114,6 @@ struct FLayerInfoFinder
 	TArray<FAssetData> LayerInfoAssets;
 };
 
-
-/** Returns true if there are some landscapes in the editor world that have been automatically modified and are in need of being saved (see LandscapeDirtyingMode) */
-LANDSCAPE_API bool HasModifiedLandscapes();
-
-/** Dirties and saves the landscapes in the editor world that have been automatically modified and are in need of being saved (see LandscapeDirtyingMode) */
-LANDSCAPE_API void SaveModifiedLandscapes();
-
 #endif // WITH_EDITOR
 
 } // end namespace UE::Landscape
