Remove unnecessary include of "Serialization/BulkData.h"

#rnx

- Deprecated the includes for public headers and just removed them from internal/private ones.

diff --git a/Engine/Plugins/FX/Niagara/Source/Niagara/Internal/DataInterface/NiagaraDataInterfaceMemoryBuffer.h b/Engine/Plugins/FX/Niagara/Source/Niagara/Internal/DataInterface/NiagaraDataInterfaceMemoryBuffer.h
--- a/Engine/Plugins/FX/Niagara/Source/Niagara/Internal/DataInterface/NiagaraDataInterfaceMemoryBuffer.h
+++ b/Engine/Plugins/FX/Niagara/Source/Niagara/Internal/DataInterface/NiagaraDataInterfaceMemoryBuffer.h
@@ -8,8 +8,6 @@
 #include "NiagaraRenderableMeshInterface.h"
 #include "NiagaraSimCacheCustomStorageInterface.h"
 
-#include "Serialization/BulkData.h"
-#include "Serialization/BulkDataBuffer.h"
 #include "VectorVM.h"
 
 #include "NiagaraDataInterfaceMemoryBuffer.generated.h"
diff --git a/Engine/Plugins/Runtime/HairStrands/Source/HairStrandsCore/Public/HairDescription.h b/Engine/Plugins/Runtime/HairStrands/Source/HairStrandsCore/Public/HairDescription.h
--- a/Engine/Plugins/Runtime/HairStrands/Source/HairStrandsCore/Public/HairDescription.h
+++ b/Engine/Plugins/Runtime/HairStrands/Source/HairStrandsCore/Public/HairDescription.h
@@ -8,9 +8,12 @@
 #include "MeshAttributeArray.h" // for TMeshAttributesRef
 #include "MeshDescription.h" // for TVertexAttributesRef
 #include "MeshTypes.h" // for FElementID, FVertexID
-#include "Serialization/BulkData.h"
 #include "Serialization/EditorBulkData.h"
 
+#if UE_ENABLE_INCLUDE_ORDER_DEPRECATED_IN_5_6
+#include "Serialization/BulkData.h"
+#endif
+
 struct FStrandID : public FElementID
 {
 	FStrandID()
diff --git a/Engine/Source/Runtime/Engine/Classes/Atmosphere/AtmosphericFogComponent.h b/Engine/Source/Runtime/Engine/Classes/Atmosphere/AtmosphericFogComponent.h
--- a/Engine/Source/Runtime/Engine/Classes/Atmosphere/AtmosphericFogComponent.h
+++ b/Engine/Source/Runtime/Engine/Classes/Atmosphere/AtmosphericFogComponent.h
@@ -6,8 +6,12 @@
 #include "HAL/ThreadSafeCounter.h"
 #include "UObject/ObjectMacros.h"
 #include "Components/SceneComponent.h"
-#include "Serialization/BulkData.h"
 #include "Components/SkyAtmosphereComponent.h"
+
+#if UE_ENABLE_INCLUDE_ORDER_DEPRECATED_IN_5_6
+#include "Serialization/BulkData.h"
+#endif
+
 #include "AtmosphericFogComponent.generated.h"
 
 /**
diff --git a/Engine/Source/Runtime/Engine/Classes/Engine/Texture2D.h b/Engine/Source/Runtime/Engine/Classes/Engine/Texture2D.h
--- a/Engine/Source/Runtime/Engine/Classes/Engine/Texture2D.h
+++ b/Engine/Source/Runtime/Engine/Classes/Engine/Texture2D.h
@@ -9,7 +9,11 @@
 #include "UObject/ScriptMacros.h"
 #include "Engine/Texture.h"
 #include "Engine/TextureAllMipDataProviderFactory.h"
+
+#if UE_ENABLE_INCLUDE_ORDER_DEPRECATED_IN_5_6
 #include "Serialization/BulkData.h"
+#endif
+
 #include "Texture2D.generated.h"
 
 class FTexture2DResourceMem;
diff --git a/Engine/Source/Runtime/Engine/Private/PhysicsEngine/PhysXSupport.h b/Engine/Source/Runtime/Engine/Private/PhysicsEngine/PhysXSupport.h
--- a/Engine/Source/Runtime/Engine/Private/PhysicsEngine/PhysXSupport.h
+++ b/Engine/Source/Runtime/Engine/Private/PhysicsEngine/PhysXSupport.h
@@ -14,7 +14,6 @@
 #include "Physics/PhysicsFiltering.h"
 #include "PhysXPublic.h"
 #include "PhysXSupportCore.h"
-#include "Serialization/BulkData.h"
 
 class UBodySetup;
 class UPhysicalMaterial;
diff --git a/Engine/Source/Runtime/Landscape/Classes/LandscapeComponent.h b/Engine/Source/Runtime/Landscape/Classes/LandscapeComponent.h
--- a/Engine/Source/Runtime/Landscape/Classes/LandscapeComponent.h
+++ b/Engine/Source/Runtime/Landscape/Classes/LandscapeComponent.h
@@ -9,13 +9,17 @@
 #include "Engine/TextureStreamingTypes.h"
 #include "Components/PrimitiveComponent.h"
 #include "UObject/PerPlatformProperties.h"
-#include "Serialization/BulkData.h"
 #include "LandscapePhysicalMaterial.h"
 #include "LandscapeInfo.h"
 #include "LandscapeWeightmapUsage.h"
 #include "Containers/ArrayView.h"
 #include "Engine/StreamableRenderAsset.h"
 #include "Engine/Texture2DArray.h"
+
+#if UE_ENABLE_INCLUDE_ORDER_DEPRECATED_IN_5_6
+#include "Serialization/BulkData.h"
+#endif
+
 #include "LandscapeComponent.generated.h"
 
 class ALandscape;
