Update the physical material mask correctly when unpainting or clearing weightmap.  Catch the component weightmaps that have lost their last non-zero texels and follow through with the correct update notifications.

diff --git a/Engine/Source/Runtime/Landscape/Private/LandscapeEditLayers.cpp b/Engine/Source/Runtime/Landscape/Private/LandscapeEditLayers.cpp
--- a/Engine/Source/Runtime/Landscape/Private/LandscapeEditLayers.cpp
+++ b/Engine/Source/Runtime/Landscape/Private/LandscapeEditLayers.cpp
@@ -3196,6 +3196,7 @@ struct FLandscapeEditLayerComponentReadbackResult
 	uint32 UpdateModes = 0;
 	/** Were the associated heightmap/weightmaps modified. */
 	bool bModified = false;
+	bool bCleared = false;
 	/** Indicates which of the component's weightmaps is not needed anymore. */
 	TArray<ULandscapeLayerInfoObject*> AllZeroLayers;
 
@@ -4347,6 +4348,21 @@ bool ALandscape::PrepareLayersResources(ERHIFeatureLevel::Type InFeatureLevel, b
 	return bIsReady;
 }
 
+namespace UE::Landscape::Private
+{
+// Find a matching entry in Array for the given Component.  If one does not exist, add one with the provided UpdateModes value.
+FLandscapeEditLayerComponentReadbackResult* FindOrAddByComponent(TArray<FLandscapeEditLayerComponentReadbackResult>& Array, ULandscapeComponent* Component, uint32 DefaultUpdateModes)
+{
+	FLandscapeEditLayerComponentReadbackResult* ComponentReadbackResult = Algo::FindBy(Array, Component, &FLandscapeEditLayerComponentReadbackResult::LandscapeComponent);
+	if (ComponentReadbackResult == nullptr)
+	{
+		ComponentReadbackResult = &Array.Emplace_GetRef(Component, DefaultUpdateModes);
+	}
+	return ComponentReadbackResult;
+}
+}  // UE::Landscape::Private
+
+
 // Must match FEditLayerHeightmapMergeInfo in LandscapeLayersHeightmapsPS.usf
 struct FLandscapeEditLayerHeightmapMergeInfo
 {
@@ -4588,11 +4604,7 @@ struct FUpdateLayersContentContext
 			{
 				if (Component->GetLayerUpdateFlagPerMode() & ELandscapeLayerUpdateMode::Update_Client_Deferred)
 				{
-					FLandscapeEditLayerComponentReadbackResult* ComponentReadbackResult = AllLandscapeComponentReadbackResults.FindByPredicate([Component](const FLandscapeEditLayerComponentReadbackResult& Element) { return Element.LandscapeComponent == Component; });
-					if (ComponentReadbackResult == nullptr)
-					{
-						AllLandscapeComponentReadbackResults.Add(FLandscapeEditLayerComponentReadbackResult(Component, ELandscapeLayerUpdateMode::Update_Client_Deferred));
-					}
+					UE::Landscape::Private::FindOrAddByComponent(AllLandscapeComponentReadbackResults, Component, ELandscapeLayerUpdateMode::Update_Client_Deferred);
 				}
 			}
 		}
@@ -8464,11 +8476,7 @@ bool ALandscape::ResolveLayersTexture(
 				ULandscapeComponent** Component = GetLandscapeInfo()->XYtoComponentMap.Find(ComponentContext.ComponentKey);
 				if (Component != nullptr && *Component != nullptr)
 				{
-					FLandscapeEditLayerComponentReadbackResult* ComponentReadbackResult = InOutComponentReadbackResults.FindByPredicate([LandscapeComponent = *Component](const FLandscapeEditLayerComponentReadbackResult& Element) { return Element.LandscapeComponent == LandscapeComponent; });
-					if (ComponentReadbackResult == nullptr)
-					{
-						ComponentReadbackResult = &InOutComponentReadbackResults.Add_GetRef(FLandscapeEditLayerComponentReadbackResult(*Component, ELandscapeLayerUpdateMode::Update_None));
-					}
+					FLandscapeEditLayerComponentReadbackResult* ComponentReadbackResult = UE::Landscape::Private::FindOrAddByComponent(InOutComponentReadbackResults, *Component, ELandscapeLayerUpdateMode::Update_None);
 					ComponentReadbackResult->UpdateModes |= ComponentContext.UpdateModes;
 					ComponentReadbackResult->bModified |= bChanged ? 1 : 0;
 				}
@@ -8912,6 +8920,15 @@ void ALandscape::ReallocateLayersWeightmaps(FUpdateLayersContentContext& InUpdat
 		return true;
 	});
 
+	// When the last pixels are removed from a component it loses its allocations, which means we stop tracking it through the update pipeline.  Add a tracking
+	// object if needed and mark it with bCleared so that the code later on will know to update the collision object.
+	for (ULandscapeComponent* Component : ComponentsWithNoAllocs)
+	{
+		FLandscapeEditLayerComponentReadbackResult* ComponentReadbackResult =  UE::Landscape::Private::FindOrAddByComponent(InUpdateLayersContentContext.AllLandscapeComponentReadbackResults, Component, ELandscapeLayerUpdateMode::Update_None);
+		ComponentReadbackResult->bCleared = true;
+		ComponentReadbackResult->bModified = true;
+	}
+
 	InvalidateGeneratedComponentData(ComponentsWithNoAllocs, false);
 	ValidateProxyLayersWeightmapUsage();
 
@@ -10801,7 +10818,8 @@ void ALandscape::UpdateForChangedWeightmaps(const TArrayView<FLandscapeEditLayer
 		}
 
 		const int32 WeightUpdateMode = ComponentReadbackResult.UpdateModes & (ELandscapeLayerUpdateMode::Update_Weightmap_All | ELandscapeLayerUpdateMode::Update_Weightmap_Editing | ELandscapeLayerUpdateMode::Update_Weightmap_Editing_NoCollision);
-		if (IsUpdateFlagEnabledForModes(ELandscapeComponentUpdateFlag::Component_Update_Weightmap_Collision, WeightUpdateMode))
+		if (IsUpdateFlagEnabledForModes(ELandscapeComponentUpdateFlag::Component_Update_Weightmap_Collision, WeightUpdateMode) 
+			|| ComponentReadbackResult.bCleared)
 		{
 			// Only update collision data if there was an actual change performed on the source data : 
 			if (ComponentReadbackResult.LandscapeComponent->GetPendingLayerCollisionDataUpdate())
@@ -11702,7 +11720,8 @@ uint32 ALandscape::UpdateCollisionAndClients(const TArrayView<FLandscapeEditLaye
 		
 		bool bDeferClientUpdateForComponent = false;
 		bool bDoUpdateClient = true;
-		if (IsUpdateFlagEnabledForModes(ELandscapeComponentUpdateFlag::Component_Update_Recreate_Collision, ComponentReadbackResult.UpdateModes))
+		if (IsUpdateFlagEnabledForModes(ELandscapeComponentUpdateFlag::Component_Update_Recreate_Collision, ComponentReadbackResult.UpdateModes)
+			|| ComponentReadbackResult.bCleared)
 		{
 			if (ULandscapeHeightfieldCollisionComponent* CollisionComp = LandscapeComponent->GetCollisionComponent())
 			{
