Don't register landscape proxies to USceneComponent::TransformUpdated if cooking or if the root component is null.
Also unregister landscape proxies from USceneComponent::TransformUpdated when removed from world.

diff --git a/Engine/Source/Runtime/Landscape/Private/LandscapeSubsystem.cpp b/Engine/Source/Runtime/Landscape/Private/LandscapeSubsystem.cpp
--- a/Engine/Source/Runtime/Landscape/Private/LandscapeSubsystem.cpp
+++ b/Engine/Source/Runtime/Landscape/Private/LandscapeSubsystem.cpp
@@ -244,10 +244,16 @@ void ULandscapeSubsystem::RegisterActor(ALandscapeProxy* Proxy)
 		TObjectPtr<ALandscape> LandscapeActorPtr(LandscapeActor);
 		LandscapeActors.AddUnique(LandscapeActorPtr);
 	}
-	else if (ALandscapeStreamingProxy* StreamingProxy = Cast<ALandscapeStreamingProxy>(Proxy))
+	else if (!IsRunningCookCommandlet())
 	{
-		// We want to know when a streaming proxy is moved, so we can adjust the map
-		StreamingProxy->GetRootComponent()->TransformUpdated.AddUObject(this, &ULandscapeSubsystem::OnProxyMoved);
+		if (ALandscapeStreamingProxy* StreamingProxy = Cast<ALandscapeStreamingProxy>(Proxy))
+		{
+			// We want to know when a streaming proxy is moved, so we can adjust the map
+			if (USceneComponent* RootComponent = StreamingProxy->GetRootComponent())
+			{
+				RootComponent->TransformUpdated.AddUObject(this, &ULandscapeSubsystem::OnProxyMoved);
+			}
+		}
 	}
 
 	Proxy->bIsRegisteredWithSubsystem = true;
@@ -286,9 +292,18 @@ void ULandscapeSubsystem::UnregisterActor(ALandscapeProxy* Proxy)
 		TObjectPtr<ALandscape> LandscapeActorPtr(LandscapeActor);
 		LandscapeActors.Remove(LandscapeActorPtr);
 	}
-	else if (ALandscapeStreamingProxy* StreamingProxy = Cast<ALandscapeStreamingProxy>(Proxy))
+	else if (!IsRunningCookCommandlet())
 	{
-		StreamingProxiesNeedingReregister.Remove(StreamingProxy);
+		if (ALandscapeStreamingProxy* StreamingProxy = Cast<ALandscapeStreamingProxy>(Proxy))
+		{
+			StreamingProxiesNeedingReregister.Remove(StreamingProxy);
+
+			// We want to know when a streaming proxy is moved, so we can adjust the map
+			if (USceneComponent* RootComponent = StreamingProxy->GetRootComponent())
+			{
+				RootComponent->TransformUpdated.RemoveAll(this);
+			}
+		}
 	}
 
 	Proxy->bIsRegisteredWithSubsystem = false;
