Fixed up layer info objects that have been re-mapped to a new asset in the main landscape actor without re-saving all the associated proxies : in such a case, the layer name can be used to retrieve the new layer info object from the main landscape actor so we can do the fixup.
Also, those layer fixups now mark the landscape proxy as "soft-"dirty (i.e. the landscape way of notifying the user that a re-save is needed) in order to make it more visible to the user that a user action is needed

Misc:
* Normalized format message logs issued by landscape

diff --git a/Engine/Source/Runtime/Core/Private/Misc/MapErrors.cpp b/Engine/Source/Runtime/Core/Private/Misc/MapErrors.cpp
--- a/Engine/Source/Runtime/Core/Private/Misc/MapErrors.cpp
+++ b/Engine/Source/Runtime/Core/Private/Misc/MapErrors.cpp
@@ -33,6 +33,7 @@ FLazyName FMapErrors::DeprecatedClass(TEXT("DeprecatedClass"));
 FLazyName FMapErrors::FoliageMissingStaticMesh(TEXT("FoliageMissingStaticMesh"));
 FLazyName FMapErrors::FoliageMissingClusterComponent(TEXT("FoliageMissingStaticMesh"));
 FLazyName FMapErrors::FixedUpDeletedLayerWeightmap(TEXT("FixedUpDeletedLayerWeightmap"));
+FLazyName FMapErrors::FixedUpNonMatchingTargetLayer(TEXT("FixedUpNonMatchingTargetLayer"));
 FLazyName FMapErrors::FixedUpInvalidLandscapeMaterialInstances(TEXT("FixedUpInvalidLandscapeMaterialInstances"));
 FLazyName FMapErrors::FixedUpIncorrectLayerWeightmap(TEXT("FixedUpIncorrectLayerWeightmap"));
 FLazyName FMapErrors::FixedUpSharedLayerWeightmap(TEXT("FixedUpSharedLayerWeightmap"));
diff --git a/Engine/Source/Runtime/Core/Public/Misc/MapErrors.h b/Engine/Source/Runtime/Core/Public/Misc/MapErrors.h
--- a/Engine/Source/Runtime/Core/Public/Misc/MapErrors.h
+++ b/Engine/Source/Runtime/Core/Public/Misc/MapErrors.h
@@ -111,6 +111,9 @@ struct FMapErrors
 	/** {LandscapeComponent} : Fixed up deleted layer weightmap: */
     static CORE_API FLazyName FixedUpDeletedLayerWeightmap;
     
+	/** {LandscapeComponent} : Fixed up target layer info object not matching with the parent landscape's: */
+	static CORE_API FLazyName FixedUpNonMatchingTargetLayer;
+
 	/** {LandscapeComponent} : Fixed up incorrect layer weightmap texture index: */
     static CORE_API FLazyName FixedUpIncorrectLayerWeightmap;
     
diff --git a/Engine/Source/Runtime/Landscape/Private/Landscape.cpp b/Engine/Source/Runtime/Landscape/Private/Landscape.cpp
--- a/Engine/Source/Runtime/Landscape/Private/Landscape.cpp
+++ b/Engine/Source/Runtime/Landscape/Private/Landscape.cpp
@@ -4382,19 +4382,16 @@ PRAGMA_ENABLE_DEPRECATION_WARNINGS
 	if (bFixedUpInvalidMaterialInstances)
 	{
 		FFormatNamedArguments Arguments;
-		Arguments.Add(TEXT("LandscapeName"), FText::FromString(GetPathName()));
 		Arguments.Add(TEXT("ProxyPackage"), FText::FromString(GetOutermost()->GetName()));
 		FMessageLog("MapCheck").Info()
-			->AddToken(FTextToken::Create(FText::Format(LOCTEXT("MapCheck_Message_FixedUpInvalidLandscapeMaterialInstances", "{LandscapeName} : Fixed up invalid landscape material instances. Please re-save {ProxyPackage}."), Arguments)))
+			->AddToken(FUObjectToken::Create(this, FText::FromString(GetActorNameOrLabel())))
+			->AddToken(FTextToken::Create(FText::Format(LOCTEXT("MapCheck_Message_FixedUpInvalidLandscapeMaterialInstances", "Fixed up invalid landscape material instances. Please re-save {ProxyPackage}."), Arguments)))
 			->AddToken(FMapErrorToken::Create(FMapErrors::FixedUpInvalidLandscapeMaterialInstances));
 	}
 
 	// Display a MapCheck warning if the Nanite data is stale with the option to trigger a rebuild & Save
 	if (!IsNaniteMeshUpToDate() && !IsRunningCookCommandlet())
 	{
-		FFormatNamedArguments Arguments;
-		Arguments.Add(TEXT("LandscapeProxyName"), FText::FromString(GetActorNameOrLabel()));
-
 		auto CreateMapCheckMessage = [](FMessageLog& MessageLog)
 		{
 			if (CVarLandscapeSupressMapCheckWarnings_Nanite->GetBool())
@@ -4408,40 +4405,41 @@ PRAGMA_ENABLE_DEPRECATION_WARNINGS
 
 		FMessageLog MessageLog("MapCheck");
 		CreateMapCheckMessage(MessageLog)
-			->AddToken(FTextToken::Create(FText::Format(LOCTEXT("MapCheck_Message_LandscapeRebuildNanite", "{LandscapeProxyName} : Landscape Nanite is enabled but the saved mesh data is out of date. "), Arguments)))
+			->AddToken(FUObjectToken::Create(this, FText::FromString(GetActorNameOrLabel())))
+			->AddToken(FTextToken::Create(LOCTEXT("MapCheck_Message_LandscapeRebuildNanite", "Landscape Nanite is enabled but the saved mesh data is out of date. ")))
 			->AddToken(FActionToken::Create(LOCTEXT("MapCheck_SaveFixedUpData", "Save Modified Landscapes"), LOCTEXT("MapCheck_SaveFixedUpData_Desc", "Saves the modified landscape proxy actors"),
 				FOnActionTokenExecuted::CreateLambda([WeakLandscapeProxy]()
-		{
-			if (!WeakLandscapeProxy.IsValid())
-			{
-				return;
-			}
-			ULandscapeInfo* Info = WeakLandscapeProxy->GetLandscapeInfo();
-			check(Info);
+					{
+						if (!WeakLandscapeProxy.IsValid())
+						{
+							return;
+						}
+						ULandscapeInfo* Info = WeakLandscapeProxy->GetLandscapeInfo();
+						check(Info);
 
-			TSet<UPackage*> DirtyNanitePackages;
-			Info->ForEachLandscapeProxy([&DirtyNanitePackages](const ALandscapeProxy* Proxy)
-			{
-				if (!Proxy->IsNaniteMeshUpToDate())
-				{
-					DirtyNanitePackages.Add(Proxy->GetOutermost());
-				}
-				return true;
-			});
+						TSet<UPackage*> DirtyNanitePackages;
+						Info->ForEachLandscapeProxy([&DirtyNanitePackages](const ALandscapeProxy* Proxy)
+							{
+								if (!Proxy->IsNaniteMeshUpToDate())
+								{
+									DirtyNanitePackages.Add(Proxy->GetOutermost());
+								}
+								return true;
+							});
 
-			Info->UpdateNanite(nullptr);
+						Info->UpdateNanite(nullptr);
 
-			constexpr bool bPromptUserToSave = true;
-			constexpr bool bSaveMapPackages = true;
-			constexpr bool bSaveContentPackages = true;
-			constexpr bool bFastSave = false;
-			constexpr bool bNotifyNoPackagesSaved = false;
-			constexpr bool bCanBeDeclined = true;
+						constexpr bool bPromptUserToSave = true;
+						constexpr bool bSaveMapPackages = true;
+						constexpr bool bSaveContentPackages = true;
+						constexpr bool bFastSave = false;
+						constexpr bool bNotifyNoPackagesSaved = false;
+						constexpr bool bCanBeDeclined = true;
 
-			FEditorFileUtils::SaveDirtyPackages(bPromptUserToSave, bSaveMapPackages, bSaveContentPackages, bFastSave, bNotifyNoPackagesSaved, bCanBeDeclined, nullptr,
-				[&DirtyNanitePackages](const UPackage* Package) { return !DirtyNanitePackages.Contains(Package); });
+						FEditorFileUtils::SaveDirtyPackages(bPromptUserToSave, bSaveMapPackages, bSaveContentPackages, bFastSave, bNotifyNoPackagesSaved, bCanBeDeclined, nullptr,
+							[&DirtyNanitePackages](const UPackage* Package) { return !DirtyNanitePackages.Contains(Package); });
 
-		}), FCanExecuteActionToken::CreateLambda([WeakLandscapeProxy]() { return WeakLandscapeProxy.IsValid() ? !WeakLandscapeProxy->IsNaniteMeshUpToDate() : false; }))
+					}), FCanExecuteActionToken::CreateLambda([WeakLandscapeProxy]() { return WeakLandscapeProxy.IsValid() ? !WeakLandscapeProxy->IsNaniteMeshUpToDate() : false; }))
 			);
 	}
 
@@ -6116,7 +6114,7 @@ void ULandscapeInfo::RegisterActorComponent(ULandscapeComponent* Component, bool
 			Arguments.Add(TEXT("XLocation"), Component->GetSectionBase().X);
 			Arguments.Add(TEXT("YLocation"), Component->GetSectionBase().Y);
 			FMessageLog("MapCheck").Warning()
-				->AddToken(FUObjectToken::Create(OurProxy))
+				->AddToken(FUObjectToken::Create(OurProxy, FText::FromString(OurProxy->GetActorNameOrLabel())))
 				->AddToken(FTextToken::Create(FText::Format(LOCTEXT("MapCheck_Message_LandscapeComponentPostLoad_Warning", "Landscape {ProxyName1} of {LevelName1} has overlapping render components with {ProxyName2} of {LevelName2} at location ({XLocation}, {YLocation})."), Arguments)))
 				->AddToken(FActionToken::Create(LOCTEXT("MapCheck_RemoveDuplicateLandscapeComponent", "Delete Duplicate"), LOCTEXT("MapCheck_RemoveDuplicateLandscapeComponentDesc", "Deletes the duplicate landscape component."), FOnActionTokenExecuted::CreateUObject(OurProxy, &ALandscapeProxy::RemoveOverlappingComponent, Component), true))
 				->AddToken(FMapErrorToken::Create(FMapErrors::LandscapeComponentPostLoad_Warning));
diff --git a/Engine/Source/Runtime/Landscape/Private/LandscapeBlueprintBrushBase.cpp b/Engine/Source/Runtime/Landscape/Private/LandscapeBlueprintBrushBase.cpp
--- a/Engine/Source/Runtime/Landscape/Private/LandscapeBlueprintBrushBase.cpp
+++ b/Engine/Source/Runtime/Landscape/Private/LandscapeBlueprintBrushBase.cpp
@@ -542,7 +542,7 @@ void ALandscapeBlueprintBrushBase::CheckForErrors()
 		if (OwningLandscape == nullptr)
 		{
 			FMessageLog("MapCheck").Error()
-				->AddToken(FUObjectToken::Create(this))
+				->AddToken(FUObjectToken::Create(this, FText::FromString(GetActorNameOrLabel())))
 				->AddToken(FTextToken::Create(LOCTEXT("MapCheck_Message_MissingLandscape", "This brush requires a Landscape. Add one to the map or remove the brush actor.")))
 				->AddToken(FMapErrorToken::Create(TEXT("LandscapeBrushMissingLandscape")));
 		}
diff --git a/Engine/Source/Runtime/Landscape/Private/LandscapeEdit.cpp b/Engine/Source/Runtime/Landscape/Private/LandscapeEdit.cpp
--- a/Engine/Source/Runtime/Landscape/Private/LandscapeEdit.cpp
+++ b/Engine/Source/Runtime/Landscape/Private/LandscapeEdit.cpp
@@ -1159,11 +1159,9 @@ void ALandscapeProxy::RepairInvalidTextures()
 
 	if (!InvalidTextures.IsEmpty())
 	{
-		FFormatNamedArguments Arguments;
-		Arguments.Add(TEXT("LandscapeName"), FText::FromString(GetPathName()));
-		Arguments.Add(TEXT("ErrorMessage"), FText::Format(LOCTEXT("InvalidTexturesMessage", "Invalid data detected on {0} {0}|plural(one = texture, other = textures). The data has been cleared to avoid fatal error."), InvalidTextures.Num()));
 		FMessageLog("MapCheck").Error()
-			->AddToken(FTextToken::Create(FText::Format(LOCTEXT("MapCheck_Message_ClearedInvalidWeightmap", "{LandscapeName} : {ErrorMessage}"), Arguments)))
+			->AddToken(FUObjectToken::Create(this, FText::FromString(GetActorNameOrLabel())))
+			->AddToken(FTextToken::Create(FText::Format(LOCTEXT("InvalidTexturesMessage", "Invalid data detected on landscape {0} {0}|plural(one = texture, other = textures). The data has been cleared to avoid fatal error."), InvalidTextures.Num())))
 			->AddToken(FMapErrorToken::Create(FMapErrors::FixedUpDeletedLayerWeightmap));
 	}
 }
@@ -1228,8 +1226,11 @@ void ULandscapeComponent::FixupWeightmaps(const FGuid& InEditLayerGuid)
 			LocalWeightmapTextureUsages.Empty();
 			LocalWeightmapTextureUsages.AddDefaulted(LocalWeightmapTextures.Num());
 
+			// There might be target layers that are referenced in the landscape component that are not present anymore in the parent landscape: we will delete them (with a log, to let the user know):
 			TArray<ULandscapeLayerInfoObject*> LayersToDelete;
-			bool bFixedLayerDeletion = false;
+			// There might be target layers that are referenced in the landscape component but the parent landscape has another LayerInfo associated with that name... 
+			//  those, we can salvage too, we simply need to point the component's layer infos to the landscape's:
+			TMap<ULandscapeLayerInfoObject*, ULandscapeLayerInfoObject*> LayersToRemap; // key == old layer info, value == new layer info
 
 			// make sure the weightmap textures are fully loaded or deleting layers from them will crash! :)
 			for (UTexture* WeightmapTexture : LocalWeightmapTextures)
@@ -1237,30 +1238,68 @@ void ULandscapeComponent::FixupWeightmaps(const FGuid& InEditLayerGuid)
 				WeightmapTexture->ConditionalPostLoad();
 			}
 
-			// LayerInfo Validation check...
-			for (const auto& Allocation : LocalWeightmapLayerAllocations)
-			{
-				if (!Allocation.LayerInfo
-					|| (Allocation.LayerInfo != ALandscapeProxy::VisibilityLayer && Landscape && !Landscape->HasTargetLayer(Allocation.LayerInfo->LayerName)))
+			auto MarkProxyDirty = [this, Info, Landscape]()
 				{
-					if (!bFixedLayerDeletion)
+					// Force resave the proxy through the modified landscape system, so that the user can then use the Build > Save Modified Landscapes (or Build > Build Landscape) button and therefore manually trigger the re-save of all modified proxies.
+					if (Landscape != nullptr)
 					{
-						FFormatNamedArguments Arguments;
-						Arguments.Add(TEXT("LandscapeName"), FText::FromString(GetPathName()));
-						Arguments.Add(TEXT("TargetLayerName"), FText::FromName(Allocation.GetLayerName()));
-						Arguments.Add(TEXT("EditLayerName"), FText::FromString(InEditLayerGuid.IsValid() ? LayersData[InEditLayerGuid].DebugName.ToString() : FString(TEXT("<Runtime>"))));
-						FMessageLog("MapCheck").Info()
-							->AddToken(FTextToken::Create(FText::Format(LOCTEXT("MapCheck_Message_FixedUpDeletedLayerWeightmap", "{LandscapeName} : Removed unused layer weightmap Edit Layer: '{EditLayerName}' Target Layer: '{TargetLayerName}'"), Arguments)))
-							->AddToken(FMapErrorToken::Create(FMapErrors::FixedUpDeletedLayerWeightmap));
+						Info->MarkObjectDirty(/*InObject = */this, /*bInForceResave = */true, Landscape);
 					}
+				};
 
-					bFixedLayerDeletion = true;
+			// LayerInfo Validation check...
+			for (const FWeightmapLayerAllocationInfo& Allocation : LocalWeightmapLayerAllocations)
+			{
+				if (!Allocation.LayerInfo)
+				{
 					LayersToDelete.Add(Allocation.LayerInfo);
 				}
+				else if(Allocation.LayerInfo != ALandscapeProxy::VisibilityLayer && Landscape && !Landscape->HasTargetLayer(Allocation.LayerInfo))
+				{
+					if (Landscape->HasTargetLayer(Allocation.LayerInfo->LayerName))
+					{
+						if (ULandscapeLayerInfoObject* NewLayerInfo = Landscape->GetTargetLayers()[Allocation.LayerInfo->LayerName].LayerInfoObj)
+						{
+							// Point the old layer info to the new one : 
+							LayersToRemap.Add(Allocation.LayerInfo, NewLayerInfo);
+						}
+						else
+						{
+							// No way to retrieve a layer info by that name, consider it a deleted layer : 
+							LayersToDelete.Add(Allocation.LayerInfo);
+						}
+					}
+					else
+					{
+						// No way to retrieve a layer info by that name, consider it a deleted layer : 
+						LayersToDelete.Add(Allocation.LayerInfo);
+					}
+				}
 			}
 			
-			if (bFixedLayerDeletion)
+			auto GetLayersText = [](int32 NumLayers) { return FText::Format(LOCTEXT("Layers", "{0}|plural(one = layer, other = layers)"), NumLayers); };
+			auto GetForEditLayerText = [&InEditLayerGuid, this]() { return InEditLayerGuid.IsValid() ? FText::Format(LOCTEXT("ForEditLayer", " for edit layer {0}"), FText::FromName(LayersData[InEditLayerGuid].DebugName)) : FText::GetEmpty(); };
+
+			if (!LayersToDelete.IsEmpty())
 			{
+				TArray<FString> LayerStrings;
+				Algo::Transform(LayersToDelete, LayerStrings, [](ULandscapeLayerInfoObject* LayerInfo) { return LayerInfo ? LayerInfo->LayerName.ToString() : FString(TEXT("null")); });
+				FString LayersToDeleteString = *FString::Join(LayerStrings, TEXT(", "));
+
+				// Report the deletion info to MapCheck (no need for a warning as the fix is automatic and trivial) :
+				FFormatNamedArguments Arguments;
+				Arguments.Add(TEXT("Layers"), GetLayersText(LayersToDelete.Num()));
+				Arguments.Add(TEXT("TargetLayerNames"), FText::FromString(LayersToDeleteString));
+				Arguments.Add(TEXT("ForEditLayer"), GetForEditLayerText());
+				FMessageLog("MapCheck").Info()
+					->AddToken(FUObjectToken::Create(Proxy, FText::FromString(Proxy->GetActorNameOrLabel())))
+					->AddToken(FTextToken::Create(FText::Format(LOCTEXT("MapCheck_Message_FixedUpDeletedLayerWeightmap", "Removed unused layer weightmap{ForEditLayer} : "
+					"Target {Layers}: \"{TargetLayerNames}\""), Arguments)))
+					->AddToken(FMapErrorToken::Create(FMapErrors::FixedUpDeletedLayerWeightmap));
+
+				// Perform the fixup proper : 
+				MarkProxyDirty();
+
 				// Delete material layer in the base/render layer
 				{
 					FGuid BaseLayerGuid;
@@ -1283,6 +1322,40 @@ void ULandscapeComponent::FixupWeightmaps(const FGuid& InEditLayerGuid)
 								
 			}
 
+			if (!LayersToRemap.IsEmpty())
+			{
+				TArray<FString> LayerStrings;
+				Algo::Transform(LayersToRemap, LayerStrings, [](const TPair<ULandscapeLayerInfoObject*, ULandscapeLayerInfoObject*> &LayerToRemap)	
+					{ 
+						ULandscapeLayerInfoObject* OldLayerInfo = LayerToRemap.Key;
+						ULandscapeLayerInfoObject* NewLayerInfo = LayerToRemap.Value;
+						return FString::Format(TEXT("{0} ({1} -> {2})"), {OldLayerInfo->LayerName.ToString(), OldLayerInfo->GetFullName(), NewLayerInfo->GetFullName()});
+					});
+				FString LayersToRemapString = *FString::Join(LayerStrings, TEXT(", "));
+
+				// Report the remapping info to MapCheck (no need for a warning as the fix is automatic and trivial) :
+				FFormatNamedArguments Arguments;
+				Arguments.Add(TEXT("Layers"), GetLayersText(LayersToRemap.Num()));
+				Arguments.Add(TEXT("TargetLayerNames"), FText::FromString(LayersToRemapString));
+				Arguments.Add(TEXT("ForEditLayer"), GetForEditLayerText());
+				FMessageLog("MapCheck").Info()
+					->AddToken(FUObjectToken::Create(Proxy, FText::FromString(Proxy->GetActorNameOrLabel())))
+					->AddToken(FTextToken::Create(FText::Format(LOCTEXT("MapCheck_Message_FixedUpDeletedLayerWeightmap", "Fixed up layer info{ForEditLayer} : "
+						"Target {Layers}: \"{TargetLayerNames}\""), Arguments)))
+					->AddToken(FMapErrorToken::Create(FMapErrors::FixedUpNonMatchingTargetLayer));
+
+				// Perform the fixup proper : 
+				MarkProxyDirty();
+				for (FWeightmapLayerAllocationInfo& Allocation : LocalWeightmapLayerAllocations)
+				{
+					if (ULandscapeLayerInfoObject** NewLayerInfo = LayersToRemap.Find(Allocation.LayerInfo))
+					{
+						check(Allocation.GetLayerName() == (*NewLayerInfo)->LayerName);
+						Allocation.LayerInfo = *NewLayerInfo;
+					}
+				}
+			}
+
 			bool bFixedWeightmapTextureIndex = false;
 
 			// Store the weightmap allocations in WeightmapUsageMap
@@ -1291,6 +1364,7 @@ void ULandscapeComponent::FixupWeightmaps(const FGuid& InEditLayerGuid)
 				FWeightmapLayerAllocationInfo& Allocation = LocalWeightmapLayerAllocations[LayerIdx];
 				if (!Allocation.IsAllocated())
 				{
+					MarkProxyDirty();
 					LocalWeightmapLayerAllocations.RemoveAt(LayerIdx);
 					continue;
 				}
@@ -1298,13 +1372,13 @@ void ULandscapeComponent::FixupWeightmaps(const FGuid& InEditLayerGuid)
 				// Fix up any problems caused by the layer deletion bug.
 				if (Allocation.WeightmapTextureIndex >= LocalWeightmapTextures.Num())
 				{
+					MarkProxyDirty();
 					Allocation.WeightmapTextureIndex = static_cast<uint8>(LocalWeightmapTextures.Num() - 1);
 					if (!bFixedWeightmapTextureIndex)
 					{
-						FFormatNamedArguments Arguments;
-						Arguments.Add(TEXT("LandscapeName"), FText::FromString(GetName()));
 						FMessageLog("MapCheck").Warning()
-							->AddToken(FTextToken::Create(FText::Format(LOCTEXT("MapCheck_Message_FixedUpIncorrectLayerWeightmap", "{LandscapeName} : Fixed up incorrect layer weightmap texture index"), Arguments)))
+							->AddToken(FUObjectToken::Create(Proxy, FText::FromString(Proxy->GetActorNameOrLabel())))
+							->AddToken(FTextToken::Create(LOCTEXT("MapCheck_Message_FixedUpIncorrectLayerWeightmap", "Fixed up incorrect layer weightmap texture index")))
 							->AddToken(FMapErrorToken::Create(FMapErrors::FixedUpIncorrectLayerWeightmap));
 					}
 					bFixedWeightmapTextureIndex = true;
@@ -1331,10 +1405,10 @@ void ULandscapeComponent::FixupWeightmaps(const FGuid& InEditLayerGuid)
 				{
 					FFormatNamedArguments Arguments;
 					Arguments.Add(TEXT("LayerName"), FText::FromString(Allocation.GetLayerName().ToString()));
-					Arguments.Add(TEXT("LandscapeName"), FText::FromString(GetName()));
 					Arguments.Add(TEXT("ChannelName"), FText::FromString(Usage->ChannelUsage[Allocation.WeightmapTextureChannel]->GetName()));
 					FMessageLog("MapCheck").Warning()
-						->AddToken(FTextToken::Create(FText::Format(LOCTEXT("MapCheck_Message_FixedUpSharedLayerWeightmap", "Fixed up shared weightmap texture for layer {LayerName} in component '{LandscapeName}' (shares with '{ChannelName}')"), Arguments)))
+						->AddToken(FUObjectToken::Create(Proxy, FText::FromString(Proxy->GetActorNameOrLabel())))
+						->AddToken(FTextToken::Create(FText::Format(LOCTEXT("MapCheck_Message_FixedUpSharedLayerWeightmap", "Fixed up shared weightmap texture for layer {LayerName} in landscape component (shares with '{ChannelName}')"), Arguments)))
 						->AddToken(FMapErrorToken::Create(FMapErrors::FixedUpSharedLayerWeightmap));
 					LocalWeightmapLayerAllocations.RemoveAt(LayerIdx);
 					continue;
@@ -1355,7 +1429,7 @@ void ULandscapeComponent::UpdateLayerAllowListFromPaintedLayers()
 {
 	const TArray<FWeightmapLayerAllocationInfo>& ComponentWeightmapLayerAllocations = GetWeightmapLayerAllocations();
 
-	for (const auto& Allocation : ComponentWeightmapLayerAllocations)
+	for (const FWeightmapLayerAllocationInfo& Allocation : ComponentWeightmapLayerAllocations)
 	{
 		LayerAllowList.AddUnique(Allocation.LayerInfo);
 	}
@@ -6489,9 +6563,10 @@ void ALandscapeStreamingProxy::PostRegisterAllComponents()
 			Arguments.Add(TEXT("LevelName"), FText::FromString(GetLevel()->GetOutermost()->GetName()));
 			Arguments.Add(TEXT("ProxyName"), FText::FromString(GetName()));
 			FMessageLog("MapCheck").Error()
-				->AddToken(FUObjectToken::Create(this))
-				->AddToken(FTextToken::Create(FText::Format(LOCTEXT("MapCheck_Message_LandscapeLazyObjectPtrDeprecation_Warning", "Landscape proxy {ProxyName} of {LevelName} points to a LandscapeActor that is not currently loaded. This will lose the property upon save. \
-				Please make sure to load the level containing the parent landscape actor prior to {LevelName} so that data deprecation can be performed adequately. It is advised to reassign the \"Landscape Actor\" property of LandscapeStreamingProxies and resave these actors."), Arguments)))
+				->AddToken(FUObjectToken::Create(this, FText::FromString(GetActorNameOrLabel())))
+				->AddToken(FTextToken::Create(FText::Format(LOCTEXT("MapCheck_Message_LandscapeLazyObjectPtrDeprecation_Warning", "Landscape proxy {ProxyName} of {LevelName} points to a LandscapeActor that is not currently loaded. "
+				"This will lose the property upon save. Please make sure to load the level containing the parent landscape actor prior to {LevelName} so that data deprecation can be performed adequately. "
+				"It is advised to reassign the \"Landscape Actor\" property of LandscapeStreamingProxies and resave these actors."), Arguments)))
 				->AddToken(FMapErrorToken::Create(FMapErrors::LandscapeLazyObjectPtrDeprecation_Warning));
 
 			// Show MapCheck window
@@ -7953,7 +8028,7 @@ void ULandscapeComponent::GenerateMobilePlatformPixelData(bool bIsCooking, const
 		int32 CurrentChannel = 0;
 		int32 RemainingChannels = 0;
 		int32 index = 0;
-		for (auto& Allocation : MobileWeightmapLayerAllocations)
+		for (FWeightmapLayerAllocationInfo& Allocation : MobileWeightmapLayerAllocations)
 		{
 			if (!Allocation.LayerInfo)
 			{
@@ -8036,7 +8111,7 @@ void ULandscapeComponent::GenerateMobilePlatformPixelData(bool bIsCooking, const
 			UMaterialInstanceDynamic* NewMobileMaterialInstance = UMaterialInstanceDynamic::Create(MobileCombinationMaterialInstances[MaterialIndex], this);
 
 			// Set the layer mask
-			for (const auto& Allocation : MobileWeightmapLayerAllocations)
+			for (const FWeightmapLayerAllocationInfo& Allocation : MobileWeightmapLayerAllocations)
 			{
 				if (Allocation.LayerInfo)
 				{
@@ -8108,7 +8183,7 @@ void ULandscapeComponent::GenerateMobilePlatformPixelData(bool bIsCooking, const
 			NewMobileMaterialInstance->SetParentEditorOnly(MobileCombinationMaterialInstances[MaterialIndex]);
 
 			// Set the layer mask
-			for (const auto& Allocation : MobileWeightmapLayerAllocations)
+			for (const FWeightmapLayerAllocationInfo& Allocation : MobileWeightmapLayerAllocations)
 			{
 				if (Allocation.LayerInfo)
 				{
@@ -8323,7 +8398,7 @@ bool ALandscapeProxy::LandscapeImportHeightmapFromRenderTarget(UTextureRenderTar
 			RenderTargetResource->ReadLinearColorPixels(OutputRTHeightmap, FReadSurfaceDataFlags(RCM_MinMax, CubeFace_MAX), SampleRect);
 			HeightData.Reserve(OutputRTHeightmap.Num());
 
-			for (auto LinearColor : OutputRTHeightmap)
+			for (const FLinearColor& LinearColor : OutputRTHeightmap)
 			{
 				if (InImportHeightFromRGChannel)
 				{
