[HWRT] Reenable r.RayTracing.DynamicInstances.Parallel but disable bSupportsParallelGDRTI for FLandscapeComponentSceneProxy.

diff --git a/Engine/Source/Runtime/Engine/Private/PrimitiveSceneProxy.cpp b/Engine/Source/Runtime/Engine/Private/PrimitiveSceneProxy.cpp
--- a/Engine/Source/Runtime/Engine/Private/PrimitiveSceneProxy.cpp
+++ b/Engine/Source/Runtime/Engine/Private/PrimitiveSceneProxy.cpp
@@ -112,7 +112,7 @@ bool IsParallelGatherDynamicMeshElementsEnabled()
 		GSupportsParallelRenderingTasksWithSeparateRHIThread;
 }
 
-bool GParallelGatherDynamicRayTracingInstances = false;
+bool GParallelGatherDynamicRayTracingInstances = true;
 static FAutoConsoleVariableRef CVarParallelGatherDynamicRayTracingInstances(
 	TEXT("r.RayTracing.DynamicInstances.Parallel"),
 	GParallelGatherDynamicRayTracingInstances,
diff --git a/Engine/Source/Runtime/Landscape/Private/LandscapeRender.cpp b/Engine/Source/Runtime/Landscape/Private/LandscapeRender.cpp
--- a/Engine/Source/Runtime/Landscape/Private/LandscapeRender.cpp
+++ b/Engine/Source/Runtime/Landscape/Private/LandscapeRender.cpp
@@ -1361,6 +1361,9 @@ FLandscapeComponentSceneProxy::FLandscapeComponentSceneProxy(ULandscapeComponent
 	// Enabled just so the GVarDumpLandscapeLODsCurrentFrame debug feature only runs once per render proxy in GetDynamicMeshElements
 	bSinglePassGDME = true;
 
+	// Disable GDRTI since GetDynamicRayTracingInstances(...) calls FRHICommandListImmediate::Get() (via FindOrCreateRayTracingState(...))
+	bSupportsParallelGDRTI = false;
+
 	VisibilityHelper.Init(InComponent, this);
 
 	if (!VisibilityHelper.ShouldBeVisible())
