[HWRT] Updated FLandscapeComponentSceneProxy::GetDynamicRayTracingInstances(...) to use Collector.GetRHICommandList() instead of FRHICommandListImmediate.
- which makes it possible to use ParallelGDRTI.

diff --git a/Engine/Source/Runtime/Landscape/Private/LandscapeRender.cpp b/Engine/Source/Runtime/Landscape/Private/LandscapeRender.cpp
--- a/Engine/Source/Runtime/Landscape/Private/LandscapeRender.cpp
+++ b/Engine/Source/Runtime/Landscape/Private/LandscapeRender.cpp
@@ -550,7 +550,7 @@ struct FLandscapeRayTracingImpl
 	TIndirectArray<FLandscapeRayTracingState> PerViewRayTracingState;
 
 	// ViewStateInterface pointer can be NULL
-	FLandscapeRayTracingState* FindOrCreateRayTracingState(FSceneViewStateInterface* ViewStateInterface, int32 NumSubsections, int32 SubsectionSizeVerts);
+	FLandscapeRayTracingState* FindOrCreateRayTracingState(FRHICommandListBase& RHICmdList, FSceneViewStateInterface* ViewStateInterface, int32 NumSubsections, int32 SubsectionSizeVerts);
 };
 
 // When views get deleted, we need to clean up the per view ray tracing data, which is handled by this class.  This is just a doubly linked list
@@ -1361,9 +1361,6 @@ FLandscapeComponentSceneProxy::FLandscapeComponentSceneProxy(ULandscapeComponent
 	// Enabled just so the GVarDumpLandscapeLODsCurrentFrame debug feature only runs once per render proxy in GetDynamicMeshElements
 	bSinglePassGDME = true;
 
-	// Disable GDRTI since GetDynamicRayTracingInstances(...) calls FRHICommandListImmediate::Get() (via FindOrCreateRayTracingState(...))
-	bSupportsParallelGDRTI = false;
-
 	VisibilityHelper.Init(InComponent, this);
 
 	if (!VisibilityHelper.ShouldBeVisible())
@@ -1820,7 +1817,7 @@ void FLandscapeComponentSceneProxy::CreateRenderThreadResources(FRHICommandListB
 }
 
 #if RHI_RAYTRACING
-FLandscapeRayTracingState* FLandscapeRayTracingImpl::FindOrCreateRayTracingState(FSceneViewStateInterface* ViewStateInterface, int32 NumSubsections, int32 SubsectionSizeVerts)
+FLandscapeRayTracingState* FLandscapeRayTracingImpl::FindOrCreateRayTracingState(FRHICommandListBase& RHICmdList, FSceneViewStateInterface* ViewStateInterface, int32 NumSubsections, int32 SubsectionSizeVerts)
 {
 	// Default view key of zero if there's no ViewStateInterface provided.  View keys start at 1, so 0 wouldn't be a valid key on an actual view.
 	uint32 ViewKey = 0;
@@ -1868,8 +1865,6 @@ FLandscapeRayTracingState* FLandscapeRayTracingImpl::FindOrCreateRayTracingState
 	// Initialize rendering data
 	RayTracingState->NumSubsections = NumSubsections;
 
-	FRHICommandListImmediate& RHICmdList = FRHICommandListImmediate::Get();
-
 	for (int32 SubY = 0; SubY < NumSubsections; SubY++)
 	{
 		for (int32 SubX = 0; SubX < NumSubsections; SubX++)
@@ -3038,7 +3033,7 @@ void FLandscapeComponentSceneProxy::GetDynamicRayTracingInstances(FRayTracingIns
 	{
 		RayTracingImpl = MakePimpl<FLandscapeRayTracingImpl>();
 	}
-	FLandscapeRayTracingState* RayTracingState = RayTracingImpl.Get()->FindOrCreateRayTracingState(SceneView.State, NumSubsections, SubsectionSizeVerts);
+	FLandscapeRayTracingState* RayTracingState = RayTracingImpl.Get()->FindOrCreateRayTracingState(Collector.GetRHICommandList(), SceneView.State, NumSubsections, SubsectionSizeVerts);
 
 	int32 LODToRender = static_cast<int32>(RenderSystem.GetSectionLODValue(SceneView, RenderCoord));
 
