Disable bSupportsParallelGDRTI in FLandscapeComponentSceneProxy again since there seems to be a race condition when FLandscapeRenderSystem::GetSectionLODValue(...) is called.

diff --git a/Engine/Source/Runtime/Landscape/Private/LandscapeRender.cpp b/Engine/Source/Runtime/Landscape/Private/LandscapeRender.cpp
--- a/Engine/Source/Runtime/Landscape/Private/LandscapeRender.cpp
+++ b/Engine/Source/Runtime/Landscape/Private/LandscapeRender.cpp
@@ -1361,6 +1361,9 @@ FLandscapeComponentSceneProxy::FLandscapeComponentSceneProxy(ULandscapeComponent
 	// Enabled just so the GVarDumpLandscapeLODsCurrentFrame debug feature only runs once per render proxy in GetDynamicMeshElements
 	bSinglePassGDME = true;
 
+	// Disable GDRTI since there's seems to be a race condition in FLandscapeRenderSystem::GetSectionLODValue(...) otherwise
+	bSupportsParallelGDRTI = false;
+
 	VisibilityHelper.Init(InComponent, this);
 
 	if (!VisibilityHelper.ShouldBeVisible())
