Fix localization string key conflict.

diff --git a/Engine/Source/Runtime/Landscape/Private/LandscapeEdit.cpp b/Engine/Source/Runtime/Landscape/Private/LandscapeEdit.cpp
--- a/Engine/Source/Runtime/Landscape/Private/LandscapeEdit.cpp
+++ b/Engine/Source/Runtime/Landscape/Private/LandscapeEdit.cpp
@@ -1293,7 +1293,7 @@ void ULandscapeComponent::FixupWeightmaps(const FGuid& InEditLayerGuid)
 				Arguments.Add(TEXT("ForEditLayer"), GetForEditLayerText());
 				FMessageLog("MapCheck").Info()
 					->AddToken(FUObjectToken::Create(Proxy, FText::FromString(Proxy->GetActorNameOrLabel())))
-					->AddToken(FTextToken::Create(FText::Format(LOCTEXT("MapCheck_Message_FixedUpDeletedLayerWeightmap", "Removed unused layer weightmap{ForEditLayer} : "
+					->AddToken(FTextToken::Create(FText::Format(LOCTEXT("MapCheck_Message_FixedRemoveUnusedLayerWeightmap", "Removed unused layer weightmap{ForEditLayer} : "
 					"Target {Layers}: \"{TargetLayerNames}\""), Arguments)))
 					->AddToken(FMapErrorToken::Create(FMapErrors::FixedUpDeletedLayerWeightmap));
 
