Fixed landscape proxy holding references on some of its foliage components if those components are selectively removed

#tests Cooked game, PIE

diff --git a/Engine/Source/Runtime/Landscape/Private/LandscapeGrass.cpp b/Engine/Source/Runtime/Landscape/Private/LandscapeGrass.cpp
--- a/Engine/Source/Runtime/Landscape/Private/LandscapeGrass.cpp
+++ b/Engine/Source/Runtime/Landscape/Private/LandscapeGrass.cpp
@@ -2350,6 +2350,7 @@ void ALandscapeProxy::FlushGrassComponents(const TSet<ULandscapeComponent*>* Onl
 					Used->ClearInstances();
 					Used->DetachFromComponent(FDetachmentTransformRules(EDetachmentRule::KeepRelative, false));
 					Used->DestroyComponent();
+					FoliageComponents.RemoveSwap(Used);
 				}
 				Iter.RemoveCurrent();
 			}
