Add null check for landscape spline segments control points in ULandscapeSplinesComponent::CheckSplinesValid.

diff --git a/Engine/Source/Runtime/Landscape/Private/LandscapeSplines.cpp b/Engine/Source/Runtime/Landscape/Private/LandscapeSplines.cpp
--- a/Engine/Source/Runtime/Landscape/Private/LandscapeSplines.cpp
+++ b/Engine/Source/Runtime/Landscape/Private/LandscapeSplines.cpp
@@ -507,13 +507,10 @@ ILandscapeSplineInterface* ULandscapeSplinesComponent::GetSplineOwner()
 void ULandscapeSplinesComponent::CheckSplinesValid()
 {
 #if DO_CHECK
-	// This shouldn't happen, but it has somehow (TTP #334549) so we have to fix it
-	ensure(!ControlPoints.Contains(nullptr));
-	ensure(!Segments.Contains(nullptr));
-
 	// Remove all null control points/segments
-	ControlPoints.Remove(nullptr);
-	Segments.Remove(nullptr);
+	// This shouldn't happen, but it has somehow (TTP #334549) so we have to fix it
+	ensure(!ControlPoints.Remove(nullptr));
+	ensure(!Segments.Remove(nullptr));
 
 	// Check for cross-spline connections, as this is a potential source of nulls
 	// this may be allowed in future, but is not currently
@@ -522,15 +519,22 @@ void ULandscapeSplinesComponent::CheckSplinesValid()
 		ensure(ControlPoint->GetOuterULandscapeSplinesComponent() == this);
 		for (const FLandscapeSplineConnection& Connection : ControlPoint->ConnectedSegments)
 		{
-			ensure(Connection.Segment->GetOuterULandscapeSplinesComponent() == this);
+			if (Connection.Segment)
+			{
+				ensure(Connection.Segment->GetOuterULandscapeSplinesComponent() == this);
+			}
 		}
 	}
+
 	for (ULandscapeSplineSegment* Segment : Segments)
 	{
 		ensure(Segment->GetOuterULandscapeSplinesComponent() == this);
 		for (const FLandscapeSplineSegmentConnection& Connection : Segment->Connections)
 		{
-			ensure(Connection.ControlPoint->GetOuterULandscapeSplinesComponent() == this);
+			if (Connection.ControlPoint)
+			{
+				ensure(Connection.ControlPoint->GetOuterULandscapeSplinesComponent() == this);
+			}
 		}
 	}
 #endif
