[Core] Fix GetResourceSizeEx for CDOs for some classes

#tests PIE
#rnx

diff --git a/Engine/Source/Runtime/Landscape/Private/LandscapeCollision.cpp b/Engine/Source/Runtime/Landscape/Private/LandscapeCollision.cpp
--- a/Engine/Source/Runtime/Landscape/Private/LandscapeCollision.cpp
+++ b/Engine/Source/Runtime/Landscape/Private/LandscapeCollision.cpp
@@ -2696,6 +2696,10 @@ void ULandscapeHeightfieldCollisionComponent::PreSave(FObjectPreSaveContext Obje
 #if WITH_EDITOR
 bool ULandscapeHeightfieldCollisionComponent::NeedsLoadForClient() const
 {
+	if (IsTemplate())
+	{
+		return true;
+	}
 	ALandscapeProxy* Proxy = GetLandscapeProxy();
 	if (ensure(Proxy))
 	{
@@ -2707,6 +2711,10 @@ bool ULandscapeHeightfieldCollisionComponent::NeedsLoadForClient() const
 
 bool ULandscapeHeightfieldCollisionComponent::NeedsLoadForServer() const
 {
+	if (IsTemplate())
+	{
+		return true;
+	}
 	ALandscapeProxy* Proxy = GetLandscapeProxy();
 	if (ensure(Proxy))
 	{
diff --git a/Engine/Source/Runtime/MassEntity/Private/MassEntityManager.cpp b/Engine/Source/Runtime/MassEntity/Private/MassEntityManager.cpp
--- a/Engine/Source/Runtime/MassEntity/Private/MassEntityManager.cpp
+++ b/Engine/Source/Runtime/MassEntity/Private/MassEntityManager.cpp
@@ -163,7 +163,7 @@ FMassEntityManager::~FMassEntityManager()
 
 void FMassEntityManager::GetResourceSizeEx(FResourceSizeEx& CumulativeResourceSize)
 {
-	SIZE_T MyExtraSize = GetEntityStorageInterface().GetAllocatedSize()
+	SIZE_T MyExtraSize = (bInitialized ? GetEntityStorageInterface().GetAllocatedSize() : 0)
 		+ FragmentHashToArchetypeMap.GetAllocatedSize()
 		+ FragmentTypeToArchetypeMap.GetAllocatedSize();
 
