Fixed assert/crash when trying to resolve a soft object ptr from BeginDestroy/FinishDestroy in landscape : calling CanHaveLayersContent on a ALandscapeStreamingProxy would resolve the soft object ptr to the parent ALandscape, which might have been GCed already, which would lead to a crash in ULevel::ResolveSubobject (accessing objects like that during GC is bad since GC doesn't have any rules with regards to the order of destruction).

diff --git a/Engine/Source/Runtime/Landscape/Private/LandscapeEditLayers.cpp b/Engine/Source/Runtime/Landscape/Private/LandscapeEditLayers.cpp
--- a/Engine/Source/Runtime/Landscape/Private/LandscapeEditLayers.cpp
+++ b/Engine/Source/Runtime/Landscape/Private/LandscapeEditLayers.cpp
@@ -11905,13 +11905,12 @@ void ALandscape::TickLayers(float DeltaTime)
 
 void ALandscapeProxy::BeginDestroy()
 {
+	// At GC time, it is important not to rely on other objects as there's no guarantee about destruction order so we mustn't 
+	//  do anything that resolves the soft object ptr LandscapeActorRef here, in particular
+
 #if WITH_EDITORONLY_DATA
-	if (CanHaveLayersContent())
-	{
-		// Note that this fence and the BeginDestroy/FinishDestroy is only actually used by the derived class ALandscape
-		// It could be moved there
-		ReleaseResourceFence.BeginFence();
-	}
+	// Prevent destruction until all render resources are done destroying
+	ReleaseResourceFence.BeginFence();
 #endif
 
 	Super::BeginDestroy();
@@ -11922,12 +11921,9 @@ bool ALandscapeProxy::IsReadyForFinishDestroy()
 	bool bReadyForFinishDestroy = Super::IsReadyForFinishDestroy();
 
 #if WITH_EDITORONLY_DATA
-	if (CanHaveLayersContent())
+	if (bReadyForFinishDestroy)
 	{
-		if (bReadyForFinishDestroy)
-		{
-			bReadyForFinishDestroy = ReleaseResourceFence.IsFenceComplete();
-		}
+		bReadyForFinishDestroy = ReleaseResourceFence.IsFenceComplete();
 	}
 #endif
 
@@ -11937,24 +11933,21 @@ bool ALandscapeProxy::IsReadyForFinishDestroy()
 void ALandscapeProxy::FinishDestroy()
 {
 #if WITH_EDITORONLY_DATA
-	if (CanHaveLayersContent())
-	{
-		check(ReleaseResourceFence.IsFenceComplete());
+	check(ReleaseResourceFence.IsFenceComplete());
 
-		for (auto& ItPair : HeightmapsCPUReadback)
-		{
-			FLandscapeEditLayerReadback* HeightmapCPUReadback = ItPair.Value;
-			delete HeightmapCPUReadback;
-		}
-		HeightmapsCPUReadback.Empty();
+	for (auto& ItPair : HeightmapsCPUReadback)
+	{
+		FLandscapeEditLayerReadback* HeightmapCPUReadback = ItPair.Value;
+		delete HeightmapCPUReadback;
+	}
+	HeightmapsCPUReadback.Empty();
 
-		for (auto& ItPair : WeightmapsCPUReadback)
-		{
-			FLandscapeEditLayerReadback* WeightmapCPUReadback = ItPair.Value;
-			delete WeightmapCPUReadback;
-		}
-		WeightmapsCPUReadback.Empty();
+	for (auto& ItPair : WeightmapsCPUReadback)
+	{
+		FLandscapeEditLayerReadback* WeightmapCPUReadback = ItPair.Value;
+		delete WeightmapCPUReadback;
 	}
+	WeightmapsCPUReadback.Empty();
 #endif
 
 	Super::FinishDestroy();
@@ -12276,20 +12269,17 @@ void ALandscape::BeginDestroy()
 void ALandscape::FinishDestroy()
 {
 #if WITH_EDITORONLY_DATA
-	if (CanHaveLayersContent())
-	{
-		check(ReleaseResourceFence.IsFenceComplete());
+	check(ReleaseResourceFence.IsFenceComplete());
 
-		delete CombinedLayersWeightmapAllMaterialLayersResource;
-		delete CurrentLayersWeightmapAllMaterialLayersResource;
-		delete WeightmapScratchExtractLayerTextureResource;
-		delete WeightmapScratchPackLayerTextureResource;
+	delete CombinedLayersWeightmapAllMaterialLayersResource;
+	delete CurrentLayersWeightmapAllMaterialLayersResource;
+	delete WeightmapScratchExtractLayerTextureResource;
+	delete WeightmapScratchPackLayerTextureResource;
 
-		CombinedLayersWeightmapAllMaterialLayersResource = nullptr;
-		CurrentLayersWeightmapAllMaterialLayersResource = nullptr;
-		WeightmapScratchExtractLayerTextureResource = nullptr;
-		WeightmapScratchPackLayerTextureResource = nullptr;
-	}
+	CombinedLayersWeightmapAllMaterialLayersResource = nullptr;
+	CurrentLayersWeightmapAllMaterialLayersResource = nullptr;
+	WeightmapScratchExtractLayerTextureResource = nullptr;
+	WeightmapScratchPackLayerTextureResource = nullptr;
 #endif
 
 	Super::FinishDestroy();
