[Backout] - CL38196998
Reverting as a mitigation because this breaks the terrain collision in the game branches

#rnx
Original CL Desc
-----------------------------------------------------------------
Prepare BodyCollisionData for Landscape HeightFields

diff --git a/Engine/Source/Runtime/Landscape/Private/LandscapeCollision.cpp b/Engine/Source/Runtime/Landscape/Private/LandscapeCollision.cpp
--- a/Engine/Source/Runtime/Landscape/Private/LandscapeCollision.cpp
+++ b/Engine/Source/Runtime/Landscape/Private/LandscapeCollision.cpp
@@ -410,10 +410,6 @@ void ULandscapeHeightfieldCollisionComponent::OnCreatePhysicsState()
 
 				TUniquePtr<Chaos::FPerShapeData> NewShape = Chaos::FShapeInstanceProxy::Make(ShapeArray.Num(), ChaosHeightFieldFromCooked);
 
-				FBodyCollisionData BodyCollisionData;
-				FBodyInstance::BuildBodyCollisionFlags(BodyCollisionData.CollisionFlags, GetCollisionEnabled(), false);
-				const FBodyCollisionFlags& CollisionFlags = BodyCollisionData.CollisionFlags;
-
 				// Setup filtering
 				FCollisionFilterData QueryFilterData, SimFilterData;
 				CreateShapeFilterData(static_cast<uint8>(GetCollisionObjectType()), FMaskFilter(0), GetOwner()->GetUniqueID(), GetCollisionResponseToChannels(),
@@ -427,14 +423,6 @@ void ULandscapeHeightfieldCollisionComponent::OnCreatePhysicsState()
 				NewShape->SetSimData(SimFilterData);
 				NewShape->SetMaterials(HeightfieldRef->UsedChaosMaterials);
 
-				const bool bNewQueryShape = CollisionFlags.bEnableQueryCollision;
-				const bool bNewSimShape = CollisionFlags.bEnableSimCollisionSimple;
-				const bool bNewProbeShape = CollisionFlags.bEnableProbeCollision;
-
-				NewShape->SetIsProbe(bNewProbeShape);
-				NewShape->SetSimEnabled(bNewSimShape);
-				NewShape->SetQueryEnabled(bNewQueryShape);
-
 				Geoms.Emplace(MoveTemp(ChaosHeightFieldFromCooked));
 				ShapeArray.Emplace(MoveTemp(NewShape));
 
@@ -457,10 +445,6 @@ void ULandscapeHeightfieldCollisionComponent::OnCreatePhysicsState()
 					NewSimpleShape->SetSimData(SimFilterDataSimple);
 					NewSimpleShape->SetMaterials(HeightfieldRef->UsedChaosMaterials);
 
-					NewSimpleShape->SetIsProbe(bNewProbeShape);
-					NewSimpleShape->SetSimEnabled(bNewSimShape);
-					NewSimpleShape->SetQueryEnabled(bNewQueryShape);
-
 					Geoms.Emplace(MoveTemp(ChaosSimpleHeightFieldFromCooked));
 					ShapeArray.Emplace(MoveTemp(NewSimpleShape));
 				}
@@ -487,10 +471,6 @@ void ULandscapeHeightfieldCollisionComponent::OnCreatePhysicsState()
 					NewEditorShape->SetSimData(SimFilterDataEd);
 					NewEditorShape->SetMaterials(HeightfieldRef->UsedChaosMaterials);
 
-					NewEditorShape->SetIsProbe(bNewProbeShape);
-					NewEditorShape->SetSimEnabled(bNewSimShape);
-					NewEditorShape->SetQueryEnabled(bNewQueryShape);
-
 					Geoms.Emplace(MoveTemp(ChaosEditorHeightFieldFromCooked));
 					ShapeArray.Emplace(MoveTemp(NewEditorShape));
 				}
