Fixed crash when destroying a landscape after an Undo

diff --git a/Engine/Source/Runtime/Landscape/Private/LandscapeEditLayers.cpp b/Engine/Source/Runtime/Landscape/Private/LandscapeEditLayers.cpp
--- a/Engine/Source/Runtime/Landscape/Private/LandscapeEditLayers.cpp
+++ b/Engine/Source/Runtime/Landscape/Private/LandscapeEditLayers.cpp
@@ -12237,30 +12237,27 @@ TArray<FName> ALandscapeProxy::SynchronizeUnmarkedSharedProperties(ALandscapePro
 void ALandscape::BeginDestroy()
 {
 #if WITH_EDITOR
-	if (CanHaveLayersContent())
+	if (CombinedLayersWeightmapAllMaterialLayersResource != nullptr)
 	{
-		if (CombinedLayersWeightmapAllMaterialLayersResource != nullptr)
-		{
-			BeginReleaseResource(CombinedLayersWeightmapAllMaterialLayersResource);
-		}
-
-		if (CurrentLayersWeightmapAllMaterialLayersResource != nullptr)
-		{
-			BeginReleaseResource(CurrentLayersWeightmapAllMaterialLayersResource);
-		}
+		BeginReleaseResource(CombinedLayersWeightmapAllMaterialLayersResource);
+	}
 
-		if (WeightmapScratchExtractLayerTextureResource != nullptr)
-		{
-			BeginReleaseResource(WeightmapScratchExtractLayerTextureResource);
-		}
+	if (CurrentLayersWeightmapAllMaterialLayersResource != nullptr)
+	{
+		BeginReleaseResource(CurrentLayersWeightmapAllMaterialLayersResource);
+	}
 
-		if (WeightmapScratchPackLayerTextureResource != nullptr)
-		{
-			BeginReleaseResource(WeightmapScratchPackLayerTextureResource);
-		}
+	if (WeightmapScratchExtractLayerTextureResource != nullptr)
+	{
+		BeginReleaseResource(WeightmapScratchExtractLayerTextureResource);
+	}
 
-		// Use ResourceFence from base class		
+	if (WeightmapScratchPackLayerTextureResource != nullptr)
+	{
+		BeginReleaseResource(WeightmapScratchPackLayerTextureResource);
 	}
+
+	// Use ResourceFence from base class		
 #endif
 
 	Super::BeginDestroy();
