Fix for landscape texture streaming issue when creating new weightmaps

diff --git a/Engine/Source/Runtime/Landscape/Private/LandscapeTextureStreamingManager.cpp b/Engine/Source/Runtime/Landscape/Private/LandscapeTextureStreamingManager.cpp
--- a/Engine/Source/Runtime/Landscape/Private/LandscapeTextureStreamingManager.cpp
+++ b/Engine/Source/Runtime/Landscape/Private/LandscapeTextureStreamingManager.cpp
@@ -1,7 +1,7 @@
 // Copyright Epic Games, Inc. All Rights Reserved.
 
 #include "LandscapeTextureStreamingManager.h"
-#include "Engine/Texture.h"
+#include "Engine/Texture2D.h"
 #include "TextureCompiler.h"
 #include "LandscapePrivate.h"
 
@@ -18,12 +18,48 @@ namespace UE::Landscape
 			Texture->bForceMiplevelsToBeResident = true;
 		}
 	}
+
+	static bool IsTextureFullyStreamedIn(UTexture* InTexture)
+	{
+		const bool bCheckForLODTransition = true;
+		return InTexture &&
+#if WITH_EDITOR
+			!InTexture->IsDefaultTexture() &&
+#endif // WITH_EDITOR
+			!InTexture->HasPendingInitOrStreaming(bCheckForLODTransition) && InTexture->IsFullyStreamedIn();
+	}
+
+	static inline bool EnforceTextureIsFullyStreamedInNow(UTexture* Texture)
+	{
+#if WITH_EDITOR
+		// in editor, textures can be not compiled yet - we should complete that first
+		Texture->BlockOnAnyAsyncBuild();
+#endif // WITH_EDITOR
+		const bool bWaitForLODTransition = true;
+		Texture->WaitForStreaming(bWaitForLODTransition);
+		bool bIsFullyStreamedIn = IsTextureFullyStreamedIn(Texture);
+#if WITH_EDITOR
+		// the above should handle ensure textures are fully streamed in.. but just in case it isn't...
+		if (!bIsFullyStreamedIn)
+		{
+			// this is a sledgehammer, but should always fix it, by rebuilding the entire texture resource with streaming disabled
+			if (UTexture2D* Texture2D = Cast<UTexture2D>(Texture))
+			{
+				Texture2D->TemporarilyDisableStreaming();
+			}
+			bIsFullyStreamedIn = IsTextureFullyStreamedIn(Texture);
+		}
+#endif // WITH_EDITOR
+		check(bIsFullyStreamedIn);
+		return bIsFullyStreamedIn;
+	}
 }
 
 TArray<FLandscapeTextureStreamingManager*> FLandscapeTextureStreamingManager::AllStreamingManagers;
 
 bool FLandscapeTextureStreamingManager::RequestTextureFullyStreamedIn(UTexture* Texture, bool bWaitForStreaming)
 {
+	using namespace UE::Landscape;
 	TWeakObjectPtr<UTexture> TexturePtr = Texture;
 	FTextureState& State = TextureStates.FindOrAdd(TexturePtr);
 
@@ -33,7 +69,7 @@ bool FLandscapeTextureStreamingManager::RequestTextureFullyStreamedIn(UTexture*
 	}
 	else
 	{
-		UE::Landscape::EnsureTextureForcedResident(Texture);
+		EnsureTextureForcedResident(Texture);
 	}
 	State.RequestCount++;
 
@@ -43,14 +79,14 @@ bool FLandscapeTextureStreamingManager::RequestTextureFullyStreamedIn(UTexture*
 	}
 	else if (bWaitForStreaming)
 	{
-		Texture->WaitForStreaming();
-		return IsTextureFullyStreamedIn(Texture);
+		return EnforceTextureIsFullyStreamedInNow(Texture);
 	}
 	return false;
 }
 
 bool FLandscapeTextureStreamingManager::RequestTextureFullyStreamedInForever(UTexture* Texture, bool bWaitForStreaming)
 {
+	using namespace UE::Landscape;
 	TWeakObjectPtr<UTexture> TexturePtr = Texture;
 	FTextureState& State = TextureStates.FindOrAdd(TexturePtr);
 	State.bForever = true;
@@ -62,8 +98,7 @@ bool FLandscapeTextureStreamingManager::RequestTextureFullyStreamedInForever(UTe
 	}
 	else if (bWaitForStreaming)
 	{
-		Texture->WaitForStreaming();
-		return IsTextureFullyStreamedIn(Texture);
+		return EnforceTextureIsFullyStreamedInNow(Texture);
 	}
 	return false;
 }
@@ -110,8 +145,9 @@ void FLandscapeTextureStreamingManager::UnrequestTextureFullyStreamedIn(UTexture
 
 bool FLandscapeTextureStreamingManager::WaitForTextureStreaming()
 {
+	using namespace UE::Landscape;
 	TRACE_CPUPROFILER_EVENT_SCOPE(LandscapeTextureStreamingaAnager_WaitForTextureStreaming);
-	bool bFullyStreamed = true;
+	bool bAllTexturesFullyStreamed = true;
 	for (auto It = TextureStates.CreateIterator(); It; ++It)
 	{
 		UTexture* Texture = It.Key().Get();
@@ -120,13 +156,9 @@ bool FLandscapeTextureStreamingManager::WaitForTextureStreaming()
 			UE::Landscape::EnsureTextureForcedResident(Texture);
 			if (!IsTextureFullyStreamedIn(Texture))
 			{
-#if WITH_EDITOR
-				// in editor, textures can be not compiled yet
-				FTextureCompilingManager::Get().FinishCompilation({ Texture });
-#endif // WITH_EDITOR
-				Texture->WaitForStreaming();
-			}
-			bFullyStreamed = bFullyStreamed && IsTextureFullyStreamedIn(Texture);
+				const bool bTextureIsFullyStreamed = EnforceTextureIsFullyStreamedInNow(Texture);
+				bAllTexturesFullyStreamed = bAllTexturesFullyStreamed && bTextureIsFullyStreamed;
+			}		
 		}
 		else
 		{
@@ -134,7 +166,7 @@ bool FLandscapeTextureStreamingManager::WaitForTextureStreaming()
 			It.RemoveCurrent();
 		}
 	}
-	return bFullyStreamed;
+	return bAllTexturesFullyStreamed;
 }
 
 void FLandscapeTextureStreamingManager::CleanupPostGarbageCollect()
@@ -186,11 +218,7 @@ void FLandscapeTextureStreamingManager::CheckRequestedTextures()
 
 bool FLandscapeTextureStreamingManager::IsTextureFullyStreamedIn(UTexture* InTexture)
 {
-	return InTexture &&
-#if WITH_EDITOR
-		!InTexture->IsDefaultTexture() &&
-#endif // WITH_EDITOR
-		!InTexture->HasPendingInitOrStreaming() && InTexture->IsFullyStreamedIn();
+	return UE::Landscape::IsTextureFullyStreamedIn(InTexture);
 }
 
 bool FLandscapeTextureStreamingManager::AnyStreamingManagerWantsTextureStreamedIn(TWeakObjectPtr<UTexture> TexturePtr)
