Add console command to force rebuild of landscape physical material data.
landscape.RebuildPhysicalMaterial

diff --git a/Engine/Source/Runtime/Landscape/Private/LandscapeSubsystem.cpp b/Engine/Source/Runtime/Landscape/Private/LandscapeSubsystem.cpp
--- a/Engine/Source/Runtime/Landscape/Private/LandscapeSubsystem.cpp
+++ b/Engine/Source/Runtime/Landscape/Private/LandscapeSubsystem.cpp
@@ -209,6 +209,13 @@ namespace UE::Landscape
 			}
 		}
 	}
+
+	static FAutoConsoleCommand CmdLandscapeRebuildPhysicalMaterial(
+		TEXT("landscape.RebuildPhysicalMaterial"),
+		TEXT("Force a rebuild of the physical material data for all landscapes."),
+		FConsoleCommandDelegate::CreateStatic(BuildPhysicalMaterial, EBuildFlags::ForceRebuild)
+	);
+
 #endif //WITH_EDITOR
 }
 
