Move landscape FixedGrid uniform buffers from component to a SharedBuffers

diff --git a/Engine/Source/Runtime/Landscape/Private/LandscapeRender.cpp b/Engine/Source/Runtime/Landscape/Private/LandscapeRender.cpp
--- a/Engine/Source/Runtime/Landscape/Private/LandscapeRender.cpp
+++ b/Engine/Source/Runtime/Landscape/Private/LandscapeRender.cpp
@@ -1761,22 +1761,7 @@ void FLandscapeComponentSceneProxy::CreateRenderThreadResources(FRHICommandListB
 	// Assign LandscapeUniformShaderParameters
 	LandscapeUniformShaderParameters.InitResource(RHICmdList);
 
-	// Create per Lod uniform buffers
 	const int32 NumMips = FMath::CeilLogTwo(SubsectionSizeVerts);
-	// create as many as there are potential mips (even if MaxLOD can be inferior than that), because the grass could need that much :
-	LandscapeFixedGridUniformShaderParameters.AddDefaulted(NumMips);
-	for (int32 LodIndex = 0; LodIndex < NumMips; ++LodIndex)
-	{
-		FLandscapeFixedGridUniformShaderParameters Parameters;
-		Parameters.LodValues = FVector4f(
-			static_cast<float>(LodIndex),
-			0.f,
-			(float)((SubsectionSizeVerts >> LodIndex) - 1),
-			1.f / (float)((SubsectionSizeVerts >> LodIndex) - 1));
-		LandscapeFixedGridUniformShaderParameters[LodIndex].SetContents(RHICmdList, Parameters);
-		LandscapeFixedGridUniformShaderParameters[LodIndex].InitResource(RHICmdList);
-	}
-
 	// Create MeshBatch for grass rendering
 	if (SharedBuffers->GrassIndexBuffer)
 	{
@@ -1802,7 +1787,7 @@ void FLandscapeComponentSceneProxy::CreateRenderThreadResources(FRHICommandListB
 		FMeshBatchElement* GrassBatchElement = &GrassMeshBatch.Elements[0];
 		FLandscapeBatchElementParams* BatchElementParams = &GrassBatchParams[0];
 		BatchElementParams->LandscapeUniformShaderParametersResource = &LandscapeUniformShaderParameters;
-		BatchElementParams->FixedGridUniformShaderParameters = &LandscapeFixedGridUniformShaderParameters;
+		BatchElementParams->FixedGridUniformShaderParameters = &SharedBuffers->LandscapeFixedGridUniformShaderParameters;
 		BatchElementParams->LandscapeSectionLODUniformParameters = nullptr; // Not needed for grass rendering
 		BatchElementParams->SceneProxy = this;
 		BatchElementParams->CurrentLOD = 0;
@@ -1973,12 +1958,6 @@ FLandscapeComponentSceneProxy::~FLandscapeComponentSceneProxy()
 {
 	// Free the subsection uniform buffer
 	LandscapeUniformShaderParameters.ReleaseResource();
-		
-	// Free the lod uniform buffers
-	for (int32 i = 0; i < LandscapeFixedGridUniformShaderParameters.Num(); ++i)
-	{
-		LandscapeFixedGridUniformShaderParameters[i].ReleaseResource();
-	}
 
 	if (SharedBuffers)
 	{
@@ -2393,7 +2372,7 @@ bool FLandscapeComponentSceneProxy::GetMeshElementForFixedGrid(int32 InLodIndex,
 	FLandscapeBatchElementParams* BatchElementParams = new(OutStaticBatchParamArray) FLandscapeBatchElementParams;
 	BatchElementParams->SceneProxy = this;
 	BatchElementParams->LandscapeUniformShaderParametersResource = &LandscapeUniformShaderParameters;
-	BatchElementParams->FixedGridUniformShaderParameters = &LandscapeFixedGridUniformShaderParameters;
+	BatchElementParams->FixedGridUniformShaderParameters = &SharedBuffers->LandscapeFixedGridUniformShaderParameters;
 	BatchElementParams->LandscapeSectionLODUniformParameters = RenderSystem.SectionLODUniformBuffer;
 	BatchElementParams->CurrentLOD = InLodIndex;
 
@@ -2449,7 +2428,7 @@ bool FLandscapeComponentSceneProxy::GetStaticMeshElement(int32 LODIndex, bool bF
 
 		FLandscapeBatchElementParams* BatchElementParams = new(OutStaticBatchParamArray) FLandscapeBatchElementParams;
 		BatchElementParams->LandscapeUniformShaderParametersResource = &LandscapeUniformShaderParameters;
-		BatchElementParams->FixedGridUniformShaderParameters = &LandscapeFixedGridUniformShaderParameters;
+		BatchElementParams->FixedGridUniformShaderParameters = &SharedBuffers->LandscapeFixedGridUniformShaderParameters;
 		BatchElementParams->LandscapeSectionLODUniformParameters = RenderSystem.SectionLODUniformBuffer;
 		BatchElementParams->SceneProxy = this;
 		BatchElementParams->CurrentLOD = LODIndex;
@@ -3096,7 +3075,7 @@ void FLandscapeComponentSceneProxy::GetDynamicRayTracingInstances(FRayTracingIns
 			FLandscapeBatchElementParams& BatchElementParams = ParameterArray.ElementParams[SubSectionIdx];
 
 			BatchElementParams.LandscapeUniformShaderParametersResource = &LandscapeUniformShaderParameters;
-			BatchElementParams.FixedGridUniformShaderParameters = &LandscapeFixedGridUniformShaderParameters;
+			BatchElementParams.FixedGridUniformShaderParameters = &SharedBuffers->LandscapeFixedGridUniformShaderParameters;
 			BatchElementParams.LandscapeSectionLODUniformParameters = RenderSystem.SectionLODUniformBuffer;
 			BatchElementParams.SceneProxy = this;
 			BatchElementParams.CurrentLOD = CurrentLOD;
@@ -3466,6 +3445,21 @@ FLandscapeSharedBuffers::FLandscapeSharedBuffers(FRHICommandListBase& RHICmdList
 			CreateGrassIndexBuffer<uint16>(RHICmdList, InOwnerName);
 		}
 	}
+
+	// create as many as there are potential mips (even if MaxLOD can be inferior than that), because the grass could need that much :
+	const int32 NumMips = FMath::CeilLogTwo(SubsectionSizeVerts);
+	LandscapeFixedGridUniformShaderParameters.AddDefaulted(NumMips);
+	for (int32 LodIndex = 0; LodIndex < NumMips; ++LodIndex)
+	{
+		FLandscapeFixedGridUniformShaderParameters Parameters;
+		Parameters.LodValues = FVector4f(
+			static_cast<float>(LodIndex),
+			0.f,
+			(float)((SubsectionSizeVerts >> LodIndex) - 1),
+			1.f / (float)((SubsectionSizeVerts >> LodIndex) - 1));
+		LandscapeFixedGridUniformShaderParameters[LodIndex].SetContents(RHICmdList, Parameters);
+		LandscapeFixedGridUniformShaderParameters[LodIndex].InitResource(RHICmdList);
+	}
 }
 
 FLandscapeSharedBuffers::~FLandscapeSharedBuffers()
@@ -3492,12 +3486,19 @@ FLandscapeSharedBuffers::~FLandscapeSharedBuffers()
 	}
 #endif
 
+	// Free the lod uniform buffers
+	for (int32 i = 0; i < LandscapeFixedGridUniformShaderParameters.Num(); ++i)
+	{
+		LandscapeFixedGridUniformShaderParameters[i].ReleaseResource();
+	}
+
 	if (GrassIndexBuffer)
 	{
 		GrassIndexBuffer->ReleaseResource();
 		delete GrassIndexBuffer;
 	}
-
+	
+	delete FixedGridVertexFactory;
 	delete VertexFactory;
 
 	delete TileMesh;
