Landscape - Optimization to avoid creating multiple FText every frame
#rnx

diff --git a/Engine/Source/Runtime/Landscape/Private/LandscapeSubsystem.cpp b/Engine/Source/Runtime/Landscape/Private/LandscapeSubsystem.cpp
--- a/Engine/Source/Runtime/Landscape/Private/LandscapeSubsystem.cpp
+++ b/Engine/Source/Runtime/Landscape/Private/LandscapeSubsystem.cpp
@@ -1266,8 +1266,8 @@ bool ULandscapeSubsystem::GetActionableMessage(FActionableMessage& OutActionable
 
 	TRACE_CPUPROFILER_EVENT_SCOPE(ULandscapeSubsystem::GetActionableMessage);
 
-	const FText DefaultTooltip = LOCTEXT("DefaultLandscapeModified.ToolTip", "Assets that affect the Landscape may have changed. Rebuild the Landscape to see the results.");
-	const FText DefaultActionMessage = LOCTEXT("DefaultLandscapeModified.Action", "Rebuild");
+	static const FText DefaultTooltip = LOCTEXT("DefaultLandscapeModified.ToolTip", "Assets that affect the Landscape may have changed. Rebuild the Landscape to see the results.");
+	static const FText DefaultActionMessage = LOCTEXT("DefaultLandscapeModified.Action", "Rebuild");
 
 	const TArray<TTuple<ALandscapeProxy*, EOutdatedDataFlags>> OutdatedProxies = GetOutdatedProxyDetails(EOutdatedDataFlags::All, /*bInMustMatchAllFlags = */false);
 	TArray<int32> NumOutdatedProxyPerFlag;
