Fixed formatting for landscape.DumpGrassData command's tooltip

#rnx

diff --git a/Engine/Source/Runtime/Landscape/Private/LandscapeGrass.cpp b/Engine/Source/Runtime/Landscape/Private/LandscapeGrass.cpp
--- a/Engine/Source/Runtime/Landscape/Private/LandscapeGrass.cpp
+++ b/Engine/Source/Runtime/Landscape/Private/LandscapeGrass.cpp
@@ -211,16 +211,16 @@ static FAutoConsoleVariableRef CVarGGrassDensityQualityLevelCVar(
 
 static FAutoConsoleCommand ConsoleCommandDumpLandscapeGrassData(
 	TEXT("grass.DumpGrassData"),
-	TEXT("[optional: -csv -detailed -byproxy -bycomponent -bygrasstype -full] - Dumps a report of all grass data being currently used on landscape components. \
-		-csv formats the report in a CSV-friendly way. \
-		-fullnames displays the listed objects' full names, rather than the user-friendly version. \
-		-showempty will dump info even from components with no grass data \
-		-detailed shows a detailed report of all grass data, for all grass types, in all landscape components. \
-		-byproxy shows a report of grass data per landscape proxy. \
-		-bycomponent shows a report of grass data per landscape component. \
-		-bygrasstype shows a report of grass data per grass type. \
-		-full enables all sub-reports. \
-		If no report type option specified, assume full report is requested."),
+	TEXT("[optional: -csv -detailed -byproxy -bycomponent -bygrasstype -full] - Dumps a report of all grass data being currently used on landscape components. \n"
+		"-csv formats the report in a CSV-friendly way. \n"
+		"-fullnames displays the listed objects' full names, rather than the user-friendly version. \n"
+		"-showempty will dump info even from components with no grass data \n"
+		"-detailed shows a detailed report of all grass data, for all grass types, in all landscape components. \n"
+		"-byproxy shows a report of grass data per landscape proxy. \n"
+		"-bycomponent shows a report of grass data per landscape component. \n"
+		"-bygrasstype shows a report of grass data per grass type. \n"
+		"-full enables all sub-reports. \n"
+		"If no report type option specified, assume full report is requested."),
 	FConsoleCommandWithArgsDelegate::CreateLambda([](const TArray<FString>& Args)
 	{
 		bool bCSV = false;
