Landscape: Cooking: Fix assertion when loading LandscapeProxyActors in the cooker that need to RemoveObsoleteLayers. They need to delay the operation until the world is initialized.
#rnx

diff --git a/Engine/Source/Runtime/Landscape/Private/LandscapeEditLayers.cpp b/Engine/Source/Runtime/Landscape/Private/LandscapeEditLayers.cpp
--- a/Engine/Source/Runtime/Landscape/Private/LandscapeEditLayers.cpp
+++ b/Engine/Source/Runtime/Landscape/Private/LandscapeEditLayers.cpp
@@ -12070,6 +12070,16 @@ void ALandscapeProxy::DeleteUnusedLayers()
 
 bool ALandscapeProxy::RemoveObsoleteLayers(const TSet<FGuid>& InExistingLayers)
 {
+	bool bModified = false;
+	ULandscapeSubsystem* LandscapeSubsystem = GetWorld()->GetSubsystem<ULandscapeSubsystem>();
+	if (!LandscapeSubsystem)
+	{
+		// Only available when the world has been initialized, and sometimes we are called from PostLoad
+		// before that happens. Skip for now and report no change made, it will be reexecuted later when
+		// registering the actor.
+		return bModified;
+	}
+
 	TSet<TPair<FGuid, FName>> ComponentLayers;
 	for (ULandscapeComponent* Component : LandscapeComponents)
 	{
@@ -12079,11 +12089,6 @@ bool ALandscapeProxy::RemoveObsoleteLayers(const TSet<FGuid>& InExistingLayers)
 		}
 	}
 
-	bool bModified = false;
-
-	ULandscapeSubsystem* LandscapeSubsystem = GetWorld()->GetSubsystem<ULandscapeSubsystem>();
-	checkf(LandscapeSubsystem != nullptr, TEXT("RemoveObsoleteLayers can only be called when a subsystem is available"));
-
 	for (const TPair<FGuid, FName>& LayerGuidAndName : ComponentLayers)
 	{
 		if (!InExistingLayers.Contains(LayerGuidAndName.Key))
