Only consider height differences when calculating whether a landscape heightmap has changed.

This fixes an issue with the normals on the edges of the heightmap being different and marking the landscape dirty depending on whether a neighbor is loaded or not.

It also addresses rounding errors with the GPU calculation of the normal, which also cause normal differences that can get picked up as actual differences that trigger rebuilding of dependent data, and dirtying of the landscape proxies.

All of the hash related code has been moved to a new class: ULandscapeTextureHash.  This class an editor-only object that is attached as a user data on the UTexture2D to track its hash state.
This tracking object is optional, and can be applied only when needed (specifically on final renderable heightmaps and weightmaps).
The hash is initially set to the Texture SourceID (matching old hash behavior) until the texture is modified.  This ensures we don't invalidate existing dependencies with this change.

We also employ a small transient cache mapping SourceID -> Texture Hash on each of the tracking objects.  This helps to ensure we minimize the downstream dependency invalidation by returning to the original hash values after an Undo,
or after intermediate renders stomp on the texture state.

diff --git a/Engine/Source/Runtime/Landscape/Classes/Landscape.h b/Engine/Source/Runtime/Landscape/Classes/Landscape.h
--- a/Engine/Source/Runtime/Landscape/Classes/Landscape.h
+++ b/Engine/Source/Runtime/Landscape/Classes/Landscape.h
@@ -545,7 +545,6 @@ private:
 	int32 PerformLayersWeightmapsBatchedMerge(FUpdateLayersContentContext& InUpdateLayersContentContext, const FEditLayersWeightmapMergeParams& InMergeParams);
 	void ResolveLayersWeightmapTexture(const FTextureToComponentHelper& MapHelper, const TSet<UTexture2D*>& WeightmapsToResolve, bool bIntermediateRender, TArray<FLandscapeEditLayerComponentReadbackResult>& InOutComponentReadbackResults);
 
-	bool HasTextureDataChanged(TArrayView<const FColor> InOldData, TArrayView<const FColor> InNewData, bool bInIsWeightmap, uint64 InPreviousHash, uint64 InNewHash, TOptional<uint8>& OutChangedWeightmapChannelsMasks) const;
 	bool ResolveLayersTexture(FTextureToComponentHelper const& MapHelper, FLandscapeEditLayerReadback* InCPUReadBack, UTexture2D* InOutputTexture, bool bIntermediateRender,
 		TArray<FLandscapeEditLayerComponentReadbackResult>& InOutComponentReadbackResults, bool bIsWeightmap);
 
diff --git a/Engine/Source/Runtime/Landscape/Classes/LandscapeComponent.h b/Engine/Source/Runtime/Landscape/Classes/LandscapeComponent.h
--- a/Engine/Source/Runtime/Landscape/Classes/LandscapeComponent.h
+++ b/Engine/Source/Runtime/Landscape/Classes/LandscapeComponent.h
@@ -46,6 +46,8 @@ struct FLandscapeEditDataInterface;
 struct FLandscapeMobileRenderData;
 class ULandscapeHeightmapTextureEdgeFixup;
 struct FLandscapeGroup;
+enum class ELandscapeTextureUsage : uint8;
+enum class ELandscapeTextureType : uint8;
 
 //
 // FLandscapeEditToolRenderData
@@ -1085,9 +1087,11 @@ public:
 	void GenerateHeightmapMips(TArray<FColor*>& HeightmapTextureMipData, int32 ComponentX1=0, int32 ComponentY1=0, int32 ComponentX2=MAX_int32, int32 ComponentY2=MAX_int32, FLandscapeTextureDataInfo* TextureDataInfo=nullptr);
 
 	/**
-	 * Generate empty mipmaps for weightmap
+	 * if bClear is true, clears all mips on the specified texture to zero
+	 * if bClear is false, clears all mips EXCEPT MIP 0 to zero
+	 * Updates the texture hash appropriately for the given usage and texture type
 	 */
-	LANDSCAPE_API static void CreateEmptyTextureMips(UTexture2D* Texture, bool bClear = false);
+	LANDSCAPE_API static void CreateEmptyTextureMips(UTexture2D* Texture, ELandscapeTextureUsage TextureUsage, ELandscapeTextureType TextureType, bool bClear = false);
 
 	/**
 	 * Generate mipmaps for weightmap
diff --git a/Engine/Source/Runtime/Landscape/Classes/LandscapeTextureHash.h b/Engine/Source/Runtime/Landscape/Classes/LandscapeTextureHash.h
new file mode 100644
--- /dev/null
+++ b/Engine/Source/Runtime/Landscape/Classes/LandscapeTextureHash.h
@@ -0,0 +1,85 @@
+// Copyright Epic Games, Inc. All Rights Reserved.
+
+#pragma once
+
+#include "CoreMinimal.h"
+#include "Engine/AssetUserData.h"
+#include "LandscapeTextureHash.generated.h"
+
+class UTexture2D;
+
+UENUM()
+enum class ELandscapeTextureType : uint8
+{
+	Unknown,
+	Heightmap,
+	Weightmap
+};
+
+UENUM()
+enum class ELandscapeTextureUsage : uint8
+{
+	Unknown,
+	EditLayerData,			// used as data for an edit layer, input to the layer merge operation
+	FinalData				// used for runtime/rendering
+};
+
+UCLASS(NotBlueprintable, MinimalAPI, Within = Texture2D)
+class ULandscapeTextureHash : public UAssetUserData
+{
+	GENERATED_BODY()
+
+	UPROPERTY()
+	FGuid TextureHashGUID;
+
+	UPROPERTY()
+	FGuid LastSourceID;
+
+	// heightmap or weightmap.  When unknown, we fallback to using the texture source ID as hash (old behavior)
+	UPROPERTY()
+	ELandscapeTextureType TextureType = ELandscapeTextureType::Unknown;
+
+	// edit layer data or final data.  When unknown, we fallback to using the texture source ID as hash (old behavior)
+	UPROPERTY()
+	ELandscapeTextureUsage TextureUsage = ELandscapeTextureUsage::Unknown;
+
+	// cache of recently serialized hash values
+	// this ensure that if the texture is brought back to match exactly a recently serialized state, it will have exactly the same hash that it was serialized with.
+	// This helps us come back to the serialized (presumably cache-friendly) state after no-ops like:  modify + undo, hide layer + unhide layer, intermediate render stomps, etc.
+	TMap<FGuid, FGuid> RecentlySerializedHashes;
+
+public:
+
+#if WITH_EDITORONLY_DATA
+	// setup initial state on load, if it doesn't yet exist
+	static void SetInitialStateOnPostLoad(UTexture2D* LandscapeTexture, ELandscapeTextureUsage TextureUsage, ELandscapeTextureType TextureType);
+
+	// Update the stored hash based on the source data
+	static void UpdateHash(UTexture2D* LandscapeTexture, ELandscapeTextureUsage TextureUsage = ELandscapeTextureUsage::Unknown, ELandscapeTextureType TextureType = ELandscapeTextureType::Unknown, bool bForceUpdate = false);
+
+	// Explicitly set the hash for the specified LandscapeTexture
+	static void SetHash64(UTexture2D* LandscapeTexture, uint64 NewHash, ELandscapeTextureUsage TextureUsage, ELandscapeTextureType TextureType);
+
+	// calculate the texture hash directly from the source mip0 contents
+	static uint64 CalculateTextureHash64(UTexture2D* LandscapeTexture, ELandscapeTextureType TextureType);
+	static FGuid CalculateTextureHashGUID(UTexture2D* LandscapeTexture, ELandscapeTextureType TextureType);
+
+	// calculate the texture hash from the Mip0Data buffer
+	static uint64 CalculateTextureHash64(const FColor* Mip0Data, int32 PixelCount, ELandscapeTextureType TextureType);
+
+	// check if the Mip0Data is within the change threshold of OldMip0Data
+	static bool DoesTextureDataChangeExceedThreshold(const FColor* Mip0Data, const FColor* OldMip0Data, int32 PixelCount, ELandscapeTextureType TextureType, uint64 OldHash, uint64 NewHash, TOptional<uint8>& OutChangedWeightmapChannelsMasks);
+
+	// get the current stored hash for the landscape texture, checking that it is up to date
+	static FGuid GetHash(UTexture2D* LandscapeTexture);
+
+	// Check the hash is up to date
+	static void CheckHashIsUpToDate(UTexture2D* LandscapeTexture);
+
+	// UObject Interface
+	virtual void Serialize(FArchive& Ar) override;
+	virtual bool IsEditorOnly() const override { return true; }
+	virtual bool NeedsLoadForClient() const override { return false; }
+	virtual bool NeedsLoadForServer() const override { return false; }
+#endif // WITH_EDITORONLY_DATA
+};
diff --git a/Engine/Source/Runtime/Landscape/Private/Landscape.cpp b/Engine/Source/Runtime/Landscape/Private/Landscape.cpp
--- a/Engine/Source/Runtime/Landscape/Private/Landscape.cpp
+++ b/Engine/Source/Runtime/Landscape/Private/Landscape.cpp
@@ -569,7 +569,7 @@ FGuid ALandscapeProxy::GetNaniteContentId() const
 		static FGuid ExportRawMeshGuid("36208D9A475B4D93B33BF84FFEDA1536");
 		ContentStateAr << ExportRawMeshGuid;
 
-		FGuid HeightmapGuid = Component->GetHeightmap()->Source.GetId();
+		FGuid HeightmapGuid = ULandscapeTextureHash::GetHash(Component->GetHeightmap());
 		ContentStateAr << HeightmapGuid;
 
 		// Take into account the Heightmap offset per component
@@ -1268,15 +1268,23 @@ void ULandscapeComponent::PostLoad()
 
 		// Remove standalone flags from data textures to ensure data is unloaded in the editor when reverting an unsaved level.
 		// Previous version of landscape set these flags on creation.
-		if (HeightmapTexture && HeightmapTexture->HasAnyFlags(RF_Standalone))
+		if (HeightmapTexture)
 		{
-			HeightmapTexture->ClearFlags(RF_Standalone);
+			ULandscapeTextureHash::SetInitialStateOnPostLoad(HeightmapTexture, ELandscapeTextureUsage::FinalData, ELandscapeTextureType::Heightmap);
+			if (HeightmapTexture->HasAnyFlags(RF_Standalone))
+			{
+				HeightmapTexture->ClearFlags(RF_Standalone);
+			}
 		}
 		for (int32 Idx = 0; Idx < WeightmapTextures.Num(); Idx++)
 		{
-			if (WeightmapTextures[Idx] && WeightmapTextures[Idx]->HasAnyFlags(RF_Standalone))
+			ULandscapeTextureHash::SetInitialStateOnPostLoad(WeightmapTextures[Idx], ELandscapeTextureUsage::FinalData, ELandscapeTextureType::Weightmap);
+			if (WeightmapTextures[Idx])
 			{
-				WeightmapTextures[Idx]->ClearFlags(RF_Standalone);
+				if (WeightmapTextures[Idx]->HasAnyFlags(RF_Standalone))
+				{
+					WeightmapTextures[Idx]->ClearFlags(RF_Standalone);
+				}
 			}
 		}
 
@@ -2593,7 +2601,7 @@ void ULandscapeComponent::AddDefaultLayerData(const FGuid& InLayerGuid, const TA
 			}
 
 			NewLayerHeightmap->Source.UnlockMip(0);
-
+			ULandscapeTextureHash::UpdateHash(NewLayerHeightmap, ELandscapeTextureUsage::EditLayerData, ELandscapeTextureType::Heightmap);
 			NewLayerHeightmap->UpdateResource();
 		}
 
@@ -6386,7 +6394,7 @@ void ALandscapeProxy::SerializeStateHashes(FArchive& Ar)
 
 void ULandscapeComponent::SerializeStateHashes(FArchive& Ar)
 {
-	FGuid HeightmapGuid = HeightmapTexture->Source.GetId();
+	FGuid HeightmapGuid = ULandscapeTextureHash::GetHash(HeightmapTexture);
 	Ar << HeightmapGuid;
 	for (auto WeightmapTexture : WeightmapTextures)
 	{
diff --git a/Engine/Source/Runtime/Landscape/Private/LandscapeEdgeFixup.cpp b/Engine/Source/Runtime/Landscape/Private/LandscapeEdgeFixup.cpp
--- a/Engine/Source/Runtime/Landscape/Private/LandscapeEdgeFixup.cpp
+++ b/Engine/Source/Runtime/Landscape/Private/LandscapeEdgeFixup.cpp
@@ -10,6 +10,8 @@
 #include "TextureResource.h"
 #include "Hash/xxhash.h"
 #include "Streaming/TextureMipDataProvider.h"
+#include "LandscapeTextureHash.h"
+
 
 // enables verbose debug spew
 //#define ENABLE_LANDSCAPE_EDGE_FIXUP_DEBUG_SPEW 1
@@ -350,7 +352,7 @@ void FHeightmapTextureEdgeSnapshot::CaptureEdgeDataFromHeightmapSource_Internal(
 
 	TArray64<uint8> MipData;
 	verify(Source.GetMipData(MipData, 0));
-	TextureSourceID = Source.GetId();
+	TextureSourceID = ULandscapeTextureHash::GetHash(InHeightmap);
 	FIXUP_DEBUG_LOG_DETAIL(TEXT("    CaptureEdgeDataFromHeightmapSource_Internal %p (%s) - ID %s"), InHeightmap, *LandscapeGridScale.ToString(), *TextureSourceID.ToString());
 
 	TArrayView64<FHeightmapTexel> TexelData((FHeightmapTexel*) MipData.GetData(), MipData.Num() / sizeof(FHeightmapTexel));
@@ -753,7 +755,7 @@ void ULandscapeHeightmapTextureEdgeFixup::SetActiveComponent(ULandscapeComponent
 			ActiveGroup->DisableAndUnmap(this);
 
 			FIXUP_DEBUG_LOG_DETAIL(TEXT("    DisableAndUnmap %p (%d,%d) <Group %d Key %d> edgefixup:%p MAPPED:%d TOTAL:%d"),
-				ActiveComponent, GroupCoord.X, GroupCoord.Y, ActiveComponent->GetLandscapeProxy()->LODGroupKey, ActiveGroup->LandscapeGroupKey, this,
+				ActiveComponent.Get(), GroupCoord.X, GroupCoord.Y, ActiveComponent->GetLandscapeProxy()->LODGroupKey, ActiveGroup->LandscapeGroupKey, this,
 				ActiveGroup->XYToEdgeFixupMap.Num(), ActiveGroup->AllRegisteredFixups.Num());
 		}
 		else
@@ -762,7 +764,7 @@ void ULandscapeHeightmapTextureEdgeFixup::SetActiveComponent(ULandscapeComponent
 			ActiveGroup->Unmap(this);
 
 			FIXUP_DEBUG_LOG_DETAIL(TEXT("    Unmap %p (%d,%d) <Group %d Key %d> edgefixup:%p MAPPED:%d TOTAL:%d"),
-				ActiveComponent, GroupCoord.X, GroupCoord.Y, ActiveComponent->GetLandscapeProxy()->LODGroupKey, ActiveGroup->LandscapeGroupKey, this, 
+				ActiveComponent.Get(), GroupCoord.X, GroupCoord.Y, ActiveComponent->GetLandscapeProxy()->LODGroupKey, ActiveGroup->LandscapeGroupKey, this,
 				ActiveGroup->XYToEdgeFixupMap.Num(), ActiveGroup->AllRegisteredFixups.Num());
 		}
 	}
@@ -796,7 +798,7 @@ void ULandscapeHeightmapTextureEdgeFixup::RequestEdgeTexturePatchingForNeighbors
 			if (ULandscapeHeightmapTextureEdgeFixup* Neighbor = ActiveGroup->GetNeighborEdgeFixup(GroupCoord, NeighborIndex))
 			{
 				FIXUP_DEBUG_LOG_DETAIL(TEXT("  %s Neighbor %p (%d,%d) flagged for edge texture patching"), 
-					*GetDirectionString(NeighborIndex), Neighbor->ActiveComponent, Neighbor->GroupCoord.X, Neighbor->GroupCoord.Y);
+					*GetDirectionString(NeighborIndex), Neighbor->ActiveComponent.Get(), Neighbor->GroupCoord.X, Neighbor->GroupCoord.Y);
 				check(ActiveGroup->AllRegisteredFixups.Contains(Neighbor));
 				check(Neighbor->bMapped);
 				ActiveGroup->HeightmapsNeedingEdgeTexturePatching.Add(Neighbor);
@@ -1167,7 +1169,7 @@ void ULandscapeHeightmapTextureEdgeFixup::PatchTextureEdge_Internal(UE::Landscap
 						constexpr uint32 SourcePitch = 4;
 						FUpdateTextureRegion2D UpdateRegion(DstStartX, DstStartY, SrcStartX, SrcStartY, SrcWidth, SrcHeight);
 
-						FIXUP_DEBUG_LOG_RENDER(TEXT("    PatchTextureEdge_Internal [RHIUpdateTexture2D] res:%p e:%d mip:%d len:%d reg:[%d %d %d %d %d %d]"), HeightTextureResource, EdgeIndex, MipIndex, MipEdgeLength,
+						FIXUP_DEBUG_LOG_RENDER(TEXT("    PatchTextureEdge_Internal [RHIUpdateTexture2D] res:%p e:%d mip:%d reg:[%d %d %d %d %d %d]"), HeightTextureResource, EdgeIndex, MipIndex,
 							UpdateRegion.DestX, UpdateRegion.DestY, UpdateRegion.SrcX, UpdateRegion.SrcY, UpdateRegion.Width, UpdateRegion.Height);
 
 						RHIUpdateTexture2D(RHIHeightmapTexture, ResourceMipIndex, UpdateRegion, SourcePitch, Buffer);
@@ -1371,10 +1373,10 @@ UE::Landscape::ENeighborFlags ULandscapeHeightmapTextureEdgeFixup::UpdateEdgeSna
 	ENeighborFlags ChangedEdges = ENeighborFlags::None;
 
 	// no need to update if source ID matches (we assume texture has not actually changed)		
-	FGuid HeightmapSourceID = HeightmapTexture->Source.GetId();
+	FGuid HeightmapSourceID = ULandscapeTextureHash::GetHash(HeightmapTexture);
 	if (HeightmapSourceID != EdgeSnapshot->TextureSourceID)
 	{
-		FIXUP_DEBUG_LOG_DETAIL(TEXT("  UpdateEdgeSnapshotFromHeightmapSource %p HeightMap: %p -- UPDATED (ID: %s --> %s)"), ActiveComponent, HeightmapTexture,
+		FIXUP_DEBUG_LOG_DETAIL(TEXT("  UpdateEdgeSnapshotFromHeightmapSource %p HeightMap: %p -- UPDATED (ID: %s --> %s)"), ActiveComponent.Get(), HeightmapTexture.Get(),
 			*EdgeSnapshot->TextureSourceID.ToString(), *HeightmapSourceID.ToString());
 			
 		// create a new edge data
@@ -1397,7 +1399,7 @@ UE::Landscape::ENeighborFlags ULandscapeHeightmapTextureEdgeFixup::UpdateEdgeSna
 	}
 	else
 	{
-		FIXUP_DEBUG_LOG_DETAIL(TEXT("  UpdateEdgeSnapshotFromHeightmapSource %p HeightMap: %p -- SourceID matches (%s)"), ActiveComponent, HeightmapTexture, *EdgeSnapshot->TextureSourceID.ToString());
+		FIXUP_DEBUG_LOG_DETAIL(TEXT("  UpdateEdgeSnapshotFromHeightmapSource %p HeightMap: %p -- SourceID matches (%s)"), ActiveComponent.Get(), HeightmapTexture, *EdgeSnapshot->TextureSourceID.ToString());
 	}
 
 	bUpdateGPUEdgeHashes = false;
diff --git a/Engine/Source/Runtime/Landscape/Private/LandscapeEdit.cpp b/Engine/Source/Runtime/Landscape/Private/LandscapeEdit.cpp
--- a/Engine/Source/Runtime/Landscape/Private/LandscapeEdit.cpp
+++ b/Engine/Source/Runtime/Landscape/Private/LandscapeEdit.cpp
@@ -90,6 +90,7 @@ LandscapeEdit.cpp: Landscape editing
 #include "ShaderPlatformCachedIniValue.h"
 #include "DataDrivenShaderPlatformInfo.h"
 #include "Misc/TransactionObjectEvent.h"
+#include "LandscapeTextureStorageProvider.h"
 #endif
 #include "Algo/Compare.h"
 #include "LandscapeProxy.h"
@@ -1183,7 +1184,8 @@ TArray<UTexture*> ULandscapeComponent::RepairInvalidTextures()
 		if (!IsValidLandscapeTextureSourceData(*Texture))
 		{
 			UE_LOG(LogLandscape, Error, TEXT("Invalid data found in texture %s from landscape component %s : clearing data."), *Texture->GetName(), *GetPathName());
-			CreateEmptyTextureMips(CastChecked<UTexture2D>(Texture), true);
+			const bool bClear = true;
+			CreateEmptyTextureMips(CastChecked<UTexture2D>(Texture), ELandscapeTextureUsage::Unknown, ELandscapeTextureType::Unknown, bClear);
 			Texture->PostEditChange();
 			InvalidTextures.Add(Texture);
 		}
@@ -2593,7 +2595,7 @@ void ULandscapeComponent::GenerateHeightmapMips(TArray<FColor*>& HeightmapTextur
 	}
 }
 
-void ULandscapeComponent::CreateEmptyTextureMips(UTexture2D* Texture, bool bClear /*= false*/)
+void ULandscapeComponent::CreateEmptyTextureMips(UTexture2D* Texture, ELandscapeTextureUsage TextureUsage, ELandscapeTextureType TextureType, bool bClear /*= false*/)
 {
 	check( Texture->Source.IsValid() );
 	ETextureSourceFormat Format = Texture->Source.GetFormat();
@@ -2604,12 +2606,18 @@ void ULandscapeComponent::CreateEmptyTextureMips(UTexture2D* Texture, bool bClea
 	{
 		Texture->Source.Init2DWithMipChain(SizeU, SizeV, Format);
 		int32 NumMips = Texture->Source.GetNumMips();
-		for (int32 MipIndex = 0; MipIndex < NumMips; ++MipIndex)
+		if (NumMips > 0)
 		{
-			uint8* MipData = Texture->Source.LockMip(MipIndex);
-			check( MipData );
-			FMemory::Memzero(MipData, Texture->Source.CalcMipSize(MipIndex));
-			Texture->Source.UnlockMip(MipIndex);
+			Texture->Source.LockMip(0); // taking an outer lock ensures we don't do expensive re-hashing until all mip modification is complete
+			for (int32 MipIndex = 0; MipIndex < NumMips; ++MipIndex)
+			{
+				uint8* MipData = Texture->Source.LockMip(MipIndex);
+				check( MipData );
+				FMemory::Memzero(MipData, Texture->Source.CalcMipSize(MipIndex));
+				Texture->Source.UnlockMip(MipIndex);
+			}
+			Texture->Source.UnlockMip(0); // release outer lock
+			ULandscapeTextureHash::UpdateHash(Texture, TextureUsage, TextureType);
 		}
 	}
 	else
@@ -2622,6 +2630,7 @@ void ULandscapeComponent::CreateEmptyTextureMips(UTexture2D* Texture, bool bClea
 		check( MipData );
 		FMemory::Memcpy(MipData, TopMipData.GetData(), TopMipData.Num());
 		Texture->Source.UnlockMip(0);
+		ULandscapeTextureHash::CheckHashIsUpToDate(Texture);	
 	}
 }
 
@@ -2644,6 +2653,8 @@ void ULandscapeComponent::GenerateMipsTempl(int32 InNumSubsections, int32 InSubs
 	{
 		Texture->Source.UnlockMip(i);
 	}
+
+	// Note that we do not need to update the heightmap hash here (even if it was a heightmap) because the hash is only based on mip 0, which we don't change in this function
 }
 
 void ULandscapeComponent::GenerateWeightmapMips(int32 InNumSubsections, int32 InSubsectionSizeQuads, UTexture2D* WeightmapTexture, FColor* BaseMipData)
@@ -3816,7 +3827,7 @@ void ALandscapeProxy::Import(const FGuid& InGuid, int32 InMinX, int32 InMinY, in
 		{
 			HeightmapInfo.HeightmapTexture->Source.UnlockMip(i);
 		}
-
+		ULandscapeTextureHash::UpdateHash(HeightmapInfo.HeightmapTexture, ELandscapeTextureUsage::FinalData, ELandscapeTextureType::Heightmap);
 		HeightmapInfo.HeightmapTexture->BeginCachePlatformData();
 		HeightmapInfo.HeightmapTexture->ClearAllCachedCookedPlatformData();
 		PendingTexturePlatformDataCreation.Add(HeightmapInfo.HeightmapTexture);
@@ -6002,7 +6013,8 @@ void ALandscape::SplitHeightmap(ULandscapeComponent* Comp, ALandscapeProxy* Targ
 
 		// Create the new heightmap texture
 		NewHeightmapTexture = DstProxy->CreateLandscapeTexture(HeightmapSizeU, HeightmapSizeV, TEXTUREGROUP_Terrain_Heightmap, TSF_BGRA8);
-		ULandscapeComponent::CreateEmptyTextureMips(NewHeightmapTexture, true);
+		const bool bClear = true;
+		ULandscapeComponent::CreateEmptyTextureMips(NewHeightmapTexture, ELandscapeTextureUsage::FinalData, ELandscapeTextureType::Heightmap, bClear);
 		NewHeightmapTexture->PostEditChange();
 
 		Comp->HeightmapScaleBias = NewHeightmapScaleBias;
@@ -6131,7 +6143,8 @@ void ALandscape::SplitHeightmap(ULandscapeComponent* Comp, ALandscapeProxy* Targ
 					{
 						// no mipchain required as these layer weight maps are used in layer compositing to generate a final set of weight maps to be used for rendering
 						UTexture2D* LayerHeightmapTexture = DstProxy->CreateLandscapeTexture(HeightmapSizeU, HeightmapSizeV, TEXTUREGROUP_Terrain_Heightmap, TSF_BGRA8, /* OptionalOverrideOuter = */ nullptr, /* bCompress = */ false, /* bMipChain = */ false);
-						ULandscapeComponent::CreateEmptyTextureMips(LayerHeightmapTexture, true);
+						const bool bClear = true;
+						ULandscapeComponent::CreateEmptyTextureMips(LayerHeightmapTexture, ELandscapeTextureUsage::FinalData, ELandscapeTextureType::Heightmap, bClear);
 						LayerHeightmapTexture->PostEditChange();
 						// Set Layer heightmap texture
 						LayerData.HeightmapData.Texture = LayerHeightmapTexture;
@@ -7507,7 +7520,8 @@ TArray<UTexture*> ULandscapeComponent::ReallocateWeightmaps(FLandscapeEditDataIn
 			// Alloc dummy mips
 			if (bIsFinalWeightmap)
 			{
-				CreateEmptyTextureMips(CurrentWeightmapTexture, true);
+				const bool bClear = true;
+				CreateEmptyTextureMips(CurrentWeightmapTexture, ELandscapeTextureUsage::FinalData, ELandscapeTextureType::Weightmap, bClear);
 			}
 			else
 			{
@@ -7764,11 +7778,13 @@ void ULandscapeComponent::InitHeightmapData(TArray<FColor>& Heights, bool bUpdat
 			SimpleCollisionMipLevel > CollisionMipLevel ? HeightmapTextureMipData[SimpleCollisionMipLevel] : nullptr);
 	}
 
+	UTexture2D* Heightmap = GetHeightmap();
 	for (int32 i = 0; i < HeightmapTextureMipData.Num(); i++)
 	{
-		GetHeightmap()->Source.UnlockMip(i);
+		Heightmap->Source.UnlockMip(i);
 	}
-	GetHeightmap()->PostEditChange();
+	ULandscapeTextureHash::UpdateHash(Heightmap, ELandscapeTextureUsage::FinalData, ELandscapeTextureType::Heightmap);
+	Heightmap->PostEditChange();
 }
 
 void ULandscapeComponent::InitWeightmapData(TArray<ULandscapeLayerInfoObject*>& LayerInfos, TArray<TArray<uint8> >& WeightmapData)
@@ -8016,7 +8032,8 @@ void ULandscapeComponent::GenerateMobilePlatformPixelData(bool bIsCooking, const
 		for (int32 i = 0; i < NumWeightTextures; ++i)
 		{
 			UTexture2D* CurrentWeightmapTexture  = GetLandscapeProxy()->CreateLandscapeTexture(WeightmapSize, WeightmapSize, TEXTUREGROUP_Terrain_Weightmap, TSF_BGRA8);
-			CreateEmptyTextureMips(CurrentWeightmapTexture, true);
+			const bool bClear = true;
+			CreateEmptyTextureMips(CurrentWeightmapTexture, ELandscapeTextureUsage::FinalData, ELandscapeTextureType::Weightmap, bClear);
 			MobileWeightmapTextures[i] = CurrentWeightmapTexture;
 		}
 	}
diff --git a/Engine/Source/Runtime/Landscape/Private/LandscapeEditInterface.cpp b/Engine/Source/Runtime/Landscape/Private/LandscapeEditInterface.cpp
--- a/Engine/Source/Runtime/Landscape/Private/LandscapeEditInterface.cpp
+++ b/Engine/Source/Runtime/Landscape/Private/LandscapeEditInterface.cpp
@@ -22,6 +22,7 @@ LandscapeEditInterface.cpp: Landscape editing interface
 #include "LandscapeRender.h"
 #include "LandscapePrivate.h"
 #include "ComponentReregisterContext.h"
+#include "LandscapeTextureHash.h"
 #include "Algo/Transform.h"
 #include "TextureCompiler.h"
 
@@ -4894,7 +4895,8 @@ void FLandscapeEditDataInterface::SetXYOffsetDataTempl(int32 X1, int32 Y1, int32
 					int32 WeightmapSize = (Component->SubsectionSizeQuads+1) * Component->NumSubsections;
 					XYOffsetTexture = Component->GetLandscapeProxy()->CreateLandscapeTexture(WeightmapSize, WeightmapSize, TEXTUREGROUP_Terrain_Weightmap, TSF_BGRA8);
 					// Alloc dummy mips
-					ULandscapeComponent::CreateEmptyTextureMips(XYOffsetTexture, true);
+					const bool bClear = true;
+					ULandscapeComponent::CreateEmptyTextureMips(XYOffsetTexture, ELandscapeTextureUsage::FinalData, ELandscapeTextureType::Unknown, bClear);
 
 					// TODO: we flag somewhere that we need to use XYOffset Factory, so we can recompute shaders based on this change
 
@@ -5694,8 +5696,10 @@ void FLandscapeEditDataInterface::GetXYOffsetDataFast(const int32 X1, const int3
 // FLandscapeTextureDataInfo
 //
 
-FLandscapeTextureDataInfo::FLandscapeTextureDataInfo(UTexture2D* InTexture, bool bShouldDirtyPackage)
-:	Texture(InTexture)
+FLandscapeTextureDataInfo::FLandscapeTextureDataInfo(UTexture2D* InTexture, bool bShouldDirtyPackage, ELandscapeTextureUsage InTextureUsage, ELandscapeTextureType InTextureType)
+	: Texture(InTexture)
+	, TextureUsage(InTextureUsage)
+	, TextureType(InTextureType)
 {
 	MipInfo.AddZeroed(Texture->Source.GetNumMips());
 	Texture->SetFlags(RF_Transactional);
@@ -5774,6 +5778,7 @@ FLandscapeTextureDataInfo::~FLandscapeTextureDataInfo()
 		}
 	}
 	Texture->ClearFlags(RF_Transactional);
+	ULandscapeTextureHash::UpdateHash(Texture, TextureUsage, TextureType);
 }
 
 #endif // WITH_EDITOR
diff --git a/Engine/Source/Runtime/Landscape/Private/LandscapeEditLayers.cpp b/Engine/Source/Runtime/Landscape/Private/LandscapeEditLayers.cpp
--- a/Engine/Source/Runtime/Landscape/Private/LandscapeEditLayers.cpp
+++ b/Engine/Source/Runtime/Landscape/Private/LandscapeEditLayers.cpp
@@ -86,6 +86,7 @@ LandscapeEditLayers.cpp: Landscape editing layers mode
 #include "UnrealEdGlobals.h" // GUnrealEd
 #include "Editor/UnrealEdEngine.h"
 #include "VisualLogger/VisualLogger.h"
+#include "LandscapeTextureStorageProvider.h"
 #endif
 
 #define LOCTEXT_NAMESPACE "Landscape"
@@ -7748,9 +7749,8 @@ int32 ALandscape::RegenerateLayersHeightmaps(const FUpdateLayersContentContext&
 			if (CPUReadback == nullptr)
 			{
 				FLandscapeEditLayerReadback* NewCPUReadback = new FLandscapeEditLayerReadback();
-				const uint8* LockedMip = ComponentHeightmap->Source.LockMip(0);
-				const uint64 Hash = FLandscapeEditLayerReadback::CalculateHash(LockedMip, ComponentHeightmap->Source.GetSizeX() * ComponentHeightmap->Source.GetSizeY() * sizeof(FColor));
-				ComponentHeightmap->Source.UnlockMip(0);
+				// gather the existing hash, pre-readback
+				const uint64 Hash = ULandscapeTextureHash::CalculateTextureHash64(ComponentHeightmap, ELandscapeTextureType::Heightmap);
 				NewCPUReadback->SetHash(Hash);
 				Proxy->HeightmapsCPUReadback.Add(ComponentHeightmap, NewCPUReadback);
 			}
@@ -8203,103 +8203,6 @@ void ALandscape::OnDirtyHeightmap(FTextureToComponentHelper const& MapHelper, UT
 	}
 }
 
-bool ALandscape::HasTextureDataChanged(TArrayView<const FColor> InOldData, TArrayView<const FColor> InNewData, bool bInIsWeightmap, uint64 InPreviousHash, uint64 InNewHash, TOptional<uint8>& OutChangedWeightmapChannelsMasks) const
-{
-	TRACE_CPUPROFILER_EVENT_SCOPE(ALandscape::HasTextureDataChanged);
-
-	if (InPreviousHash != InNewHash)
-	{
-		int32 TextureSize = InOldData.Num();
-		check(TextureSize == InNewData.Num());
-
-		// If necessary, perform a deep comparison to avoid taking into account precision issues as actual changes :
-		if (bInIsWeightmap)
-		{
-			int32 DirtyWeightmapThreshold = CVarLandscapeDirtyWeightmapThreshold.GetValueOnGameThread();
-			if ((DirtyWeightmapThreshold > 0) || OutChangedWeightmapChannelsMasks.IsSet())
-			{
-				bool bHasChanged = false;
-
-				TRACE_CPUPROFILER_EVENT_SCOPE(DeepCompareWeightmap);
-				for (int32 Index = 0; Index < TextureSize; ++Index)
-				{
-					const FColor& OldColor = InOldData[Index];
-					const FColor& NewColor = InNewData[Index];
-					if (OldColor != NewColor)
-					{
-						auto DiffChannel = [DirtyWeightmapThreshold](uint8 InOldValue, uint8 InNewValue) -> bool
-							{
-								uint8 Diff = (InNewValue > InOldValue) ? (InNewValue - InOldValue) : (InOldValue - InNewValue);
-								return (Diff > DirtyWeightmapThreshold);
-							};
-
-						uint8 DiffMask = ((uint8)(DiffChannel(OldColor.R, NewColor.R) ? 1 : 0) << 0)
-							| ((uint8)(DiffChannel(OldColor.G, NewColor.G) ? 1 : 0) << 1)
-							| ((uint8)(DiffChannel(OldColor.B, NewColor.B) ? 1 : 0) << 2)
-							| ((uint8)(DiffChannel(OldColor.A, NewColor.A) ? 1 : 0) << 3);
-
-						if (DiffMask != 0)
-						{
-							bHasChanged = true;
-							if (OutChangedWeightmapChannelsMasks.IsSet())
-							{
-								*OutChangedWeightmapChannelsMasks |= DiffMask;
-							}
-							else
-							{
-								// no need to report which channel has been changed, early out : 
-								return true;
-							}
-						}
-					}
-				}
-
-				return bHasChanged;
-			}
-		}
-		else
-		{
-			int32 DirtyHeightmapHeightThreshold = CVarLandscapeDirtyHeightmapHeightThreshold.GetValueOnGameThread();
-			int32 DirtyHeightmapNormalThreshold = CVarLandscapeDirtyHeightmapNormalThreshold.GetValueOnGameThread();
-			if (DirtyHeightmapHeightThreshold > 0 || DirtyHeightmapNormalThreshold > 0)
-			{
-				TRACE_CPUPROFILER_EVENT_SCOPE(DeepCompareHeightmap);
-				for (int32 Index = 0; Index < TextureSize; ++Index)
-				{
-					const FColor& OldColor = InOldData[Index];
-					const FColor& NewColor = InNewData[Index];
-					if (OldColor != NewColor)
-					{
-						uint16 OldHeight = ((static_cast<uint16>(OldColor.R) << 8) | static_cast<uint16>(OldColor.G));
-						uint16 NewHeight = ((static_cast<uint16>(NewColor.R) << 8) | static_cast<uint16>(NewColor.G));
-						if (uint16 Diff = (NewHeight > OldHeight) ? (NewHeight - OldHeight) : (OldHeight - NewHeight); Diff > DirtyHeightmapHeightThreshold)
-						{
-							return true;
-						}
-						if (uint8 Diff = (NewColor.B > OldColor.B) ? (NewColor.B - OldColor.B) : (OldColor.B - NewColor.B); Diff > DirtyHeightmapNormalThreshold)
-						{
-							return true;
-						}
-						if (uint8 Diff = (NewColor.A > OldColor.A) ? (NewColor.A - OldColor.A) : (OldColor.A - NewColor.A); Diff > DirtyHeightmapNormalThreshold)
-						{
-							return true;
-						}
-					}
-				}
-
-				// No significant difference detected : 
-				return false;
-			}
-		}
-
-		// Hash is different and we don't need a deep comparison :
-		return true;
-	}
-
-	// Hash is identical :
-	return false;
-}
-
 bool ALandscape::ResolveLayersTexture(
 	FTextureToComponentHelper const& MapHelper,
 	FLandscapeEditLayerReadback* InCPUReadback,
@@ -8316,7 +8219,7 @@ bool ALandscape::ResolveLayersTexture(
 
 	bool bUserTriggered = false;
 
-	if (TArray<ULandscapeComponent*> const* Components = bIsWeightmap ? MapHelper.WeightmapToComponents.Find(InOutputTexture) : MapHelper.HeightmapToComponents.Find(InOutputTexture))
+	if (TArray<ULandscapeComponent*> const * Components = bIsWeightmap ? MapHelper.WeightmapToComponents.Find(InOutputTexture) : MapHelper.HeightmapToComponents.Find(InOutputTexture))
 	{
 		for (ULandscapeComponent* Component : *Components)
 		{
@@ -8345,25 +8248,29 @@ bool ALandscape::ResolveLayersTexture(
 		TRACE_CPUPROFILER_EVENT_SCOPE(LandscapeLayers_PerformReadbacks);
 
 		// Copy final result to texture source.
-		TArray<TArray<FColor>> const& OutMipsData = InCPUReadback->GetResult(CompletedReadbackNum - 1);
+		TArray<TArray<FColor>> const& NewMipsData = InCPUReadback->GetResult(CompletedReadbackNum - 1);
+
+		ELandscapeTextureType TextureType = bIsWeightmap ? ELandscapeTextureType::Weightmap : ELandscapeTextureType::Heightmap;
 
 		// Keep track if we locked the texture for write or not.  bChanged implies bLockedForWrite, but not the other way around.  We might have
 		// bLockedForWrite because of a data hash change, but bChanged=false due to the threshold check.
 		bool bLockedForWrite = false;
+		uint64 NewHash = 0;
+		uint64 OldHash = InCPUReadback->GetHash();
 
-		for (int8 MipIndex = 0; MipIndex < OutMipsData.Num(); ++MipIndex)
+		for (int8 MipIndex = 0; MipIndex < NewMipsData.Num(); ++MipIndex)
 		{
-			int32 TextureSize = OutMipsData[MipIndex].Num();
-			if (TextureSize > 0)
+			int32 MipTexels = NewMipsData[MipIndex].Num();
+			if (MipTexels > 0)
 			{
-				uint8* TextureData = nullptr;
+				FColor* MipDataWriteable = nullptr;
+				const FColor* NewMipData = NewMipsData[MipIndex].GetData();
 
 				// Do dirty detection on first mip.
 				// Don't do this for intermediate renders.
 				if (MipIndex == 0 && !bIntermediateRender)
 				{
-					uint64 OldHash = InCPUReadback->GetHash();
-					uint64 NewHash = FLandscapeEditLayerReadback::CalculateHash((uint8*)OutMipsData[MipIndex].GetData(), TextureSize * sizeof(FColor));
+					NewHash = ULandscapeTextureHash::CalculateTextureHash64(NewMipData, MipTexels, TextureType);
 
 					// If the last readback bypassed the hash test because it was an intermediate render, force a lock and copy to restore the "correct" non-intermediate data.  The regular
 					// significance test and conditional notification can proceed as normal.
@@ -8378,35 +8285,52 @@ bool ALandscape::ResolveLayersTexture(
 					{
 						// Defer locking the texture for ReadWrite until after we know that the data is changing at all.  Unlocking after a ReadWrite causes an 
 						// expensive rehashing using the slower hash function used in IoHash.  
-						TextureData = InOutputTexture->Source.LockMip(MipIndex);
-						check(!bLockedForWrite);
-						bLockedForWrite = true;
+						const FColor* OldMipDataReadOnly = (const FColor*) InOutputTexture->Source.LockMipReadOnly(MipIndex);
+						check(OldMipDataReadOnly);
 
-						if (HasTextureDataChanged(MakeArrayView(reinterpret_cast<const FColor*>(TextureData), TextureSize), MakeArrayView(OutMipsData[MipIndex].GetData(), TextureSize), bIsWeightmap, OldHash, NewHash, ChangedChannelsMask))
+						if (ULandscapeTextureHash::DoesTextureDataChangeExceedThreshold(OldMipDataReadOnly, NewMipData, MipTexels, bIsWeightmap ? ELandscapeTextureType::Weightmap : ELandscapeTextureType::Heightmap, OldHash, NewHash, ChangedChannelsMask))
 						{
-							bChanged |= InCPUReadback->SetHash(NewHash);
-							check(bChanged);
-						}
+							// convert the lock to a read/write lock (first have to release the existing read-only lock, then acquire a new read/write lock)
+							InOutputTexture->Source.UnlockMip(MipIndex);
+							OldMipDataReadOnly = nullptr;
+							MipDataWriteable = (FColor*) InOutputTexture->Source.LockMip(MipIndex);
+							check(!bLockedForWrite);
+							bLockedForWrite = true;
+
+							InCPUReadback->SetHash(NewHash);
+							bChanged = true;
 
-						if (bChanged)
-						{
 							// We're about to modify the texture's source data, the texture needs to know so that it can handle properly update cached platform data (additionally, the package needs to be dirtied) :
+							ULandscapeInfo* LandscapeInfo = GetLandscapeInfo();
 							if (GetDefault<ULandscapeSettings>()->LandscapeDirtyingMode == ELandscapeDirtyingMode::InLandscapeModeAndUserTriggeredChanges)
 							{
-								FLandscapeDirtyOnlyInModeScope Scope(GetLandscapeInfo(), !bUserTriggered);
-								GetLandscapeInfo()->ModifyObject(InOutputTexture);
+								FLandscapeDirtyOnlyInModeScope Scope(LandscapeInfo, !bUserTriggered);
+								LandscapeInfo->ModifyObject(InOutputTexture);
+								if (ULandscapeTextureHash* TextureHash = InOutputTexture->GetAssetUserData<ULandscapeTextureHash>())
+								{
+									LandscapeInfo->ModifyObject(TextureHash);
+								}
 							}
 							else
 							{
-								GetLandscapeInfo()->ModifyObject(InOutputTexture);
+								LandscapeInfo->ModifyObject(InOutputTexture);
+								if (ULandscapeTextureHash* TextureHash = InOutputTexture->GetAssetUserData<ULandscapeTextureHash>())
+								{
+									LandscapeInfo->ModifyObject(TextureHash);
+								}
 							}
 						}
+						else
+						{
+							// release the read-only lock
+							InOutputTexture->Source.UnlockMip(MipIndex);
+						}
 					}
 				}
 				else if (bLockedForWrite || bIntermediateRender)
 				{
 					// If we locked mip 0, lock the rest too, even though FTextureSource isn't tracking mip-level locking.
-					TextureData = InOutputTexture->Source.LockMip(MipIndex);
+					MipDataWriteable = (FColor*) InOutputTexture->Source.LockMip(MipIndex);
 					bLockedForWrite = true;
 
 					if (MipIndex == 0 && bIntermediateRender)
@@ -8420,21 +8344,24 @@ bool ALandscape::ResolveLayersTexture(
 
 				if (bChanged)
 				{
-					check(bLockedForWrite && TextureData);
+					// issue callbacks before overwriting the MipData (so the callback can compare old vs new)
+					check(bLockedForWrite && MipDataWriteable);
 					if (bIsWeightmap)
 					{
-						OnDirtyWeightmap(MapHelper, InOutputTexture, (FColor*)TextureData, OutMipsData[MipIndex].GetData(), MipIndex, ChangedChannelsMask.Get(0));
+						// NOTE: comparing the New and Old data can be misleading when bOldDataIsFromIntermediateRender
+						OnDirtyWeightmap(MapHelper, InOutputTexture, (FColor*)MipDataWriteable, NewMipData, MipIndex, ChangedChannelsMask.Get(0));
 					}
 					else
 					{
-						OnDirtyHeightmap(MapHelper, InOutputTexture, (FColor*)TextureData, OutMipsData[MipIndex].GetData(), MipIndex);
+						// NOTE: comparing the New and Old data can be misleading when bOldDataIsFromIntermediateRender
+						OnDirtyHeightmap(MapHelper, InOutputTexture, (FColor*)MipDataWriteable, NewMipData, MipIndex);
 					}
 				}
 
 				if (bLockedForWrite)
 				{
 					TRACE_CPUPROFILER_EVENT_SCOPE(ReadbackToCPU);
-					FMemory::Memcpy(TextureData, OutMipsData[MipIndex].GetData(), OutMipsData[MipIndex].Num() * sizeof(FColor));
+					FMemory::Memcpy(MipDataWriteable, NewMipData, MipTexels * sizeof(FColor));
 				}
 			}
 		}
@@ -8442,18 +8369,23 @@ bool ALandscape::ResolveLayersTexture(
 		if (bLockedForWrite)
 		{
 			// Unlock all mips at once because there's a lock counter in FTextureSource that recomputes the content hash when reaching 0 (which means we'd recompute the hash several times over if we Lock/Unlock/Lock/Unloc/... for each mip ):
-			for (int8 MipIndex = 0; MipIndex < OutMipsData.Num(); ++MipIndex)
+			for (int8 MipIndex = 0; MipIndex < NewMipsData.Num(); ++MipIndex)
 			{
-				if (OutMipsData[MipIndex].Num() > 0)
+				if (NewMipsData[MipIndex].Num() > 0)
 				{
 					InOutputTexture->Source.UnlockMip(MipIndex);
 				}
 			}
+
+			check(!bIntermediateRender || !bChanged);	// intermediate renders should not be considered changed (we keep the same old hash)
+
+			// update the hash (New hash if it's considered changed, otherwise continue to use the old hash)
+			// this must happen after UnlockMip, so the SourceID is up to date.
+			ULandscapeTextureHash::SetHash64(InOutputTexture, bChanged ? NewHash : OldHash, ELandscapeTextureUsage::FinalData, bIsWeightmap ? ELandscapeTextureType::Weightmap : ELandscapeTextureType::Heightmap);
 		}
 
 		// change lighting guid to be the hash of the source data (so we can use lighting guid to detect when it actually changes)
-		// TODO [chris.tchou] we should check(InOutputTexture->Source.bGuidIsHash); .. but it's not a public value currently.
-		InOutputTexture->SetLightingGuid(InOutputTexture->Source.GetId());
+		InOutputTexture->SetLightingGuid(ULandscapeTextureHash::GetHash(InOutputTexture));
 
 		// Find out whether some channels from this weightmap are now all zeros : 
 		static constexpr uint32 AllChannelsAllZerosMask = 15;
@@ -8463,7 +8395,7 @@ bool ALandscape::ResolveLayersTexture(
 		{
 			TRACE_CPUPROFILER_EVENT_SCOPE(LandscapeLayers_AnalyzeWeightmap);
 			const FColor* TextureData = reinterpret_cast<const FColor*>(InOutputTexture->Source.LockMipReadOnly(0));
-			const int32 TexSize = OutMipsData[0].Num();
+			const int32 TexSize = NewMipsData[0].Num();
 			// We can stop iterating as soon as all of the channels are non-zero :
 			for (int32 Index = 0; (Index < TexSize) && (AllZerosTextureChannelMask != 0); ++Index)
 			{
@@ -9724,10 +9656,8 @@ void ALandscape::PrepareLayersWeightmapsLocalMergeRenderThreadData(const FUpdate
 						{
 							// Lazily create the readback objects as required (ReallocateLayersWeightmaps might have created new weightmaps)
 							FLandscapeEditLayerReadback* NewCPUReadback = new FLandscapeEditLayerReadback();
-							const uint8* LockedMip = ComponentWeightmap->Source.LockMipReadOnly(0);
-							const uint64 Hash = FLandscapeEditLayerReadback::CalculateHash(LockedMip, ComponentWeightmap->Source.GetSizeX() * ComponentWeightmap->Source.GetSizeY() * sizeof(FColor));
+							const uint64 Hash = ULandscapeTextureHash::CalculateTextureHash64(ComponentWeightmap, ELandscapeTextureType::Weightmap);
 							NewCPUReadback->SetHash(Hash);
-							ComponentWeightmap->Source.UnlockMip(0);
 							CPUReadback = &Proxy->WeightmapsCPUReadback.Add(ComponentWeightmap, NewCPUReadback);
 						}
 						check(*CPUReadback != nullptr);
@@ -10226,10 +10156,8 @@ int32 ALandscape::PerformLayersWeightmapsGlobalMerge(FUpdateLayersContentContext
 					if (CPUReadback == nullptr)
 					{
 						FLandscapeEditLayerReadback* NewCPUReadback = new FLandscapeEditLayerReadback();
-						const uint8* LockedMip = WeightmapTexture->Source.LockMipReadOnly(0);
-						const uint64 Hash = FLandscapeEditLayerReadback::CalculateHash(LockedMip, WeightmapTexture->Source.GetSizeX() * WeightmapTexture->Source.GetSizeY() * sizeof(FColor));
+						const uint64 Hash = ULandscapeTextureHash::CalculateTextureHash64(WeightmapTexture, ELandscapeTextureType::Weightmap);
 						NewCPUReadback->SetHash(Hash);
-						WeightmapTexture->Source.UnlockMip(0);
 						Proxy->WeightmapsCPUReadback.Add(WeightmapTexture, NewCPUReadback);
 					}
 				}
@@ -10554,10 +10482,8 @@ int32 ALandscape::PerformLayersWeightmapsBatchedMerge(FUpdateLayersContentContex
 							{
 								// Lazily create the readback objects as required (ReallocateLayersWeightmaps might have created new weightmaps)
 								FLandscapeEditLayerReadback* NewCPUReadback = new FLandscapeEditLayerReadback();
-								const uint8* LockedMip = Weightmap->Source.LockMipReadOnly(0);
-								const uint64 Hash = FLandscapeEditLayerReadback::CalculateHash(LockedMip, Weightmap->Source.GetSizeX() * Weightmap->Source.GetSizeY() * sizeof(FColor));
+								const uint64 Hash = ULandscapeTextureHash::CalculateTextureHash64(Weightmap, ELandscapeTextureType::Weightmap);
 								NewCPUReadback->SetHash(Hash);
-								Weightmap->Source.UnlockMip(0);
 								CPUReadback = &Proxy->WeightmapsCPUReadback.Add(Weightmap, NewCPUReadback);
 							}
 							check(*CPUReadback != nullptr);
@@ -11976,6 +11902,8 @@ bool ALandscapeProxy::HasLayersContent() const
 
 void ALandscapeProxy::UpdateCachedHasLayersContent(bool InCheckComponentDataIntegrity)
 {
+	// TODO: do we need to update the texture types here?
+
 	// In the case of InCheckComponentDataIntegrity we will loop through all components to make sure they all have the same state and in the other case we will assume that the 1st component represent the state of all the others.
 	bHasLayersContent = (!LandscapeComponents.IsEmpty() && (LandscapeComponents[0] != nullptr)) ? LandscapeComponents[0]->HasLayersData() : false;
 
diff --git a/Engine/Source/Runtime/Landscape/Private/LandscapeEditReadback.cpp b/Engine/Source/Runtime/Landscape/Private/LandscapeEditReadback.cpp
--- a/Engine/Source/Runtime/Landscape/Private/LandscapeEditReadback.cpp
+++ b/Engine/Source/Runtime/Landscape/Private/LandscapeEditReadback.cpp
@@ -332,13 +332,6 @@ FLandscapeEditLayerReadback::~FLandscapeEditLayerReadback()
 	ensure(!bLastReadbackWasIntermediate);  // Expecting intermediate render data to always be cleaned up by a regular render.  It shouldn't escape the lifetime of this object.
 }
 
-uint64 FLandscapeEditLayerReadback::CalculateHash(const uint8* InMipData, int32 InSizeInBytes)
-{
-	TRACE_CPUPROFILER_EVENT_SCOPE(FLandscapeEditLayerReadback::CalculateHash);
-
-	return CityHash64(reinterpret_cast<const char*>(InMipData), InSizeInBytes);
-}
-
 bool FLandscapeEditLayerReadback::SetHash(uint64 InHash)
 {
 	const bool bChanged = InHash != Hash;
diff --git a/Engine/Source/Runtime/Landscape/Private/LandscapeEditReadback.h b/Engine/Source/Runtime/Landscape/Private/LandscapeEditReadback.h
--- a/Engine/Source/Runtime/Landscape/Private/LandscapeEditReadback.h
+++ b/Engine/Source/Runtime/Landscape/Private/LandscapeEditReadback.h
@@ -19,9 +19,6 @@ public:
 	FLandscapeEditLayerReadback();
 	~FLandscapeEditLayerReadback();
 
-	/** Hash function to generate a hash used to detect change in read back data. This expects to take raw data from first mip. */
-	static uint64 CalculateHash(const uint8* InMipData, int32 InSizeInBytes);
-
 	/** Update the stored hash value. Return true if this changes the value. */
 	bool SetHash(uint64 InHash);
 	/** Get the stored hash value. */
diff --git a/Engine/Source/Runtime/Landscape/Private/LandscapeGrassMapsBuilder.cpp b/Engine/Source/Runtime/Landscape/Private/LandscapeGrassMapsBuilder.cpp
--- a/Engine/Source/Runtime/Landscape/Private/LandscapeGrassMapsBuilder.cpp
+++ b/Engine/Source/Runtime/Landscape/Private/LandscapeGrassMapsBuilder.cpp
@@ -16,6 +16,7 @@
 #include "Materials/MaterialInstance.h"
 #include "Stats/Stats.h"
 #include "Components/InstancedStaticMeshComponent.h"
+#include "LandscapeTextureHash.h"
 #include "SceneInterface.h"
 
 #define LOCTEXT_NAMESPACE "Landscape"
@@ -151,16 +152,14 @@ namespace UE::Landscape
 		// Take into account any material state change : (excluding texture state)
 		Hash = FCrc::TypeCrc32(Material->ComputeAllStateCRC(), Hash);
 
-		// hash the heightmap texture (we use Source Id as it is a content hash, lighting guid is random)
-		const UTexture2D* Heightmap = Component->GetHeightmap();
-		check(Heightmap->Source.IsValid());
-		Hash = FCrc::TypeCrc32(Heightmap->Source.GetId(), Hash);
+		// hash the heightmap texture (we use the component's heightmap texture hash that ignores normals)
+		Hash = FCrc::TypeCrc32(ULandscapeTextureHash::GetHash(Component->GetHeightmap()), Hash);
 
 		// hash the weightmap textures
 		for (UTexture2D* Weightmap : Component->GetWeightmapTextures())
 		{
 			check(Weightmap->Source.IsValid());
-			Hash = FCrc::TypeCrc32(Weightmap->Source.GetId(), Hash);
+			Hash = FCrc::TypeCrc32(ULandscapeTextureHash::GetHash(Weightmap), Hash);
 		}
 	#endif // WITH_EDITOR
 
diff --git a/Engine/Source/Runtime/Landscape/Private/LandscapeTextureHash.cpp b/Engine/Source/Runtime/Landscape/Private/LandscapeTextureHash.cpp
new file mode 100644
--- /dev/null
+++ b/Engine/Source/Runtime/Landscape/Private/LandscapeTextureHash.cpp
@@ -0,0 +1,382 @@
+// Copyright Epic Games, Inc. All Rights Reserved.
+
+/*=============================================================================
+LandscapeTextureHash.cpp: Track a Custom Hash on each landscape texture.
+This hash tries to be insensitive to changes that are less than the thresholds,
+and also ignores normal data channels on the heightmaps.
+=============================================================================*/
+
+#include "LandscapeTextureHash.h"
+#include "LandscapePrivate.h"
+#include "Engine/Texture2D.h"
+
+#include UE_INLINE_GENERATED_CPP_BY_NAME(LandscapeTextureHash)
+
+#if WITH_EDITORONLY_DATA
+
+extern TAutoConsoleVariable<int32> CVarLandscapeDirtyHeightmapHeightThreshold;
+extern TAutoConsoleVariable<int32> CVarLandscapeDirtyWeightmapThreshold;
+extern int32 GPatchEdges;
+extern int32 GPatchStreamingMipEdges;
+
+namespace UE::Landscape::Private
+{
+	// produces a Valid (non-zero) FGuid from a 64 bit hash
+	static inline FGuid Hash64ToGUID(uint64 Hash)
+	{
+		uint32 LowBits = Hash & 0xffffffff;
+		uint32 HighBits = Hash >> 32;
+		return FGuid(HighBits, HighBits + LowBits + 0xbb4824dc, HighBits ^ LowBits, LowBits);
+	}
+
+	// converts a GUID into a 64 bit hash (inverse of Hash64ToGUID)
+	static inline uint64 GUIDToHash64(const FGuid& GUID)
+	{
+		return (static_cast<uint64>(GUID.A) << 32) + GUID.D;
+	}
+};
+
+void ULandscapeTextureHash::SetInitialStateOnPostLoad(UTexture2D* LandscapeTexture, ELandscapeTextureUsage TextureUsage, ELandscapeTextureType TextureType)
+{
+	ULandscapeTextureHash* TextureHash = LandscapeTexture->GetAssetUserData<ULandscapeTextureHash>();
+	if (TextureHash == nullptr)
+	{
+		// if there is no texture hash recorded, create a new one (using default SourceID as the hash) and record it in the recent serialized hashes
+		TextureHash = NewObject<ULandscapeTextureHash>(LandscapeTexture);
+		LandscapeTexture->AddAssetUserData(TextureHash);
+
+		FGuid LandscapeTextureSourceId = LandscapeTexture->Source.GetId();
+		TextureHash->TextureHashGUID = LandscapeTextureSourceId;
+		TextureHash->LastSourceID = LandscapeTextureSourceId;
+		TextureHash->TextureType = TextureType;
+		TextureHash->TextureUsage = TextureUsage;
+		TextureHash->RecentlySerializedHashes.Add(LandscapeTextureSourceId, LandscapeTextureSourceId);
+	}
+	else
+	{
+		// TextureHash serialization should ensure this
+		check(TextureHash->RecentlySerializedHashes.Contains(TextureHash->LastSourceID));
+		check(TextureHash->RecentlySerializedHashes[TextureHash->LastSourceID] == TextureHash->TextureHashGUID);
+	}
+}
+
+void ULandscapeTextureHash::CheckHashIsUpToDate(UTexture2D* LandscapeTexture)
+{
+	ULandscapeTextureHash* TextureHash = LandscapeTexture->GetAssetUserData<ULandscapeTextureHash>();
+	check(TextureHash != nullptr);
+	check(LandscapeTexture->Source.GetId() == TextureHash->LastSourceID);
+}
+
+FGuid ULandscapeTextureHash::CalculateTextureHashGUID(UTexture2D* LandscapeTexture, ELandscapeTextureType TextureType)
+{
+	uint64 Hash64 = CalculateTextureHash64(LandscapeTexture, TextureType);
+	return UE::Landscape::Private::Hash64ToGUID(Hash64);
+}
+
+uint64 ULandscapeTextureHash::CalculateTextureHash64(UTexture2D* LandscapeTexture, ELandscapeTextureType TextureType)
+{
+	const int32 MipIndex = 0;
+	const int64 MipSizeInBytes = LandscapeTexture->Source.CalcMipSize(MipIndex);
+	const int64 MipSizeInPixels = MipSizeInBytes / 4;
+	const FColor* MipData = (const FColor*)LandscapeTexture->Source.LockMipReadOnly(MipIndex);
+	uint64 Hash = CalculateTextureHash64(MipData, MipSizeInPixels, TextureType);
+	LandscapeTexture->Source.UnlockMip(MipIndex);
+	return Hash;
+}
+
+uint64 ULandscapeTextureHash::CalculateTextureHash64(const FColor* Mip0Data, int32 PixelCount, ELandscapeTextureType TextureType)
+{
+	TRACE_CPUPROFILER_EVENT_SCOPE(ULandscapeTextureHash::CalculateTextureHash);
+
+	uint64 NewHash = 0;
+
+	switch (TextureType)
+	{
+		case ELandscapeTextureType::Unknown:
+			check(TextureType != ELandscapeTextureType::Unknown);
+			break;
+		case ELandscapeTextureType::Heightmap:
+			{
+				uint32 CRCHashG = 0;
+				uint32 CRCHashR = 1;
+
+				for (int32 i = 0; i < PixelCount; i++)
+				{
+					// the height is the Red and Green channels (ignore the normal data in the other channels)
+					CRCHashG = FCrc::TypeCrc32(Mip0Data->G, CRCHashG);
+					CRCHashR = FCrc::TypeCrc32(Mip0Data->R, CRCHashR);
+					Mip0Data++;
+				}
+
+				NewHash = (((uint64)CRCHashR) << 32) + CRCHashG;
+			}
+			break;
+		case ELandscapeTextureType::Weightmap:
+			{
+				NewHash = CityHash64(reinterpret_cast<const char*>(Mip0Data), PixelCount * sizeof(FColor));
+			}
+			break;
+		default:
+			checkf(false, TEXT("TextureType is invalid: %d"), TextureType);
+			break;
+	}
+
+	return NewHash;
+}
+
+bool ULandscapeTextureHash::DoesTextureDataChangeExceedThreshold(
+	const FColor* Mip0Data, const FColor* OldMip0Data, int32 PixelCount, ELandscapeTextureType TextureType, uint64 OldHash, uint64 NewHash, TOptional<uint8>& OutChangedWeightmapChannelsMasks)
+{
+	TRACE_CPUPROFILER_EVENT_SCOPE(ULandscapeTextureHash::DoesTextureDataChangeExceedThreshold);
+
+	bool bExceedsThreshold = false;
+
+	switch (TextureType)
+	{
+		case ELandscapeTextureType::Unknown:
+			check(TextureType != ELandscapeTextureType::Unknown);
+			break;
+		case ELandscapeTextureType::Heightmap:
+			{
+				int32 DirtyHeightmapHeightThreshold = CVarLandscapeDirtyHeightmapHeightThreshold.GetValueOnGameThread();
+				if (DirtyHeightmapHeightThreshold <= 0)
+				{
+					// at at threshold of zero, any change at all will exceed
+					bExceedsThreshold = (OldHash != NewHash);
+					break;
+				}
+
+				for (int32 i = 0; i < PixelCount; i++)
+				{
+					const FColor& OldColor = *OldMip0Data;
+					const FColor& NewColor = *Mip0Data;
+
+					if (OldColor != NewColor)
+					{
+						uint16 OldHeight = ((static_cast<uint16>(OldColor.R) << 8) | static_cast<uint16>(OldColor.G));
+						uint16 NewHeight = ((static_cast<uint16>(NewColor.R) << 8) | static_cast<uint16>(NewColor.G));
+						if (uint16 Diff = (NewHeight > OldHeight) ? (NewHeight - OldHeight) : (OldHeight - NewHeight); Diff > DirtyHeightmapHeightThreshold)
+						{
+							bExceedsThreshold = true;
+							break;
+						}
+					}
+
+					Mip0Data++;
+					OldMip0Data++;
+				}
+			}
+			break;
+		case ELandscapeTextureType::Weightmap:
+			{
+				int32 DirtyWeightmapThreshold = CVarLandscapeDirtyWeightmapThreshold.GetValueOnGameThread();
+				if ((DirtyWeightmapThreshold <= 0) && !OutChangedWeightmapChannelsMasks.IsSet())
+				{
+					// at at threshold of zero, any change at all will exceed
+					bExceedsThreshold = (OldHash != NewHash);
+					break;
+				}
+
+				for (int32 Index = 0; Index < PixelCount; ++Index)
+				{
+					const FColor& OldColor = *OldMip0Data;
+					const FColor& NewColor = *Mip0Data;
+
+					if (OldColor != NewColor)
+					{
+						auto DiffChannel = [DirtyWeightmapThreshold](uint8 InOldValue, uint8 InNewValue) -> bool
+							{
+								uint8 Diff = (InNewValue > InOldValue) ? (InNewValue - InOldValue) : (InOldValue - InNewValue);
+								return (Diff > DirtyWeightmapThreshold);
+							};
+
+						uint8 DiffMask =
+							((uint8)(DiffChannel(OldColor.R, NewColor.R) ? 1 : 0) << 0)
+							| ((uint8)(DiffChannel(OldColor.G, NewColor.G) ? 1 : 0) << 1)
+							| ((uint8)(DiffChannel(OldColor.B, NewColor.B) ? 1 : 0) << 2)
+							| ((uint8)(DiffChannel(OldColor.A, NewColor.A) ? 1 : 0) << 3);
+
+						if (DiffMask != 0)
+						{
+							bExceedsThreshold = true;
+							if (OutChangedWeightmapChannelsMasks.IsSet())
+							{
+								*OutChangedWeightmapChannelsMasks |= DiffMask;
+							}
+							else
+							{
+								// no need to report which channel has been changed, early out :
+								break;
+							}
+						}
+					}
+
+					Mip0Data++;
+					OldMip0Data++;
+				}
+			}
+			break;
+		default:
+			checkf(false, TEXT("TextureType is invalid: %d"), TextureType);
+			break;
+	}
+
+	return bExceedsThreshold;
+}
+
+void ULandscapeTextureHash::SetHash64(UTexture2D* LandscapeTexture, uint64 NewHash64, ELandscapeTextureUsage TextureUsage, ELandscapeTextureType TextureType)
+{
+	ULandscapeTextureHash* TextureHash = LandscapeTexture->GetAssetUserData<ULandscapeTextureHash>();
+	if (TextureHash == nullptr)
+	{
+		// create a new one (with LandscapeTexture as outer)
+		TextureHash = NewObject<ULandscapeTextureHash>(LandscapeTexture);
+		LandscapeTexture->AddAssetUserData(TextureHash);
+	}
+	else
+	{
+		// pre-existing -- should have the same type
+		check(TextureHash->TextureType == TextureType);
+		check(TextureHash->TextureUsage == TextureUsage);
+	}
+	FGuid LandscapeTextureSourceId = LandscapeTexture->Source.GetId();
+
+	// cached hashes take precedence -- this ensure that if the texture is brought back to a recently serialized state, it will have exactly the same hash that it was serialized with
+	FGuid NewHashGUID;
+	if (FGuid* CachedHash = TextureHash->RecentlySerializedHashes.Find(LandscapeTextureSourceId))
+	{
+		NewHashGUID = *CachedHash;
+	}
+	else
+	{
+		NewHashGUID = UE::Landscape::Private::Hash64ToGUID(NewHash64);
+	}
+	TextureHash->TextureHashGUID = NewHashGUID;
+	TextureHash->LastSourceID = LandscapeTextureSourceId;
+	TextureHash->TextureType = TextureType;
+	TextureHash->TextureUsage = TextureUsage;
+}
+
+void ULandscapeTextureHash::UpdateHash(UTexture2D* LandscapeTexture, ELandscapeTextureUsage TextureUsage, ELandscapeTextureType TextureType, bool bForceUpdate)
+{
+	TRACE_CPUPROFILER_EVENT_SCOPE(ULandscapeTextureHash::UpdateHash);
+
+	FGuid LandscapeTextureSourceId = LandscapeTexture->Source.GetId();
+
+	bool bNewlyCreated = false;
+	ULandscapeTextureHash* TextureHash = LandscapeTexture->GetAssetUserData<ULandscapeTextureHash>();
+	if (TextureHash == nullptr)
+	{
+		// create a new one (with LandscapeTexture as outer)
+		TextureHash = NewObject<ULandscapeTextureHash>(LandscapeTexture);
+		LandscapeTexture->AddAssetUserData(TextureHash);
+		bNewlyCreated = true;
+	}
+	else
+	{
+		if (!bForceUpdate && (LandscapeTextureSourceId == TextureHash->LastSourceID))
+		{
+			// no need to update, it's the same
+			check(TextureHash->TextureUsage == TextureUsage || TextureUsage == ELandscapeTextureUsage::Unknown);
+			check(TextureHash->TextureType == TextureType || TextureType == ELandscapeTextureType::Unknown);
+			return;
+		}
+	}
+
+	if (TextureUsage == ELandscapeTextureUsage::Unknown)
+	{
+		TextureUsage = TextureHash->TextureUsage;
+	}
+	if (TextureType == ELandscapeTextureType::Unknown)
+	{
+		TextureType = TextureHash->TextureType;
+	}
+
+	FGuid NewHash;
+	if ((TextureUsage != ELandscapeTextureUsage::FinalData) ||
+		(TextureType == ELandscapeTextureType::Unknown))
+	{
+		// non-final data and/or unknown types don't need to use a hash, as we just use the SourceID directly
+		NewHash = LandscapeTextureSourceId;
+	}
+	else
+	{
+		// if this SourceId is familiar, use the corresponding hash
+		if (FGuid* CachedHash = TextureHash->RecentlySerializedHashes.Find(LandscapeTextureSourceId))
+		{
+			NewHash = *CachedHash;
+		}
+		else
+		{
+			// otherwise compute a new one
+			NewHash = CalculateTextureHashGUID(LandscapeTexture, TextureType);
+		}
+	}
+
+	TextureHash->TextureHashGUID = NewHash;
+	TextureHash->LastSourceID = LandscapeTextureSourceId;
+	TextureHash->TextureType = TextureType;
+	TextureHash->TextureUsage = TextureUsage;
+
+	if (bNewlyCreated)
+	{
+		TextureHash->RecentlySerializedHashes.Add(LandscapeTextureSourceId, NewHash);
+	}
+}
+
+FGuid ULandscapeTextureHash::GetHash(UTexture2D* LandscapeTexture)
+{
+	// if no texture hash exists, or it's not a final layer-merged texture, just use the source ID
+	ULandscapeTextureHash* TextureHash = LandscapeTexture->GetAssetUserData<ULandscapeTextureHash>();
+	if ((TextureHash == nullptr) || (TextureHash->TextureUsage != ELandscapeTextureUsage::FinalData) || (TextureHash->TextureType == ELandscapeTextureType::Unknown))
+	{
+		// fallback to using the Source ID (matches old behavior)
+		check(LandscapeTexture->Source.IsValid());
+		return LandscapeTexture->Source.GetId();
+	}
+	if (LandscapeTexture->Source.GetId() != TextureHash->LastSourceID)
+	{
+		// NOTE: this can happen in WP mode when a final data texture is transacted for undo/redo.
+		// It can also happen when we are using non-WP mode when directly modifying the final texture source on the CPU (as we don't rehash on all CPU modifications)
+
+		// in either case we can just force update the hash to get an ok hash to use.
+		// This won't take change thresholds into account, but that's as good as we can do in these cases.
+		UpdateHash(LandscapeTexture, TextureHash->TextureUsage, TextureHash->TextureType, /*bForceUpdate=*/ true);
+	}
+	return TextureHash->TextureHashGUID;
+}
+
+void ULandscapeTextureHash::Serialize(FArchive& Ar)
+{
+	bool bUpdateRecentlySerializedHashes = Ar.IsPersistent() && !Ar.IsObjectReferenceCollector() && !HasAnyFlags(RF_ArchetypeObject | RF_ClassDefaultObject | RF_DefaultSubObject);
+	if (bUpdateRecentlySerializedHashes && Ar.IsSaving())
+	{
+		UTexture2D* ParentTexture = Cast<UTexture2D>(GetOuter());
+		if (ParentTexture && (ParentTexture->Source.GetId() != LastSourceID))
+		{
+			// The stored hash is out of date (it was modified without explicitly updating the hash)  Update it before serializing!
+			// This won't take change thresholds into account, but that's as good as we can do in these cases.
+			UpdateHash(ParentTexture, TextureUsage, TextureType, /*bForceUpdate=*/ true);
+		}
+
+		// as we're about to save this, make it an official recent value
+		// (this guarantees that if we get back to the current state, we will get the same texture hash, despite the threshold-change shenanigans that might go on in the meantime)
+		if (!RecentlySerializedHashes.Contains(LastSourceID))
+		{
+			RecentlySerializedHashes.Add(LastSourceID, TextureHashGUID);
+		}
+	}
+
+	Super::Serialize(Ar);
+
+	if (bUpdateRecentlySerializedHashes && Ar.IsLoading())
+	{
+		// as we just loaded this, make it an official recent value
+		// (this guarantees that if we get back to the current state, we will get the same texture hash, despite the threshold-change shenanigans that might go on in the meantime)
+		if (!RecentlySerializedHashes.Contains(LastSourceID))
+		{
+			RecentlySerializedHashes.Add(LastSourceID, TextureHashGUID);
+		}
+	}
+}
+#endif // WITH_EDITORONLY_DATA
diff --git a/Engine/Source/Runtime/Landscape/Private/LandscapeUtils.cpp b/Engine/Source/Runtime/Landscape/Private/LandscapeUtils.cpp
--- a/Engine/Source/Runtime/Landscape/Private/LandscapeUtils.cpp
+++ b/Engine/Source/Runtime/Landscape/Private/LandscapeUtils.cpp
@@ -110,12 +110,16 @@ uint32 GetTypeHash(const FTextureCopyRequest& InKey)
 {
 	uint32 Hash = ::GetTypeHash(InKey.Source);
 	uint32 HashSlice = ::GetTypeHash(InKey.DestinationSlice);
-	return HashCombine(Hash, ::GetTypeHash(InKey.Destination));
+	return HashCombine(Hash, ::GetTypeHash(InKey.Destination)) ^ (HashSlice << 4) ^ ((uint32) InKey.TextureUsage << 2) ^ ((uint32) InKey.TextureType);
 }
 
 bool operator==(const FTextureCopyRequest& InEntryA, const FTextureCopyRequest& InEntryB)
 {
-	return (InEntryA.Source == InEntryB.Source) && (InEntryA.Destination == InEntryB.Destination) && (InEntryA.DestinationSlice == InEntryB.DestinationSlice);
+	return (InEntryA.Source == InEntryB.Source) &&
+			(InEntryA.Destination == InEntryB.Destination) &&
+			(InEntryA.DestinationSlice == InEntryB.DestinationSlice) &&
+			(InEntryA.TextureUsage == InEntryB.TextureUsage) &&
+			(InEntryA.TextureType == InEntryB.TextureType);
 }
 
 bool FBatchTextureCopy::AddWeightmapCopy(UTexture* InDestination, int8 InDestinationSlice, int8 InDestinationChannel, const ULandscapeComponent* InComponent, ULandscapeLayerInfoObject* InLayerInfo)
@@ -163,6 +167,8 @@ struct FDestinationDataMipNumber
 {
 	TArray<uint8*> DestinationDataPtr;
 	int32 MipNumber = 0;
+	ELandscapeTextureUsage TextureUsage = ELandscapeTextureUsage::Unknown;
+	ELandscapeTextureType TextureType = ELandscapeTextureType::Unknown;
 };
 
 bool FBatchTextureCopy::ProcessTextureCopies()
@@ -184,6 +190,8 @@ bool FBatchTextureCopy::ProcessTextureCopies()
 
 		FDestinationDataMipNumber& DestinationData = Destinations.Add(CopyRequest.Key.Destination);
 		DestinationData.MipNumber = CopyRequest.Key.Destination->Source.GetNumMips();
+		DestinationData.TextureUsage = CopyRequest.Key.TextureUsage;
+		DestinationData.TextureType = CopyRequest.Key.TextureType;
 	}
 
 	// Decompress (if needed) and get the source textures ready for access
@@ -259,6 +267,8 @@ bool FBatchTextureCopy::ProcessTextureCopies()
 		{
 			Destination.Key->Source.UnlockMip(MipLevel);
 		}
+
+		ULandscapeTextureHash::UpdateHash(CastChecked<UTexture2D>(Destination.Key), Destination.Value.TextureUsage, Destination.Value.TextureType);
 	}
 
 	return true;
diff --git a/Engine/Source/Runtime/Landscape/Public/LandscapeEdit.h b/Engine/Source/Runtime/Landscape/Public/LandscapeEdit.h
--- a/Engine/Source/Runtime/Landscape/Public/LandscapeEdit.h
+++ b/Engine/Source/Runtime/Landscape/Public/LandscapeEdit.h
@@ -21,6 +21,7 @@ LandscapeEdit.h: Classes for the editor to access to Landscape data
 
 #if WITH_EDITOR
 #include "Containers/ArrayView.h"
+#include "LandscapeTextureHash.h"
 #endif
 
 #include "Landscape.h"
@@ -54,7 +55,7 @@ struct FLandscapeTextureDataInfo
 		bool bFull;
 	};
 
-	FLandscapeTextureDataInfo(UTexture2D* InTexture, bool bShouldDirtyPackage);
+	FLandscapeTextureDataInfo(UTexture2D* InTexture, bool bShouldDirtyPackage, ELandscapeTextureUsage TextureUsage = ELandscapeTextureUsage::Unknown, ELandscapeTextureType TextureType = ELandscapeTextureType::Unknown);
 	virtual ~FLandscapeTextureDataInfo();
 
 	// returns true if we need to block on the render thread before unlocking the mip data
@@ -113,6 +114,8 @@ struct FLandscapeTextureDataInfo
 private:
 	UTexture2D* Texture;
 	TArray<FMipInfo> MipInfo;
+	ELandscapeTextureUsage TextureUsage = ELandscapeTextureUsage::Unknown;
+	ELandscapeTextureType TextureType = ELandscapeTextureType::Unknown;
 };
 
 struct FLandscapeTextureDataInterface
diff --git a/Engine/Source/Runtime/Landscape/Public/LandscapeUtils.h b/Engine/Source/Runtime/Landscape/Public/LandscapeUtils.h
--- a/Engine/Source/Runtime/Landscape/Public/LandscapeUtils.h
+++ b/Engine/Source/Runtime/Landscape/Public/LandscapeUtils.h
@@ -6,6 +6,7 @@
 #include "CoreMinimal.h"
 #include "UObject/NameTypes.h"
 #include "AssetRegistry/AssetData.h"
+#include "LandscapeTextureHash.h"
 
 class ULevel;
 enum EShaderPlatform : uint16;
@@ -37,6 +38,8 @@ struct FTextureCopyRequest
 	UTexture2D* Source = nullptr;
 	UTexture* Destination = nullptr;
 	int8 DestinationSlice = 0;
+	ELandscapeTextureUsage TextureUsage = ELandscapeTextureUsage::Unknown;
+	ELandscapeTextureType TextureType = ELandscapeTextureType::Unknown;
 };
 
 uint32 GetTypeHash(const FTextureCopyRequest& InKey);
