Fix a crash in landscape grass map building, hit a handful of times when shutting down the editor.

diff --git a/Engine/Source/Runtime/Landscape/Private/LandscapeAsyncTextureReadback.cpp b/Engine/Source/Runtime/Landscape/Private/LandscapeAsyncTextureReadback.cpp
--- a/Engine/Source/Runtime/Landscape/Private/LandscapeAsyncTextureReadback.cpp
+++ b/Engine/Source/Runtime/Landscape/Private/LandscapeAsyncTextureReadback.cpp
@@ -4,6 +4,7 @@
 #include "DataDrivenShaderPlatformInfo.h"
 #include "RenderGraphBuilder.h"
 #include "RenderGraphUtils.h"
+#include "LandscapePrivate.h"
 
 void FLandscapeAsyncTextureReadback::StartReadback_RenderThread(FRDGBuilder& GraphBuilder, FRDGTextureRef RDGTexture)
 {
@@ -134,18 +135,39 @@ void FLandscapeAsyncTextureReadback::CancelAndSelfDestruct()
 	ENQUEUE_RENDER_COMMAND(FLandscapeAsyncTextureReadback_CancelAndSelfDestruct)(
 		[Readback](FRHICommandListImmediate& RHICmdList)
 		{
-			check(Readback->bStartedOnRenderThread);
 			check(Readback->bCancel);
-
 			// check no commands were queued AFTER calling CancelAndSelfDestruct()
 			check(Readback->PendingRenderThreadCommands == 1);
 
-			// if not yet finished, force run the finish command (may stall til completion)
-			if (!Readback->bFinishedOnRenderThread)
+			if (Readback->bStartedOnRenderThread)
+			{
+				if (!Readback->bFinishedOnRenderThread)
+				{
+					// not yet finished, force run the finish command (may stall til completion)
+					Readback->FinishReadback_RenderThread();
+				}
+				check(Readback->bFinishedOnRenderThread);
+			}
+			else
 			{
-				Readback->FinishReadback_RenderThread();
+				// it was never started - there shouldn't be any async readback allocated if nothing has started
+				if (!ensure(!Readback->AsyncReadback.IsValid()))
+				{
+					UE_LOG(LogLandscape, Warning, TEXT("In FLandscapeAsyncTextureReadback::CancelAndSelfDestruct(), readback not started, but AsyncReadback structure is unexpectedly allocated (%d, %d, %d, %d, %p, %d, %d, %d).  Attempting to clean it up."),
+						Readback->bStartedOnRenderThread.load(),
+						Readback->bFinishedOnRenderThread.load(),
+						Readback->PendingRenderThreadCommands.load(),
+						Readback->bCancel.load(),
+						Readback->AsyncReadback.Get(),
+						Readback->TextureWidth,
+						Readback->TextureHeight,
+						Readback->ReadbackResults.Num()
+					);
+
+					// but if there is, try to clean it up
+					Readback->AsyncReadback.Reset();
+				}
 			}
-			check(Readback->bFinishedOnRenderThread);
 
 			// self destruct (from render thread)
 			delete Readback;
