Fix generating bad landscape physical material data after opening a map.

diff --git a/Engine/Source/Runtime/Landscape/Private/LandscapeEdit.cpp b/Engine/Source/Runtime/Landscape/Private/LandscapeEdit.cpp
--- a/Engine/Source/Runtime/Landscape/Private/LandscapeEdit.cpp
+++ b/Engine/Source/Runtime/Landscape/Private/LandscapeEdit.cpp
@@ -97,6 +97,7 @@ LandscapeEdit.cpp: Landscape editing
 #include "Algo/Count.h"
 #include "Algo/Transform.h"
 #include "Algo/ForEach.h"
+#include "Algo/AllOf.h"
 #include "Serialization/MemoryWriter.h"
 #include "Engine/Canvas.h"
 #include "Spatial/PointHashGrid3.h"
@@ -2263,8 +2264,22 @@ void ULandscapeComponent::UpdatePhysicalMaterialTasks()
 
 	if (PhysicalMaterialTask.GetHash() != Hash)
 	{
+		// ULandscapeSubsystem can start ticking before the weightmaps are loaded.  They are "streamed in", but are still just the default placeholder texture.
+		// Rendering the physical material in this state will silently produce the wrong result.
+		bool bWeightMapsReady = Algo::AllOf(WeightmapTextures, [](UTexture2D* Weightmap)
+			{
+				// We expect FLandscapeTextureStreamingManager to have already started managing the data lifetime of the weightmaps, via UpdateLayersContent.
+				check(!Weightmap || Weightmap->bForceMiplevelsToBeResident);
+
+				return Weightmap && !Weightmap->IsDefaultTexture() && !Weightmap->HasPendingInitOrStreaming() && Weightmap->IsFullyStreamedIn();
+			});
+
 		TArray<UPhysicalMaterial*> PhysicalMaterials;
-		if (GetRenderPhysicalMaterials(PhysicalMaterials))
+		if (!bWeightMapsReady)
+		{
+			//Do nothing.  Don't initialize the task or set PhysicalMaterialHash.  Wait for a future tick when the weightmaps are ready.
+		}
+		else if (GetRenderPhysicalMaterials(PhysicalMaterials))
 		{
 			bool bSuccess = PhysicalMaterialTask.Init(this, Hash);
 			check(bSuccess && PhysicalMaterialTask.IsValid());
diff --git a/Engine/Source/Runtime/Landscape/Private/LandscapeSubsystem.cpp b/Engine/Source/Runtime/Landscape/Private/LandscapeSubsystem.cpp
--- a/Engine/Source/Runtime/Landscape/Private/LandscapeSubsystem.cpp
+++ b/Engine/Source/Runtime/Landscape/Private/LandscapeSubsystem.cpp
@@ -37,6 +37,7 @@
 #include "Algo/Transform.h"
 #include "Algo/RemoveIf.h"
 #include "Algo/Unique.h"
+#include "Algo/Partition.h"
 #include "Misc/App.h"
 #include "Misc/DateTime.h"
 #include "AssetCompilingManager.h"
@@ -668,14 +669,13 @@ void ULandscapeSubsystem::Tick(float DeltaTime)
 
 	static TArray<ALandscapeProxy*> ActiveProxies;
 	{
-	ActiveProxies.Reset(Proxies.Num());
+		ActiveProxies.Reset(Proxies.Num());
 
 		for (TObjectPtr<ALandscapeProxy> ProxyPtr : Proxies)
-	{
-			ALandscapeProxy* Proxy = ProxyPtr.Get();
 		{
+			ALandscapeProxy* Proxy = ProxyPtr.Get();
 			ActiveProxies.Add(Proxy);
-			
+
 			// Update the proxies proxy
 			{
 				if (!Proxy->GetDisableRuntimeGrassMapGeneration())
@@ -701,7 +701,9 @@ void ULandscapeSubsystem::Tick(float DeltaTime)
 				}
 			}
 		}
-	}
+
+		// Move the parent ALandscapes to the front of the list.  Better to be consistent with TickLayers running on those first before doing other work on the proxies.
+		Algo::Partition(ActiveProxies, [](ALandscapeProxy* Proxy) { return Cast<ALandscape>(Proxy) != nullptr; });
 	}
 
 	bool bGrassMapGenerationDisabled = bAllProxiesRuntimeGrassMapsDisabled;
