Fix for a crash when PIE after modifying a texture.

We shouldn't be asserting this here -- because it is called from the component PostLoad, it is not tied directly to the texture lifetime, so
we can't assume the texture is in an initial state.

The crash happens in PIE because the textures are shared but the components are duplicated and PostLoaded again, so if the texture has been modified in editor, then it will not match a serialized hash.

diff --git a/Engine/Source/Runtime/Landscape/Private/LandscapeTextureHash.cpp b/Engine/Source/Runtime/Landscape/Private/LandscapeTextureHash.cpp
--- a/Engine/Source/Runtime/Landscape/Private/LandscapeTextureHash.cpp
+++ b/Engine/Source/Runtime/Landscape/Private/LandscapeTextureHash.cpp
@@ -52,12 +52,6 @@ void ULandscapeTextureHash::SetInitialStateOnPostLoad(UTexture2D* LandscapeTextu
 		TextureHash->TextureUsage = TextureUsage;
 		TextureHash->RecentlySerializedHashes.Add(LandscapeTextureSourceId, LandscapeTextureSourceId);
 	}
-	else
-	{
-		// TextureHash serialization should ensure this
-		check(TextureHash->RecentlySerializedHashes.Contains(TextureHash->LastSourceID));
-		check(TextureHash->RecentlySerializedHashes[TextureHash->LastSourceID] == TextureHash->TextureHashGUID);
-	}
 }
 
 void ULandscapeTextureHash::CheckHashIsUpToDate(UTexture2D* LandscapeTexture)
