Fix for a crash when using save-as on a level containing a landscape needing grass maps built

The issue is that when it calls TemporarilyDisableStreaming() it rebuilds the texture, which sets a "PendingInit" flag on the texture as it is not yet initialized on the RHI render thread.
This causes "IsFullyStreamedIn" to return false (until the render thread actually inits the resource).  However, we are actually ok to continue with rendering, as the texture will have full mips once initialized, which will happen before any renders.  So I removed the PendingInit check from calculation after TemporarilyDisableStreaming().

diff --git a/Engine/Source/Runtime/Landscape/Private/LandscapeTextureStreamingManager.cpp b/Engine/Source/Runtime/Landscape/Private/LandscapeTextureStreamingManager.cpp
--- a/Engine/Source/Runtime/Landscape/Private/LandscapeTextureStreamingManager.cpp
+++ b/Engine/Source/Runtime/Landscape/Private/LandscapeTextureStreamingManager.cpp
@@ -47,7 +47,10 @@ namespace UE::Landscape
 			{
 				Texture2D->TemporarilyDisableStreaming();
 			}
-			bIsFullyStreamedIn = IsTextureFullyStreamedIn(Texture);
+			// It is almost guaranteed there is a pending RHI Init because of the full texture rebuild in TemporarilyDisableStreaming.
+			// This should be fine for rendering purposes -- the texture will complete the RHI init before we render with it, and will have full mips.
+			// (to ensure the texture is init before returning, we would have to stall to drain the render thread here..)
+			bIsFullyStreamedIn = !Texture->IsDefaultTexture() && Texture->IsFullyStreamedIn();
 		}
 #endif // WITH_EDITOR
 		check(bIsFullyStreamedIn);
