Adds a console command to dump the list of weightmap texture allocations for the current feature level

diff --git a/Engine/Source/Runtime/Landscape/Private/Landscape.cpp b/Engine/Source/Runtime/Landscape/Private/Landscape.cpp
--- a/Engine/Source/Runtime/Landscape/Private/Landscape.cpp
+++ b/Engine/Source/Runtime/Landscape/Private/Landscape.cpp
@@ -155,11 +155,11 @@ static void PrintNumLandscapeShadows()
 			NumShadowCasters++;
 		}
 	}
-	UE_LOG(LogConsoleResponse, Display, TEXT("%d/%d landscape components cast shadows"), NumShadowCasters, NumComponents);
+	UE_LOG(LogLandscape, Log, TEXT("%d/%d landscape components cast shadows"), NumShadowCasters, NumComponents);
 }
 
 FAutoConsoleCommand CmdPrintNumLandscapeShadows(
-	TEXT("ls.PrintNumLandscapeShadows"),
+	TEXT("landscape.PrintNumLandscapeShadows"),
 	TEXT("Prints the number of landscape components that cast shadows."),
 	FConsoleCommandDelegate::CreateStatic(PrintNumLandscapeShadows)
 	);
diff --git a/Engine/Source/Runtime/Landscape/Private/LandscapeSubsystem.cpp b/Engine/Source/Runtime/Landscape/Private/LandscapeSubsystem.cpp
--- a/Engine/Source/Runtime/Landscape/Private/LandscapeSubsystem.cpp
+++ b/Engine/Source/Runtime/Landscape/Private/LandscapeSubsystem.cpp
@@ -80,6 +80,92 @@ DECLARE_CYCLE_STAT(TEXT("LandscapeSubsystem Tick"), STAT_LandscapeSubsystemTick,
 
 namespace UE::Landscape
 {
+	static void DumpLandscapeWeightmapAllocations(const TArray<FString>& Args)
+	{
+		bool bDumpDetails = false;
+
+		for (const FString& Arg : Args)
+		{
+			if (FParse::Param(*Arg, TEXT("details")))
+			{
+				bDumpDetails = true;
+			}
+		}
+
+		// Command supports editor and cooked games, so get worlds from Engine instead of Editor
+		for (const FWorldContext& WorldContext : GEngine->GetWorldContexts())
+		{
+			if (UWorld* World = WorldContext.World())
+			{
+				UE_LOG(LogLandscape, Log, TEXT("In World: %s"), *World->GetFullName());
+
+				if (ULandscapeSubsystem* LandscapeSubsystem = World->GetSubsystem<ULandscapeSubsystem>())
+				{
+					LandscapeSubsystem->ForEachLandscapeInfo([&](ULandscapeInfo* LandscapeInfo)
+					{
+						TMap<FName, int32> TargetLayerToNumComponents;
+
+						const FString& LandscapeName = LandscapeInfo->GetLandscapeProxy()->GetActorNameOrLabel();
+						UE_LOG(LogLandscape, Log, TEXT("- Landscape: %s"), *LandscapeName);
+
+						LandscapeInfo->ForEachLandscapeProxy([&TargetLayerToNumComponents, &bDumpDetails, &LandscapeName](ALandscapeProxy* LandscapeProxy)
+						{
+							for (ULandscapeComponent* Component : LandscapeProxy->LandscapeComponents)
+							{
+								TSet<FName> ComponentLayerNameSet;
+								for (const FWeightmapLayerAllocationInfo& AllocInfo : Component->GetWeightmapLayerAllocations())
+								{
+									if (AllocInfo.IsAllocated())
+									{
+										FName LayerName = AllocInfo.LayerInfo->LayerName;
+										TargetLayerToNumComponents.FindOrAdd(LayerName)++;
+										ComponentLayerNameSet.Add(LayerName);
+									}
+								}
+
+								// Detailed print option
+								if (bDumpDetails)
+								{
+									const FString& ComponentName = Component->GetName();
+									const FString& ProxyName = LandscapeProxy->GetActorNameOrLabel();
+									UE_LOG(LogLandscape, Log, TEXT("	Proxy: %s, Component: %s [%s]"), *ProxyName, *ComponentName, *Component->GetSectionBase().ToString());
+									UE_LOG(LogLandscape, Log, TEXT("	- Target Layers: %d"), ComponentLayerNameSet.Num());
+
+									ComponentLayerNameSet.Sort(FNameLexicalLess());
+
+									for (const FName& LayerName : ComponentLayerNameSet)
+									{
+										UE_LOG(LogLandscape, Log, TEXT("	- %s"), *LayerName.ToString());
+									}
+								}
+							}
+							return true;
+						});
+
+						// Print the total layers and their names
+						UE_LOG(LogLandscape, Log, TEXT("Number of unique Target Layers: %d"), TargetLayerToNumComponents.Num());
+
+						TargetLayerToNumComponents.KeySort(FNameLexicalLess());
+
+						for (const TPair<FName, int32>& Pair : TargetLayerToNumComponents)
+						{
+							UE_LOG(LogLandscape, Log, TEXT("- %s, used in %d landscape components"), *Pair.Key.ToString(), Pair.Value);
+						}
+
+						return true;
+					});
+				}
+			}
+		}
+	}
+
+	static FAutoConsoleCommand CmdDumpLandscapeWeightmapAllocations(
+		TEXT("landscape.DumpTargetLayerAllocations"),
+		TEXT("[optional: -details] - Dumps a report of target layers allocated for every landscape. \n"
+			"-details shows a detailed report of allocated target layers for each individual landscape component. \n"),
+		FConsoleCommandWithArgsDelegate::CreateStatic(&DumpLandscapeWeightmapAllocations)
+		);
+
 #if WITH_EDITOR
 	bool HasModifiedLandscapes()
 	{
