Removed SkipTracking from external texture resources used in landscape batched merge. The graph builders simply need to tell the final state of those textures to prevent the automatic transitioning to SRVMask. This also allows avoids tagging the passes as NeverCull, since the external textures being written to are now tracked by the graph builder.

diff --git a/Engine/Source/Runtime/Landscape/Private/LandscapeEditLayerRenderer.cpp b/Engine/Source/Runtime/Landscape/Private/LandscapeEditLayerRenderer.cpp
--- a/Engine/Source/Runtime/Landscape/Private/LandscapeEditLayerRenderer.cpp
+++ b/Engine/Source/Runtime/Landscape/Private/LandscapeEditLayerRenderer.cpp
@@ -661,7 +661,7 @@ void FMergeRenderContext::RenderValidityRenderTargets(const FMergeRenderBatch& I
 		{
 			FRDGBufferRef RectBuffer = CreateUploadBuffer(GraphBuilder, TEXT("MarkValidityRects"), MakeArrayView(TextureAndRects.Rects));
 			FRDGBufferSRVRef RectBufferSRV = GraphBuilder.CreateSRV(FRDGBufferSRVDesc(RectBuffer, PF_R32G32B32A32_UINT));
-			FRDGTextureRef OutputTexture = GraphBuilder.RegisterExternalTexture(CreateRenderTarget(TextureAndRects.TextureResource->GetTexture2DRHI(), TEXT("ValidityMask")), ERDGTextureFlags::SkipTracking);
+			FRDGTextureRef OutputTexture = GraphBuilder.RegisterExternalTexture(CreateRenderTarget(TextureAndRects.TextureResource->GetTexture2DRHI(), TEXT("ValidityMask")));
 
 			FMarkValidityPSParameters* PassParameters = GraphBuilder.AllocParameters<FMarkValidityPSParameters>();
 			PassParameters->RenderTargets[0] = FRenderTargetBinding(OutputTexture, ERenderTargetLoadAction::ELoad);
@@ -680,10 +680,10 @@ void FMergeRenderContext::RenderValidityRenderTargets(const FMergeRenderBatch& I
 				/*RasterizerState = */ nullptr,
 				/*DepthStencilState = */ nullptr,
 				/*StencilRef = */ 0,
-				/*TextureSize = */ OutputTexture->Desc.Extent,
-				/*RectUVBufferSRV = */ nullptr,
-				/*DownsampleFactor = */ 1,
-				ERDGPassFlags::NeverCull); // Use NeverCull because it renders a texture for which tracking is disabled
+				/*TextureSize = */ OutputTexture->Desc.Extent);
+
+			// We need to specify the final state of the external texture to prevent the graph builder from transitioning it to SRVMask :
+			GraphBuilder.SetTextureAccessFinal(OutputTexture, ERHIAccess::RTV);
 		}
 
 		GraphBuilder.Execute();
@@ -723,8 +723,8 @@ void FMergeRenderContext::RenderExpandedRenderTarget(const FMergeRenderBatch& In
 		FRDGBufferRef RectUVBuffer = CreateUploadBuffer(GraphBuilder, TEXT("ExpandRectsUVs"), MakeArrayView(SourceRects));
 		FRDGBufferSRVRef RectUVBufferSRV = GraphBuilder.CreateSRV(FRDGBufferSRVDesc(RectUVBuffer, PF_R32G32B32A32_UINT));
 
-		FRDGTextureRef OutputTexture = GraphBuilder.RegisterExternalTexture(CreateRenderTarget(OutputResource->GetTexture2DRHI(), TEXT("OutputTexture")), ERDGTextureFlags::SkipTracking);
-		FRDGTextureRef SourceTexture = GraphBuilder.RegisterExternalTexture(CreateRenderTarget(SourceResource->GetTexture2DRHI(), TEXT("SourceTexture")), ERDGTextureFlags::SkipTracking);
+		FRDGTextureRef OutputTexture = GraphBuilder.RegisterExternalTexture(CreateRenderTarget(OutputResource->GetTexture2DRHI(), TEXT("OutputTexture")));
+		FRDGTextureRef SourceTexture = GraphBuilder.RegisterExternalTexture(CreateRenderTarget(SourceResource->GetTexture2DRHI(), TEXT("SourceTexture")));
 
 		// TODO [jonathan.bard] this is just a rhi validation error for unoptimized shaders... once validation is made to not issue those errors, we can remove this
 		// Create a SceneView to please the shader bindings, but it's unused in practice 
@@ -756,9 +756,11 @@ void FMergeRenderContext::RenderExpandedRenderTarget(const FMergeRenderBatch& In
 			/*DepthStencilState = */nullptr,
 			/*StencilRef = */0,
 			/*TextureSize = */SourceTexture->Desc.Extent,
-			RectUVBufferSRV,
-			/*DownsampleFactor = */1,
-			ERDGPassFlags::NeverCull); // Use NeverCull because it renders a texture for which tracking is disabled
+			RectUVBufferSRV);
+
+		// We need to specify the final state of the external textures to prevent the graph build from transitioning them to SRVMask :
+		GraphBuilder.SetTextureAccessFinal(OutputTexture, ERHIAccess::RTV);
+		GraphBuilder.SetTextureAccessFinal(SourceTexture, ERHIAccess::SRVMask);
 
 		GraphBuilder.Execute();
 	});
diff --git a/Engine/Source/Runtime/Landscape/Private/LandscapeEditLayers.cpp b/Engine/Source/Runtime/Landscape/Private/LandscapeEditLayers.cpp
--- a/Engine/Source/Runtime/Landscape/Private/LandscapeEditLayers.cpp
+++ b/Engine/Source/Runtime/Landscape/Private/LandscapeEditLayers.cpp
@@ -1684,12 +1684,7 @@ public:
 			RDG_EVENT_NAME("LandscapeLayers_MergeEditLayerPS (%s)", *InTargetLayerName.ToString()),
 			PixelShader,
 			InParameters,
-			FIntRect(0, 0, InTextureSize.X, InTextureSize.Y), 
-			/*BlendState = */nullptr, 
-			/*RasterizerState = */nullptr, 
-			/*DepthStencilState = */nullptr, 
-			/*StencilRef = */0, 
-			ERDGPassFlags::NeverCull); // Use NeverCull because it renders a texture for which tracking is disabled
+			FIntRect(0, 0, InTextureSize.X, InTextureSize.Y));
 	}
 };
 
@@ -1734,12 +1729,7 @@ public:
 			RDG_EVENT_NAME("LandscapeLayers_GenerateNormalsPS"),
 			PixelShader,
 			InParameters,
-			FIntRect(0, 0, InTextureSize.X, InTextureSize.Y),
-			/*BlendState = */nullptr, 
-			/*RasterizerState = */nullptr, 
-			/*DepthStencilState = */nullptr, 
-			/*StencilRef = */0, 
-			ERDGPassFlags::NeverCull); // Use NeverCull because it renders a texture for which tracking is disabled
+			FIntRect(0, 0, InTextureSize.X, InTextureSize.Y));
 	}
 };
 
@@ -1778,12 +1768,7 @@ public:
 			RDG_EVENT_NAME("LandscapeLayers_GenerateMipsPS"),
 			PixelShader,
 			InParameters,
-			FIntRect(0, 0, InTextureSize.X, InTextureSize.Y),
-			/*BlendState = */nullptr,
-			/*RasterizerState = */nullptr,
-			/*DepthStencilState = */nullptr,
-			/*StencilRef = */0,
-			ERDGPassFlags::NeverCull); // Use NeverCull because it renders a texture for which tracking is disabled
+			FIntRect(0, 0, InTextureSize.X, InTextureSize.Y));
 	}
 };
 
@@ -1829,12 +1814,7 @@ public:
 			RDG_EVENT_NAME("LandscapeLayers_MergeEditLayerPS (%s)", *InTargetLayerName.ToString()),
 			PixelShader,
 			InParameters,
-			FIntRect(0, 0, InTextureSize.X, InTextureSize.Y), 
-			/*BlendState = */nullptr,
-			/*RasterizerState = */nullptr,
-			/*DepthStencilState = */nullptr,
-			/*StencilRef = */0,
-			ERDGPassFlags::NeverCull); // Use NeverCull because it renders a texture for which tracking is disabled
+			FIntRect(0, 0, InTextureSize.X, InTextureSize.Y));
 	}
 };
 
@@ -1875,12 +1855,7 @@ public:
 			RDG_EVENT_NAME("LandscapeLayers_PerformLegacyWeightBlendingPS (%s)", *InTargetLayerName.ToString()),
 			PixelShader,
 			InParameters,
-			FIntRect(0, 0, InTextureSize.X, InTextureSize.Y), 
-			/*BlendState = */nullptr,
-			/*RasterizerState = */nullptr,
-			/*DepthStencilState = */nullptr,
-			/*StencilRef = */0,
-			ERDGPassFlags::NeverCull); // Use NeverCull because it renders a texture for which tracking is disabled
+			FIntRect(0, 0, InTextureSize.X, InTextureSize.Y));
 	}
 };
 
@@ -1923,12 +1898,7 @@ public:
 			RDG_EVENT_NAME("LandscapeLayers_PackWeightmapPS"),
 			PixelShader,
 			InParameters,
-			InTextureRect,
-			/*BlendState = */nullptr,
-			/*RasterizerState = */nullptr,
-			/*DepthStencilState = */nullptr,
-			/*StencilRef = */0,
-			ERDGPassFlags::NeverCull); // Use NeverCull because it renders a texture for which tracking is disabled
+			InTextureRect);
 	}
 };
 
@@ -1967,12 +1937,7 @@ public:
 			RDG_EVENT_NAME("LandscapeLayers_GenerateMipsPS"),
 			PixelShader,
 			InParameters,
-			FIntRect(0, 0, InTextureSize.X, InTextureSize.Y), 
-			/*BlendState = */nullptr,
-			/*RasterizerState = */nullptr,
-			/*DepthStencilState = */nullptr,
-			/*StencilRef = */0,
-			ERDGPassFlags::NeverCull); // Use NeverCull because it renders a texture for which tracking is disabled
+			FIntRect(0, 0, InTextureSize.X, InTextureSize.Y));
 	}
 };
 
@@ -5738,8 +5703,10 @@ void ULandscapeEditLayerPersistent::RenderLayer(ILandscapeEditLayerRenderer::FRe
 			TArray<FUintVector2> CurrentQuadInfos;
 			TArray<FUintVector4> CurrentSourceRects, CurrentDestinationRects;
 
+			FRDGTextureRef OutputTexture = GraphBuilder.RegisterExternalTexture(CreateRenderTarget(OutputResource->GetTextureRHI(), TEXT("OutputTexture")));
+
 			// When we've reached the limit of what we can render in one pass, flush all and add a render pass : 
-			auto FlushCopyQuads = [&GraphBuilder, View, OutputResource, &CurrentDestinationArrayIndex, &CurrentSourceTextureResources, &CurrentSourceRects, &CurrentDestinationRects,
+			auto FlushCopyQuads = [&GraphBuilder, View, OutputTexture, &CurrentDestinationArrayIndex, &CurrentSourceTextureResources, &CurrentSourceRects, &CurrentDestinationRects,
 				&CurrentQuadInfos, &CommonSourceTextureSize, BlackDummySRVRef, ShaderMap, PixelShader]()
 			{
 				if (CurrentSourceTextureResources.IsEmpty())
@@ -5759,8 +5726,6 @@ void ULandscapeEditLayerPersistent::RenderLayer(ILandscapeEditLayerRenderer::FRe
 
 				FRDGBufferRef QuadInfosBuffer = CreateUploadBuffer(GraphBuilder, TEXT("QuadInfos"), TConstArrayView<FUintVector2>(MakeArrayView(CurrentQuadInfos)));
 				FRDGBufferSRVRef QuadInfosBufferSRV = GraphBuilder.CreateSRV(FRDGBufferSRVDesc(QuadInfosBuffer, PF_R32G32_UINT));
-
-				FRDGTextureRef OutputTexture = GraphBuilder.RegisterExternalTexture(CreateRenderTarget(OutputResource->GetTextureRHI(), TEXT("OutputTexture")), ERDGTextureFlags::SkipTracking);
 				
 				// We may target an array slice directly : 
 				int32 ArrayIndex = -1;
@@ -5798,9 +5763,7 @@ void ULandscapeEditLayerPersistent::RenderLayer(ILandscapeEditLayerRenderer::FRe
 					/*DepthStencilState = */nullptr,
 					/*StencilRef = */0,
 					/*TextureSize = */FIntPoint(FMath::FloorToInt32(CommonSourceTextureSize.X), FMath::FloorToInt32(CommonSourceTextureSize.Y)),
-					RectUVBufferSRV,
-					/*DownsampleFactor = */1,
-					ERDGPassFlags::NeverCull); // Use NeverCull because it renders a texture for which tracking is disabled
+					RectUVBufferSRV);
 
 				// Reset all for the next pass that comes along : 
 				CurrentDestinationArrayIndex = INDEX_NONE;
@@ -5853,6 +5816,9 @@ void ULandscapeEditLayerPersistent::RenderLayer(ILandscapeEditLayerRenderer::FRe
 			// Flush the remaining quads if any : 
 			FlushCopyQuads();
 
+			// We need to specify the final state of the external texture to prevent the graph builder from transitioning it to SRVMask :
+			GraphBuilder.SetTextureAccessFinal(OutputTexture, ERHIAccess::RTV);
+
 			GraphBuilder.Execute();
 		});
 	}
@@ -5881,7 +5847,7 @@ void ULandscapeEditLayerPersistent::RenderLayer(ILandscapeEditLayerRenderer::FRe
 				SCOPED_DRAW_EVENTF(InRHICmdList, LandscapeLayers, TEXT("MergeEditLayer -> %s"), OutputResourceName);
 				FRDGBuilder GraphBuilder(InRHICmdList, RDG_EVENT_NAME("HeightmapsMergeEditLayer"));
 
-				FRDGTextureRef OutputTextureRef = GraphBuilder.RegisterExternalTexture(CreateRenderTarget(OutputResource->TextureRHI, TEXT("OutputTexture")), ERDGTextureFlags::SkipTracking);
+				FRDGTextureRef OutputTextureRef = GraphBuilder.RegisterExternalTexture(CreateRenderTarget(OutputResource->TextureRHI, TEXT("OutputTexture")));
 				FRDGTextureRef CurrentEditLayerTextureRef = GraphBuilder.RegisterExternalTexture(CreateRenderTarget(CurrentEditLayerResource->TextureRHI, TEXT("CurrentEditLayerTexture")));
 				FRDGTextureRef PreviousEditLayersTextureRef = GraphBuilder.RegisterExternalTexture(CreateRenderTarget(PreviousEditLayersResource->TextureRHI, TEXT("PreviousEditLayersTexture")));
 
@@ -5893,6 +5859,10 @@ void ULandscapeEditLayerPersistent::RenderLayer(ILandscapeEditLayerRenderer::FRe
 				PSParams->InPreviousEditLayersHeightmap = GraphBuilder.CreateSRV(FRDGTextureSRVDesc::Create(PreviousEditLayersTextureRef));
 
 				FLandscapeEditLayersHeightmapsMergeEditLayerPS::MergeEditLayerPS(GraphBuilder, PSParams, TargetLayerName, EffectiveTextureSize);
+
+				// We need to specify the final state of the external texture to prevent the graph builder from transitioning it to SRVMask :
+				GraphBuilder.SetTextureAccessFinal(OutputTextureRef, ERHIAccess::RTV);
+
 				GraphBuilder.Execute();
 			});
 		}
@@ -5944,7 +5914,7 @@ void ULandscapeEditLayerPersistent::RenderLayer(ILandscapeEditLayerRenderer::FRe
 				SCOPED_DRAW_EVENTF(InRHICmdList, LandscapeLayers, TEXT("MergeEditLayer -> %s"), OutputResourceName);
 				FRDGBuilder GraphBuilder(InRHICmdList, RDG_EVENT_NAME("WeightmapsMergeEditLayer"));
 
-				FRDGTextureRef OutputTextureRef = GraphBuilder.RegisterExternalTexture(CreateRenderTarget(OutputResource->TextureRHI, TEXT("OutputTexture")), ERDGTextureFlags::SkipTracking);
+				FRDGTextureRef OutputTextureRef = GraphBuilder.RegisterExternalTexture(CreateRenderTarget(OutputResource->TextureRHI, TEXT("OutputTexture")));
 				FRDGTextureRef CurrentEditLayerTextureRef = GraphBuilder.RegisterExternalTexture(CreateRenderTarget(CurrentEditLayerResource->TextureRHI, TEXT("CurrentEditLayerTexture")));
 				FRDGTextureSRVRef CurrentEditLayerTextureSRVRef = GraphBuilder.CreateSRV(FRDGTextureSRVDesc::Create(CurrentEditLayerTextureRef));
 				FRDGTextureRef PreviousEditLayersTextureRef = GraphBuilder.RegisterExternalTexture(CreateRenderTarget(PreviousEditLayersResource->TextureRHI, TEXT("PreviousEditLayersTexture")));
@@ -5963,6 +5933,10 @@ void ULandscapeEditLayerPersistent::RenderLayer(ILandscapeEditLayerRenderer::FRe
 
 					FLandscapeEditLayersWeightmapsMergeEditLayerPS::MergeEditLayerPS(GraphBuilder, PSParams, MergePaintLayersParams[TargetLayerIndex].PaintLayerName, EffectiveTextureSize);
 				}
+
+				// We need to specify the final state of the external texture to prevent the graph builder from transitioning it to SRVMask :
+				GraphBuilder.SetTextureAccessFinal(OutputTextureRef, ERHIAccess::RTV);
+
 				GraphBuilder.Execute();
 			});
 		}
@@ -6465,7 +6439,7 @@ void ULandscapeHeightmapNormalsEditLayerRenderer::RenderLayer(FRenderParams& InR
 		SCOPED_DRAW_EVENTF(InRHICmdList, LandscapeLayers, TEXT("GenerateNormals -> %s"), OutputResourceName);
 		FRDGBuilder GraphBuilder(InRHICmdList, RDG_EVENT_NAME("HeightmapsGenerateNormals"));
 
-		FRDGTextureRef OutputTextureRef = GraphBuilder.RegisterExternalTexture(CreateRenderTarget(OutputResource->TextureRHI, TEXT("OutputTexture")), ERDGTextureFlags::SkipTracking);
+		FRDGTextureRef OutputTextureRef = GraphBuilder.RegisterExternalTexture(CreateRenderTarget(OutputResource->TextureRHI, TEXT("OutputTexture")));
 		FRDGTextureRef SourceTextureRef = GraphBuilder.RegisterExternalTexture(CreateRenderTarget(SourceResource->TextureRHI, TEXT("SourceTexture")));
 		FRDGTextureRef ValidityTextureRef = GraphBuilder.RegisterExternalTexture(CreateRenderTarget(ValidityResource->TextureRHI, TEXT("ValidityTexture")));
 
@@ -6479,6 +6453,10 @@ void ULandscapeHeightmapNormalsEditLayerRenderer::RenderLayer(FRenderParams& InR
 		PSParams->InSourceHeightmap = GraphBuilder.CreateSRV(FRDGTextureSRVDesc::Create(SourceTextureRef));
 		PSParams->InValidityTexture = GraphBuilder.CreateSRV(FRDGTextureSRVDesc::Create(ValidityTextureRef));
 		FLandscapeEditLayersHeightmapsGenerateNormalsPS::GenerateNormalsPS(GraphBuilder, PSParams, EffectiveTextureSize);
+
+		// We need to specify the final state of the external texture to prevent the graph builder from transitioning it to SRVMask :
+		GraphBuilder.SetTextureAccessFinal(OutputTextureRef, ERHIAccess::RTV);
+
 		GraphBuilder.Execute();
 	});
 }
@@ -6523,7 +6501,7 @@ void ULandscapeWeightmapWeightBlendedLayersRenderer::RenderLayer(FRenderParams&
 		SCOPED_DRAW_EVENTF(InRHICmdList, LandscapeLayers, TEXT("LegacyWeightBlending -> %s"), OutputResourceName);
 		FRDGBuilder GraphBuilder(InRHICmdList, RDG_EVENT_NAME("WeightmapsPerformLegacyWeightBlending"));
 
-		FRDGTextureRef OutputTextureRef = GraphBuilder.RegisterExternalTexture(CreateRenderTarget(OutputResource->TextureRHI, TEXT("OutputTexture")), ERDGTextureFlags::SkipTracking);
+		FRDGTextureRef OutputTextureRef = GraphBuilder.RegisterExternalTexture(CreateRenderTarget(OutputResource->TextureRHI, TEXT("OutputTexture"))/*, ERDGTextureFlags::SkipTracking*/);
 		FRDGTextureRef CurrentEditLayerTextureRef = GraphBuilder.RegisterExternalTexture(CreateRenderTarget(CurrentEditLayerResource->TextureRHI, TEXT("CurrentEditLayerTexture")));
 		FRDGTextureSRVRef CurrentEditLayerTextureSRVRef = GraphBuilder.CreateSRV(FRDGTextureSRVDesc::Create(CurrentEditLayerTextureRef));
 
@@ -6542,6 +6520,10 @@ void ULandscapeWeightmapWeightBlendedLayersRenderer::RenderLayer(FRenderParams&
 
 			FLandscapeEditLayersWeightmapsPerformLegacyWeightBlendingPS::PerformLegacyWeightBlendingPS(GraphBuilder, PSParams, TargetLayerNames[TargetLayerIndex], EffectiveTextureSize);
 		}
+
+		// We need to specify the final state of the external texture to prevent the graph builder from transitioning it to SRVMask :
+		GraphBuilder.SetTextureAccessFinal(OutputTextureRef, ERHIAccess::RTV);
+
 		GraphBuilder.Execute();
 	});
 }
@@ -7641,7 +7623,7 @@ int32 ALandscape::PerformLayersHeightmapsBatchedMerge(const FUpdateLayersContent
 					{
 						FRDGBuilder GraphBuilder(InRHICmdList, RDG_EVENT_NAME("HeightmapsGenerateMips"));
 
-						FRDGTextureRef OutputTextureRef = GraphBuilder.RegisterExternalTexture(CreateRenderTarget(OutputResource->TextureRHI, TEXT("OutputTexture")), ERDGTextureFlags::SkipTracking);
+						FRDGTextureRef OutputTextureRef = GraphBuilder.RegisterExternalTexture(CreateRenderTarget(OutputResource->TextureRHI, TEXT("OutputTexture")));
 						FRDGTextureRef SourceTextureRef = GraphBuilder.RegisterExternalTexture(CreateRenderTarget(SourceResource->TextureRHI, TEXT("SourceTexture")));
 
 						FLandscapeEditLayersHeightmapsGenerateMipsPS::FParameters* PSParams = GraphBuilder.AllocParameters<FLandscapeEditLayersHeightmapsGenerateMipsPS::FParameters>();
@@ -7649,6 +7631,10 @@ int32 ALandscape::PerformLayersHeightmapsBatchedMerge(const FUpdateLayersContent
 						PSParams->InCurrentMipSubsectionSize = FUintVector2(CurrentMipSubsectionSize.X, CurrentMipSubsectionSize.Y);
 						PSParams->InSourceHeightmap = GraphBuilder.CreateSRV(FRDGTextureSRVDesc::Create(SourceTextureRef));
 						FLandscapeEditLayersHeightmapsGenerateMipsPS::GenerateMipsPS(GraphBuilder, PSParams, CurrentMipResolution);
+
+						// We need to specify the final state of the external texture to prevent the graph builder from transitioning it to SRVMask :
+						GraphBuilder.SetTextureAccessFinal(OutputTextureRef, ERHIAccess::RTV);
+
 						GraphBuilder.Execute();
 					});
 				}
