Landscape debug material refactor, persistent Landscape view mode when toggling Landscape mode

diff --git a/Engine/Source/Editor/LandscapeEditor/Private/LandscapeEdMode.cpp b/Engine/Source/Editor/LandscapeEditor/Private/LandscapeEdMode.cpp
--- a/Engine/Source/Editor/LandscapeEditor/Private/LandscapeEdMode.cpp
+++ b/Engine/Source/Editor/LandscapeEditor/Private/LandscapeEdMode.cpp
@@ -75,22 +75,6 @@ void FLandscapeTool::SetEditRenderType()
 	GLandscapeEditRenderMode = ELandscapeEditRenderMode::SelectRegion | (GLandscapeEditRenderMode & ELandscapeEditRenderMode::BitMaskForMask);
 }
 
-namespace UE::Landscape::Editor::Tool
-{
-	UMaterialInstance* CreateMaterialInstance(UMaterialInterface* BaseMaterial)
-	{
-		UObject* Outer = GetTransientPackage();
-		// Use the base material's name as the base of our MIC to help debug: 
-		FString MICName(FString::Format(TEXT("LandscapeMaterialInstanceConstant_{0}"), { *BaseMaterial->GetName() }));
-		ULandscapeMaterialInstanceConstant* MaterialInstance = NewObject<ULandscapeMaterialInstanceConstant>(Outer, MakeUniqueObjectName(Outer, ULandscapeMaterialInstanceConstant::StaticClass(), FName(MICName)));
-		MaterialInstance->bEditorToolUsage = true;
-		MaterialInstance->SetParentEditorOnly(BaseMaterial);
-		MaterialInstance->PostEditChange();
-		return MaterialInstance;
-	}
-}
-// namespace UE::Landscape::Editor::Tool
-
 namespace UE::Landscape::Editor::Private
 {
 	/** Indicates the user is currently moving the landscape gizmo object by dragging the mouse. */
@@ -119,6 +103,7 @@ FEdModeLandscape::FEdModeLandscape()
 	, GizmoBrush(nullptr)
 	, CurrentToolIndex(INDEX_NONE)
 	, CurrentBrushSetIndex(0)
+	, CurrentLandscapeViewMode(ELandscapeViewMode::Normal)
 	, NewLandscapePreviewMode(ENewLandscapePreviewMode::None)
 	, ImportExportMode(EImportExportMode::Import)
 	, CurrentGizmoActor(nullptr)
@@ -132,16 +117,7 @@ FEdModeLandscape::FEdModeLandscape()
 	, bNeedsUpdateLayerUsageInformation(false)
 	, bUpdatingLandscapeInfo(false)
 {
-	using namespace UE::Landscape::Editor::Tool;
-
-	GLayerDebugColorMaterial = CreateMaterialInstance(LoadObject<UMaterial>(nullptr, TEXT("/Engine/EditorLandscapeResources/LayerVisMaterial.LayerVisMaterial")));
-	GSelectionColorMaterial = CreateMaterialInstance(LoadObject<UMaterialInstanceConstant>(nullptr, TEXT("/Engine/EditorLandscapeResources/SelectBrushMaterial_Selected.SelectBrushMaterial_Selected")));
-	GSelectionRegionMaterial = CreateMaterialInstance(LoadObject<UMaterialInstanceConstant>(nullptr, TEXT("/Engine/EditorLandscapeResources/SelectBrushMaterial_SelectedRegion.SelectBrushMaterial_SelectedRegion")));
-	GMaskRegionMaterial = CreateMaterialInstance(LoadObject<UMaterialInstanceConstant>(nullptr, TEXT("/Engine/EditorLandscapeResources/MaskBrushMaterial_MaskedRegion.MaskBrushMaterial_MaskedRegion")));
-	GColorMaskRegionMaterial = CreateMaterialInstance(LoadObject<UMaterialInstanceConstant>(nullptr, TEXT("/Engine/EditorLandscapeResources/ColorMaskBrushMaterial_MaskedRegion.ColorMaskBrushMaterial_MaskedRegion")));
-	GLandscapeDirtyMaterial = CreateMaterialInstance(LoadObject<UMaterial>(nullptr, TEXT("/Engine/EditorLandscapeResources/LandscapeDirtyMaterial.LandscapeDirtyMaterial")));
-	GLandscapeBlackTexture = LoadObject<UTexture2D>(nullptr, TEXT("/Engine/EngineResources/Black.Black"));
-	GLandscapeLayerUsageMaterial = CreateMaterialInstance(LoadObject<UMaterial>(nullptr, TEXT("/Engine/EditorLandscapeResources/LandscapeLayerUsageMaterial.LandscapeLayerUsageMaterial")));
+	using namespace UE::Landscape::Editor;
 	
 	// Initialize modes
 	UpdateToolModes();
@@ -206,14 +182,6 @@ FEdModeLandscape::~FEdModeLandscape()
 
 	// Clean up Debug Materials
 	FlushRenderingCommands();
-	GLayerDebugColorMaterial = nullptr;
-	GSelectionColorMaterial = nullptr;
-	GSelectionRegionMaterial = nullptr;
-	GMaskRegionMaterial = nullptr;
-	GColorMaskRegionMaterial = nullptr;
-	GLandscapeBlackTexture = nullptr;
-	GLandscapeLayerUsageMaterial = nullptr;
-	GLandscapeDirtyMaterial = nullptr;
 
 	InteractorPainting = nullptr;
 }
@@ -225,15 +193,6 @@ void FEdModeLandscape::AddReferencedObjects(FReferenceCollector& Collector)
 	FEdMode::AddReferencedObjects(Collector);
 
 	Collector.AddReferencedObject(UISettings);
-
-	Collector.AddReferencedObject(GLayerDebugColorMaterial);
-	Collector.AddReferencedObject(GSelectionColorMaterial);
-	Collector.AddReferencedObject(GSelectionRegionMaterial);
-	Collector.AddReferencedObject(GMaskRegionMaterial);
-	Collector.AddReferencedObject(GColorMaskRegionMaterial);
-	Collector.AddReferencedObject(GLandscapeBlackTexture);
-	Collector.AddReferencedObject(GLandscapeLayerUsageMaterial);
-	Collector.AddReferencedObject(GLandscapeDirtyMaterial);
 }
 
 void FEdModeLandscape::UpdateToolModes()
@@ -375,6 +334,11 @@ ULandscapeLayerInfoObject* FEdModeLandscape::GetSelectedLandscapeLayerInfo() con
 	return CurrentToolTarget.LayerInfo.Get();
 }
 
+bool FEdModeLandscape::IsLandscapeViewModeExclusiveToEditorMode(ELandscapeViewMode::Type ViewMode)
+{
+	return ViewMode == ELandscapeViewMode::DebugLayer || ViewMode == ELandscapeViewMode::LayerContribution;
+}
+
 void FEdModeLandscape::SetLandscapeInfo(ULandscapeInfo* InLandscapeInfo)
 {
 	check(!InLandscapeInfo || InLandscapeInfo->SupportsLandscapeEditing());
@@ -664,6 +628,12 @@ void FEdModeLandscape::Enter()
 	// Reset mouse tracking info : 
 	GIsGizmoDragging = false;
 	GIsAdjustingBrush  = false;
+
+	// Restore view mode if it was auto changed last time the landscape editor was closed
+	if (IsLandscapeViewModeExclusiveToEditorMode(CurrentLandscapeViewMode))
+	{
+		GLandscapeViewMode = CurrentLandscapeViewMode;
+	}
 }
 
 /** FEdMode: Called when the mode is exited */
@@ -745,7 +715,15 @@ void FEdModeLandscape::Exit()
 
 	// Save UI settings to config file
 	UISettings->Save();
-	GLandscapeViewMode = ELandscapeViewMode::Normal;
+
+	// Save the last used view mode 
+	CurrentLandscapeViewMode = static_cast<ELandscapeViewMode::Type>(GLandscapeViewMode);
+
+	// Reset global view if current mode is only supported with editor open
+	if (IsLandscapeViewModeExclusiveToEditorMode(CurrentLandscapeViewMode))
+	{
+		GLandscapeViewMode = ELandscapeViewMode::Normal;
+	}
 	GLandscapeEditRenderMode = ELandscapeEditRenderMode::None;
 	GLandscapeEditModeActive = false;
 
diff --git a/Engine/Source/Editor/LandscapeEditor/Private/LandscapeEdMode.h b/Engine/Source/Editor/LandscapeEditor/Private/LandscapeEdMode.h
--- a/Engine/Source/Editor/LandscapeEditor/Private/LandscapeEdMode.h
+++ b/Engine/Source/Editor/LandscapeEditor/Private/LandscapeEdMode.h
@@ -42,6 +42,10 @@ class ULandscapeLayerInfoObject;
 class FLandscapeToolSplines;
 class UViewportInteractor;
 struct FViewportActionKeyInput;
+namespace ELandscapeViewMode 
+{
+	enum Type : int32;
+}
 
 struct FHeightmapToolTarget;
 template<typename TargetType> class FLandscapeToolCopyPaste;
@@ -293,6 +297,8 @@ public:
 	int32 CurrentToolIndex;
 	// UI setting for additional UI Tools
 	int32 CurrentBrushSetIndex;
+	// Persistent View Mode when toggling landscape editor
+	ELandscapeViewMode::Type CurrentLandscapeViewMode;
 
 	ENewLandscapePreviewMode NewLandscapePreviewMode;
 	EImportExportMode ImportExportMode;
@@ -626,6 +632,8 @@ public:
 		return GetEditingState() == ELandscapeEditingState::Enabled;
 	}
 
+	static bool IsLandscapeViewModeExclusiveToEditorMode(ELandscapeViewMode::Type ViewMode);
+
 	void SetLandscapeInfo(ULandscapeInfo* InLandscapeInfo);
 
 	/** Returns the sum of all landscape actors resolution. */
diff --git a/Engine/Source/Editor/LandscapeEditor/Private/LandscapeEdModeBrushes.cpp b/Engine/Source/Editor/LandscapeEditor/Private/LandscapeEdModeBrushes.cpp
--- a/Engine/Source/Editor/LandscapeEditor/Private/LandscapeEdModeBrushes.cpp
+++ b/Engine/Source/Editor/LandscapeEditor/Private/LandscapeEdModeBrushes.cpp
@@ -12,7 +12,7 @@
 #include "LandscapeGizmoActor.h"
 #include "LandscapeEdMode.h"
 #include "LandscapeEditorObject.h"
-
+#include "LandscapeUtils.h"
 #include "LandscapeRender.h"
 
 #include "LevelUtils.h"
@@ -66,7 +66,7 @@ protected:
 	/** Protected so that only subclasses can create instances of this class. */
 	FLandscapeBrushCircle(FEdModeLandscape* InEdMode, UMaterialInterface* InBrushMaterial)
 		: LastMousePosition(0, 0)
-		, BrushMaterial(UE::Landscape::Editor::Tool::CreateMaterialInstance(InBrushMaterial))
+		, BrushMaterial(UE::Landscape::CreateToolLandscapeMaterialInstanceConstant(InBrushMaterial))
 		, bCanPaint(false)
 		, EdMode(InEdMode)
 	{
@@ -432,7 +432,7 @@ public:
 		, EdMode(InEdMode)
 	{
 		UMaterial* BaseBrushMaterial = LoadObject<UMaterial>(nullptr, TEXT("/Engine/EditorLandscapeResources/SelectBrushMaterial.SelectBrushMaterial"));
-		BrushMaterial = UE::Landscape::Editor::Tool::CreateMaterialInstance(BaseBrushMaterial);
+		BrushMaterial = UE::Landscape::CreateToolLandscapeMaterialInstanceConstant(BaseBrushMaterial);
 	}
 
 	// FGCObject interface
@@ -642,7 +642,7 @@ public:
 		, EdMode(InEdMode)
 	{
 		UMaterialInterface* GizmoMaterial = LoadObject<UMaterialInstanceConstant>(nullptr, TEXT("/Engine/EditorLandscapeResources/MaskBrushMaterial_Gizmo.MaskBrushMaterial_Gizmo"));
-		BrushMaterial = UMaterialInstanceDynamic::Create(UE::Landscape::Editor::Tool::CreateMaterialInstance(GizmoMaterial), nullptr);
+		BrushMaterial = UMaterialInstanceDynamic::Create(UE::Landscape::CreateToolLandscapeMaterialInstanceConstant(GizmoMaterial), nullptr);
 	}
 
 	// FGCObject interface
diff --git a/Engine/Source/Editor/LandscapeEditor/Public/LandscapeToolInterface.h b/Engine/Source/Editor/LandscapeEditor/Public/LandscapeToolInterface.h
--- a/Engine/Source/Editor/LandscapeEditor/Public/LandscapeToolInterface.h
+++ b/Engine/Source/Editor/LandscapeEditor/Public/LandscapeToolInterface.h
@@ -272,9 +272,3 @@ public:
 	int32 PreviousBrushIndex = INDEX_NONE;
 	TArray<FName> ValidBrushes;
 };
-
-namespace UE::Landscape::Editor::Tool
-{
-	UMaterialInstance* CreateMaterialInstance(UMaterialInterface* BaseMaterial);
-}
-// namespace UE::Landscape::Editor::Tool
diff --git a/Engine/Source/Runtime/Landscape/Private/LandscapeEditResourcesSubsystem.cpp b/Engine/Source/Runtime/Landscape/Private/LandscapeEditResourcesSubsystem.cpp
--- a/Engine/Source/Runtime/Landscape/Private/LandscapeEditResourcesSubsystem.cpp
+++ b/Engine/Source/Runtime/Landscape/Private/LandscapeEditResourcesSubsystem.cpp
@@ -8,6 +8,10 @@
 #include "UObject/Package.h"
 #include "RenderGraphUtils.h"
 
+#if WITH_EDITOR
+#include "LandscapeRender.h"
+#include "LandscapeMaterialInstanceConstant.h"
+#endif // WITH_EDITOR
 
 // ----------------------------------------------------------------------------------
 
@@ -392,12 +396,48 @@ ULandscapeEditResourcesSubsystem::ULandscapeEditResourcesSubsystem()
 void ULandscapeEditResourcesSubsystem::Initialize(FSubsystemCollectionBase& Collection)
 {
 	Super::Initialize(Collection);
+	
+#if WITH_EDITOR  
+	LayerDebugColorMaterial = UE::Landscape::CreateToolLandscapeMaterialInstanceConstant (LoadObject<UMaterial>(nullptr, TEXT("/Engine/EditorLandscapeResources/LayerVisMaterial.LayerVisMaterial")));
+	SelectionColorMaterial = UE::Landscape::CreateToolLandscapeMaterialInstanceConstant (LoadObject<UMaterialInstanceConstant>(nullptr, TEXT("/Engine/EditorLandscapeResources/SelectBrushMaterial_Selected.SelectBrushMaterial_Selected")));
+	SelectionRegionMaterial = UE::Landscape::CreateToolLandscapeMaterialInstanceConstant (LoadObject<UMaterialInstanceConstant>(nullptr, TEXT("/Engine/EditorLandscapeResources/SelectBrushMaterial_SelectedRegion.SelectBrushMaterial_SelectedRegion")));
+	MaskRegionMaterial = UE::Landscape::CreateToolLandscapeMaterialInstanceConstant (LoadObject<UMaterialInstanceConstant>(nullptr, TEXT("/Engine/EditorLandscapeResources/MaskBrushMaterial_MaskedRegion.MaskBrushMaterial_MaskedRegion")));
+	ColorMaskRegionMaterial = UE::Landscape::CreateToolLandscapeMaterialInstanceConstant (LoadObject<UMaterialInstanceConstant>(nullptr, TEXT("/Engine/EditorLandscapeResources/ColorMaskBrushMaterial_MaskedRegion.ColorMaskBrushMaterial_MaskedRegion")));
+	LandscapeDirtyMaterial = UE::Landscape::CreateToolLandscapeMaterialInstanceConstant (LoadObject<UMaterial>(nullptr, TEXT("/Engine/EditorLandscapeResources/LandscapeDirtyMaterial.LandscapeDirtyMaterial")));
+	LandscapeLayerUsageMaterial = UE::Landscape::CreateToolLandscapeMaterialInstanceConstant (LoadObject<UMaterial>(nullptr, TEXT("/Engine/EditorLandscapeResources/LandscapeLayerUsageMaterial.LandscapeLayerUsageMaterial")));
+	LandscapeBlackTexture = LoadObject<UTexture2D>(nullptr, TEXT("/Engine/EngineResources/Black.Black"));
+
+	PRAGMA_DISABLE_DEPRECATION_WARNINGS
+	// Deprecated global variables, use private variables with Getters instead
+	GLayerDebugColorMaterial = LayerDebugColorMaterial;
+	GSelectionColorMaterial = SelectionColorMaterial;
+	GSelectionRegionMaterial = SelectionRegionMaterial;
+	GMaskRegionMaterial = MaskRegionMaterial;
+	GColorMaskRegionMaterial = ColorMaskRegionMaterial;
+	GLandscapeDirtyMaterial = LandscapeDirtyMaterial;
+	GLandscapeLayerUsageMaterial = LandscapeLayerUsageMaterial;
+	GLandscapeBlackTexture = LandscapeBlackTexture;
+	PRAGMA_ENABLE_DEPRECATION_WARNINGS
+#endif // WITH_EDITOR
 }
 
 
 void ULandscapeEditResourcesSubsystem::Deinitialize()
 {
 	Super::Deinitialize();
+
+#if WITH_EDITOR  
+	PRAGMA_DISABLE_DEPRECATION_WARNINGS
+	GLayerDebugColorMaterial = nullptr;
+	GSelectionColorMaterial = nullptr;
+	GSelectionRegionMaterial = nullptr;
+	GMaskRegionMaterial = nullptr;
+	GColorMaskRegionMaterial = nullptr;
+	GLandscapeDirtyMaterial = nullptr;
+	GLandscapeLayerUsageMaterial = nullptr;
+	GLandscapeBlackTexture = nullptr;
+	PRAGMA_ENABLE_DEPRECATION_WARNINGS
+#endif // WITH_EDITOR
 }
 
 ULandscapeScratchRenderTarget* ULandscapeEditResourcesSubsystem::RequestScratchRenderTarget(const UE::Landscape::FScratchRenderTargetParams& InParams)
diff --git a/Engine/Source/Runtime/Landscape/Private/LandscapeRender.cpp b/Engine/Source/Runtime/Landscape/Private/LandscapeRender.cpp
--- a/Engine/Source/Runtime/Landscape/Private/LandscapeRender.cpp
+++ b/Engine/Source/Runtime/Landscape/Private/LandscapeRender.cpp
@@ -57,6 +57,10 @@ LandscapeRender.cpp: New terrain rendering
 #include "LandscapeUtilsPrivate.h"
 #include "SceneRendererInterface.h"
 
+#if WITH_EDITOR
+#include "LandscapeEditResourcesSubsystem.h"
+#endif
+
 using namespace UE::Landscape;
 
 IMPLEMENT_GLOBAL_SHADER_PARAMETER_STRUCT(FLandscapeUniformShaderParameters, "LandscapeParameters");
@@ -408,13 +412,15 @@ void ULandscapeComponent::GetUsedMaterials(TArray<UMaterialInterface*>& OutMater
 #if WITH_EDITOR
 	//if (bGetDebugMaterials) // TODO: This should be tested and enabled
 	{
-		OutMaterials.Add(GLayerDebugColorMaterial);
-		OutMaterials.Add(GSelectionColorMaterial);
-		OutMaterials.Add(GSelectionRegionMaterial);
-		OutMaterials.Add(GMaskRegionMaterial);
-		OutMaterials.Add(GColorMaskRegionMaterial);
-		OutMaterials.Add(GLandscapeLayerUsageMaterial);
-		OutMaterials.Add(GLandscapeDirtyMaterial);
+		ULandscapeEditResourcesSubsystem* LandscapeEditResourcesSubsystem = GEngine->GetEngineSubsystem<ULandscapeEditResourcesSubsystem>();
+		check(LandscapeEditResourcesSubsystem != nullptr);
+		OutMaterials.Add(LandscapeEditResourcesSubsystem->GetLayerDebugColorMaterial());
+		OutMaterials.Add(LandscapeEditResourcesSubsystem->GetSelectionColorMaterial());
+		OutMaterials.Add(LandscapeEditResourcesSubsystem->GetSelectionRegionMaterial());
+		OutMaterials.Add(LandscapeEditResourcesSubsystem->GetMaskRegionMaterial());
+		OutMaterials.Add(LandscapeEditResourcesSubsystem->GetColorMaskRegionMaterial());
+		OutMaterials.Add(LandscapeEditResourcesSubsystem->GetLandscapeLayerUsageMaterial());
+		OutMaterials.Add(LandscapeEditResourcesSubsystem->GetLandscapeDirtyMaterial());
 	}
 #endif
 }
@@ -1344,6 +1350,10 @@ bool FLandscapeVisibilityHelper::OnRemoveFromWorld()
 	return true;
 }
 
+#if WITH_EDITOR
+ULandscapeEditResourcesSubsystem* FLandscapeComponentSceneProxy::LandscapeEditResourcesSubsystem = nullptr;
+#endif // WITH_EDITOR
+
 FLandscapeComponentSceneProxy::FLandscapeComponentSceneProxy(ULandscapeComponent* InComponent)
 	: FPrimitiveSceneProxy(InComponent, NAME_LandscapeResourceNameForDebugging)
 	, FLandscapeSectionInfo(InComponent->GetWorld(), InComponent->GetLandscapeProxy()->GetLandscapeGuid(), InComponent->GetSectionBase() / InComponent->ComponentSizeQuads, InComponent->GetLandscapeProxy()->LODGroupKey, InComponent->GetLandscapeProxy()->ComputeLandscapeKey())
@@ -1691,6 +1701,8 @@ FLandscapeComponentSceneProxy::FLandscapeComponentSceneProxy(ULandscapeComponent
 	}
 
 #if WITH_EDITOR
+	LandscapeEditResourcesSubsystem = GEngine->GetEngineSubsystem<ULandscapeEditResourcesSubsystem>();
+
 	const TArray<FWeightmapLayerAllocationInfo>& ComponentWeightmapLayerAllocations = InComponent->GetWeightmapLayerAllocations();
 	for (const FWeightmapLayerAllocationInfo& Allocation : ComponentWeightmapLayerAllocations)
 	{
@@ -2055,37 +2067,44 @@ FPrimitiveViewRelevance FLandscapeComponentSceneProxy::GetViewRelevance(const FS
 			// Region selection
 			if (EditToolRenderData.SelectedType)
 			{
+				TObjectPtr<UMaterialInterface> SelectionRegionMaterial = LandscapeEditResourcesSubsystem->GetSelectionRegionMaterial();
+				TObjectPtr<UMaterialInterface> SelectionColorMaterial = LandscapeEditResourcesSubsystem->GetSelectionColorMaterial();
+
 				if ((GLandscapeEditRenderMode & ELandscapeEditRenderMode::SelectRegion) && (EditToolRenderData.SelectedType & FLandscapeEditToolRenderData::ST_REGION)
-					&& !(GLandscapeEditRenderMode & ELandscapeEditRenderMode::Mask) && GSelectionRegionMaterial)
+					&& !(GLandscapeEditRenderMode & ELandscapeEditRenderMode::Mask) && SelectionRegionMaterial)
 				{
 					Result.bDynamicRelevance = true;
-					ToolRelevance |= GSelectionRegionMaterial->GetRelevance_Concurrent(FeatureLevel);
+					ToolRelevance |= SelectionRegionMaterial->GetRelevance_Concurrent(FeatureLevel);
 				}
-				if ((GLandscapeEditRenderMode & ELandscapeEditRenderMode::SelectComponent) && (EditToolRenderData.SelectedType & FLandscapeEditToolRenderData::ST_COMPONENT) && GSelectionColorMaterial)
+				if ((GLandscapeEditRenderMode & ELandscapeEditRenderMode::SelectComponent) && (EditToolRenderData.SelectedType & FLandscapeEditToolRenderData::ST_COMPONENT) && SelectionColorMaterial)
 				{
 					Result.bDynamicRelevance = true;
-					ToolRelevance |= GSelectionColorMaterial->GetRelevance_Concurrent(FeatureLevel);
+					ToolRelevance |= SelectionColorMaterial->GetRelevance_Concurrent(FeatureLevel);
 				}
 			}
 
+			TObjectPtr<UMaterialInterface> MaskRegionMaterial = LandscapeEditResourcesSubsystem->GetMaskRegionMaterial();
+			TObjectPtr<UMaterialInterface> ColorMaskRegionMaterial = LandscapeEditResourcesSubsystem->GetColorMaskRegionMaterial();
+			TObjectPtr<UMaterialInterface> LandscapeDirtyMaterial = LandscapeEditResourcesSubsystem->GetLandscapeDirtyMaterial();
+
 			// Mask
-			if ((GLandscapeEditRenderMode & ELandscapeEditRenderMode::Mask) && GMaskRegionMaterial != nullptr &&
+			if ((GLandscapeEditRenderMode & ELandscapeEditRenderMode::Mask) && MaskRegionMaterial != nullptr &&
 				(((EditToolRenderData.SelectedType & FLandscapeEditToolRenderData::ST_REGION)) || (!(GLandscapeEditRenderMode & ELandscapeEditRenderMode::InvertedMask))))
 			{
 				Result.bDynamicRelevance = true;
-				ToolRelevance |= GMaskRegionMaterial->GetRelevance_Concurrent(FeatureLevel);
+				ToolRelevance |= MaskRegionMaterial->GetRelevance_Concurrent(FeatureLevel);
 			}
 
 			if (GLandscapeViewMode == ELandscapeViewMode::LayerContribution)
 			{
 				Result.bDynamicRelevance = true;
-				ToolRelevance |= GColorMaskRegionMaterial->GetRelevance_Concurrent(FeatureLevel);
+				ToolRelevance |= ColorMaskRegionMaterial->GetRelevance_Concurrent(FeatureLevel);
 			}
 
-			if (CVarLandscapeShowDirty.GetValueOnRenderThread() && GLandscapeDirtyMaterial)
+			if (CVarLandscapeShowDirty.GetValueOnRenderThread() && LandscapeDirtyMaterial)
 			{
 				Result.bDynamicRelevance = true;
-				ToolRelevance |= GLandscapeDirtyMaterial->GetRelevance_Concurrent(FeatureLevel);
+				ToolRelevance |= LandscapeDirtyMaterial->GetRelevance_Concurrent(FeatureLevel);
 			}
 
 			ToolRelevance.SetPrimitiveViewRelevance(Result);
@@ -2114,7 +2133,7 @@ FPrimitiveViewRelevance FLandscapeComponentSceneProxy::GetViewRelevance(const FS
 #if WITH_EDITOR
 		(IsSelected() && !GLandscapeEditModeActive) ||
 		(GLandscapeViewMode != ELandscapeViewMode::Normal) ||
-		(CVarLandscapeShowDirty.GetValueOnAnyThread() && GLandscapeDirtyMaterial) ||
+		(CVarLandscapeShowDirty.GetValueOnAnyThread() && LandscapeEditResourcesSubsystem->GetLandscapeDirtyMaterial()) ||
 		(GetViewLodOverride(*View, LandscapeKey) >= 0)
 #else
 		IsSelected() ||
@@ -2643,13 +2662,15 @@ void FLandscapeComponentSceneProxy::GetDynamicMeshElements(const TArray<const FS
 
 			// Render the landscape component
 #if WITH_EDITOR
+			TObjectPtr<UTexture2D> LandscapeBlackTexture = LandscapeEditResourcesSubsystem->GetLandscapeBlackTexture();
+
 			switch (GLandscapeViewMode)
 			{
 			case ELandscapeViewMode::DebugLayer:
 			{
-				if (GLayerDebugColorMaterial)
+				if (TObjectPtr<UMaterialInterface> LayerDebugColorMaterial = LandscapeEditResourcesSubsystem->GetLayerDebugColorMaterial())
 				{
-					auto DebugColorMaterialInstance = new FLandscapeDebugMaterialRenderProxy(GLayerDebugColorMaterial->GetRenderProxy(),
+					auto DebugColorMaterialInstance = new FLandscapeDebugMaterialRenderProxy(LayerDebugColorMaterial->GetRenderProxy(),
 						(EditToolRenderData.DebugChannelR >= 0 ? WeightmapTextures[EditToolRenderData.DebugChannelR / 4] : nullptr),
 						(EditToolRenderData.DebugChannelG >= 0 ? WeightmapTextures[EditToolRenderData.DebugChannelG / 4] : nullptr),
 						(EditToolRenderData.DebugChannelB >= 0 ? WeightmapTextures[EditToolRenderData.DebugChannelB / 4] : nullptr),
@@ -2694,10 +2715,10 @@ void FLandscapeComponentSceneProxy::GetDynamicMeshElements(const TArray<const FS
 
 			case ELandscapeViewMode::LayerUsage:
 			{
-				if (GLandscapeLayerUsageMaterial)
+				if (TObjectPtr<UMaterialInterface> LandscapeLayerUsageMaterial = LandscapeEditResourcesSubsystem->GetLandscapeLayerUsageMaterial())
 				{
 					float Rotation = ((SectionBase.X / ComponentSizeQuads) ^ (SectionBase.Y / ComponentSizeQuads)) & 1 ? 0 : 2.f * PI;
-					auto LayerUsageMaterialInstance = new FLandscapeLayerUsageRenderProxy(GLandscapeLayerUsageMaterial->GetRenderProxy(), ComponentSizeVerts, LayerColors, Rotation);
+					auto LayerUsageMaterialInstance = new FLandscapeLayerUsageRenderProxy(LandscapeLayerUsageMaterial->GetRenderProxy(), ComponentSizeVerts, LayerColors, Rotation);
 					MeshTools.MaterialRenderProxy = LayerUsageMaterialInstance;
 					Collector.RegisterOneFrameMaterialProxy(LayerUsageMaterialInstance);
 					MeshTools.bCanApplyViewModeOverrides = true;
@@ -2774,7 +2795,8 @@ void FLandscapeComponentSceneProxy::GetDynamicMeshElements(const TArray<const FS
 
 				FMeshBatch& MaskMesh = Collector.AllocateMesh();
 				MaskMesh = MeshTools;
-				auto ColorMaskMaterialInstance = new FLandscapeMaskMaterialRenderProxy(GColorMaskRegionMaterial->GetRenderProxy(), EditToolRenderData.LayerContributionTexture ? ToRawPtr(EditToolRenderData.LayerContributionTexture) : ToRawPtr(GLandscapeBlackTexture), true);
+				TObjectPtr<UMaterialInterface> ColorMaskRegionMaterial = LandscapeEditResourcesSubsystem->GetColorMaskRegionMaterial();
+				auto ColorMaskMaterialInstance = new FLandscapeMaskMaterialRenderProxy(ColorMaskRegionMaterial->GetRenderProxy(), EditToolRenderData.LayerContributionTexture ? ToRawPtr(EditToolRenderData.LayerContributionTexture) : ToRawPtr(LandscapeBlackTexture), true);
 				MaskMesh.MaterialRenderProxy = ColorMaskMaterialInstance;
 				Collector.RegisterOneFrameMaterialProxy(ColorMaskMaterialInstance);
 				Collector.AddMesh(ViewIndex, MaskMesh);
@@ -2785,11 +2807,12 @@ void FLandscapeComponentSceneProxy::GetDynamicMeshElements(const TArray<const FS
 			break;
 
 			default:
-
+				UMaterialInterface* LandscapeDirtyMaterial = LandscapeEditResourcesSubsystem->GetLandscapeDirtyMaterial();
 #endif // WITH_EDITOR
 
 #if WITH_EDITOR || !(UE_BUILD_SHIPPING || UE_BUILD_TEST)
 				const bool bInCollisionView = View->Family->EngineShowFlags.CollisionVisibility || View->Family->EngineShowFlags.CollisionPawn;
+
 				if (AllowDebugViewmodes() && bInCollisionView)
 				{
 					const bool bDrawSimpleCollision = View->Family->EngineShowFlags.CollisionPawn && CollisionResponse.GetResponse(ECC_Pawn) != ECR_Ignore;
@@ -2815,7 +2838,7 @@ void FLandscapeComponentSceneProxy::GetDynamicMeshElements(const TArray<const FS
 					}
 				}
 #if WITH_EDITOR
-				else if (CVarLandscapeShowDirty.GetValueOnRenderThread() && GLandscapeDirtyMaterial)
+				else if (CVarLandscapeShowDirty.GetValueOnRenderThread() && LandscapeDirtyMaterial)
 				{
 					Mesh.bCanApplyViewModeOverrides = false;
 					Collector.AddMesh(ViewIndex, Mesh);
@@ -2826,7 +2849,7 @@ void FLandscapeComponentSceneProxy::GetDynamicMeshElements(const TArray<const FS
 					FMeshBatch& MaskMesh = Collector.AllocateMesh();
 					MaskMesh = MeshTools;
 
-					auto DirtyMaterialInstance = new FLandscapeMaskMaterialRenderProxy(GLandscapeDirtyMaterial->GetRenderProxy(), EditToolRenderData.DirtyTexture ? ToRawPtr(EditToolRenderData.DirtyTexture) : ToRawPtr(GLandscapeBlackTexture), true);
+					auto DirtyMaterialInstance = new FLandscapeMaskMaterialRenderProxy(LandscapeDirtyMaterial->GetRenderProxy(), EditToolRenderData.DirtyTexture ? ToRawPtr(EditToolRenderData.DirtyTexture) : ToRawPtr(LandscapeBlackTexture), true);
 					MaskMesh.MaterialRenderProxy = DirtyMaterialInstance;
 					Collector.RegisterOneFrameMaterialProxy(DirtyMaterialInstance);
 					Collector.AddMesh(ViewIndex, MaskMesh);
@@ -2874,9 +2897,10 @@ void FLandscapeComponentSceneProxy::GetDynamicMeshElements(const TArray<const FS
 					if ((GLandscapeEditRenderMode & ELandscapeEditRenderMode::SelectRegion) && (EditToolRenderData.SelectedType & FLandscapeEditToolRenderData::ST_REGION)
 						&& !(GLandscapeEditRenderMode & ELandscapeEditRenderMode::Mask))
 					{
+						TObjectPtr<UMaterialInterface> SelectionRegionMaterial = LandscapeEditResourcesSubsystem->GetSelectionRegionMaterial();
 						FMeshBatch& SelectMesh = Collector.AllocateMesh();
 						SelectMesh = MeshTools;
-						auto SelectMaterialInstance = new FLandscapeSelectMaterialRenderProxy(GSelectionRegionMaterial->GetRenderProxy(), EditToolRenderData.DataTexture ? ToRawPtr(EditToolRenderData.DataTexture) : ToRawPtr(GLandscapeBlackTexture));
+						auto SelectMaterialInstance = new FLandscapeSelectMaterialRenderProxy(SelectionRegionMaterial->GetRenderProxy(), EditToolRenderData.DataTexture ? ToRawPtr(EditToolRenderData.DataTexture) : ToRawPtr(LandscapeBlackTexture));
 						SelectMesh.MaterialRenderProxy = SelectMaterialInstance;
 						Collector.RegisterOneFrameMaterialProxy(SelectMaterialInstance);
 						Collector.AddMesh(ViewIndex, SelectMesh);
@@ -2887,9 +2911,10 @@ void FLandscapeComponentSceneProxy::GetDynamicMeshElements(const TArray<const FS
 
 					if ((GLandscapeEditRenderMode & ELandscapeEditRenderMode::SelectComponent) && (EditToolRenderData.SelectedType & FLandscapeEditToolRenderData::ST_COMPONENT))
 					{
+						TObjectPtr<UMaterialInterface> SelectionColorMaterial = LandscapeEditResourcesSubsystem->GetSelectionColorMaterial();
 						FMeshBatch& SelectMesh = Collector.AllocateMesh();
 						SelectMesh = MeshTools;
-						SelectMesh.MaterialRenderProxy = GSelectionColorMaterial->GetRenderProxy();
+						SelectMesh.MaterialRenderProxy = SelectionColorMaterial->GetRenderProxy();
 						Collector.AddMesh(ViewIndex, SelectMesh);
 						NumPasses++;
 						NumTriangles += SelectMesh.GetNumPrimitives();
@@ -2900,11 +2925,13 @@ void FLandscapeComponentSceneProxy::GetDynamicMeshElements(const TArray<const FS
 				// Mask
 				if ((GLandscapeEditRenderMode & ELandscapeEditRenderMode::SelectRegion) && (GLandscapeEditRenderMode & ELandscapeEditRenderMode::Mask))
 				{
+					TObjectPtr<UMaterialInterface> MaskRegionMaterial = LandscapeEditResourcesSubsystem->GetMaskRegionMaterial();
+					
 					if (EditToolRenderData.SelectedType & FLandscapeEditToolRenderData::ST_REGION)
 					{
 						FMeshBatch& MaskMesh = Collector.AllocateMesh();
 						MaskMesh = MeshTools;
-						auto MaskMaterialInstance = new FLandscapeMaskMaterialRenderProxy(GMaskRegionMaterial->GetRenderProxy(), EditToolRenderData.DataTexture ? ToRawPtr(EditToolRenderData.DataTexture) : ToRawPtr(GLandscapeBlackTexture), !!(GLandscapeEditRenderMode & ELandscapeEditRenderMode::InvertedMask));
+						auto MaskMaterialInstance = new FLandscapeMaskMaterialRenderProxy(MaskRegionMaterial->GetRenderProxy(), EditToolRenderData.DataTexture ? ToRawPtr(EditToolRenderData.DataTexture) : ToRawPtr(LandscapeBlackTexture), !!(GLandscapeEditRenderMode & ELandscapeEditRenderMode::InvertedMask));
 						MaskMesh.MaterialRenderProxy = MaskMaterialInstance;
 						Collector.RegisterOneFrameMaterialProxy(MaskMaterialInstance);
 						Collector.AddMesh(ViewIndex, MaskMesh);
@@ -2916,7 +2943,7 @@ void FLandscapeComponentSceneProxy::GetDynamicMeshElements(const TArray<const FS
 					{
 						FMeshBatch& MaskMesh = Collector.AllocateMesh();
 						MaskMesh = MeshTools;
-						auto MaskMaterialInstance = new FLandscapeMaskMaterialRenderProxy(GMaskRegionMaterial->GetRenderProxy(), GLandscapeBlackTexture, false);
+						auto MaskMaterialInstance = new FLandscapeMaskMaterialRenderProxy(MaskRegionMaterial->GetRenderProxy(), LandscapeBlackTexture, false);
 						MaskMesh.MaterialRenderProxy = MaskMaterialInstance;
 						Collector.RegisterOneFrameMaterialProxy(MaskMaterialInstance);
 						Collector.AddMesh(ViewIndex, MaskMesh);
diff --git a/Engine/Source/Runtime/Landscape/Private/LandscapeUtils.cpp b/Engine/Source/Runtime/Landscape/Private/LandscapeUtils.cpp
--- a/Engine/Source/Runtime/Landscape/Private/LandscapeUtils.cpp
+++ b/Engine/Source/Runtime/Landscape/Private/LandscapeUtils.cpp
@@ -16,6 +16,7 @@
 #include "Engine/Texture2D.h"
 #include "LandscapeComponent.h"
 #include "ObjectTools.h"
+#include "LandscapeMaterialInstanceConstant.h"
 #endif // WITH_EDITOR
 
 #if WITH_EDITOR
@@ -299,6 +300,19 @@ ULandscapeLayerInfoObject* FLayerInfoFinder::Find(const FName& LayerName) const
 	return nullptr;
 }
 
+
+UMaterialInstance* CreateToolLandscapeMaterialInstanceConstant(UMaterialInterface* BaseMaterial)
+{
+	UObject* Outer = GetTransientPackage();
+	// Use the base material's name as the base of our MIC to help debug: 
+	FString MICName(FString::Format(TEXT("LandscapeMaterialInstanceConstant_{0}"), { *BaseMaterial->GetName() }));
+	ULandscapeMaterialInstanceConstant* MaterialInstance = NewObject<ULandscapeMaterialInstanceConstant>(Outer, MakeUniqueObjectName(Outer, ULandscapeMaterialInstanceConstant::StaticClass(), FName(MICName)));
+	MaterialInstance->bEditorToolUsage = true;
+	MaterialInstance->SetParentEditorOnly(BaseMaterial);
+	MaterialInstance->PostEditChange();
+	return MaterialInstance;
+}
+
 #endif // WITH_EDITOR
 
 } // end namespace UE::Landscape
\ No newline at end of file
diff --git a/Engine/Source/Runtime/Landscape/Public/LandscapeEditResourcesSubsystem.h b/Engine/Source/Runtime/Landscape/Public/LandscapeEditResourcesSubsystem.h
--- a/Engine/Source/Runtime/Landscape/Public/LandscapeEditResourcesSubsystem.h
+++ b/Engine/Source/Runtime/Landscape/Public/LandscapeEditResourcesSubsystem.h
@@ -8,6 +8,11 @@
 #include "RHIAccess.h"
 #include "RHITransition.h"
 
+#if WITH_EDITOR
+#include "Materials/Material.h"
+#include "LandscapeUtils.h"
+#endif // WITH_EDITOR
+
 #include "LandscapeEditResourcesSubsystem.generated.h"
 
 struct FResourceSizeEx;
@@ -250,7 +255,45 @@ public:
 	 */
 	void ReleaseScratchRenderTarget(ULandscapeScratchRenderTarget* InScratchRenderTarget);
 
+#if WITH_EDITOR
+	UMaterialInterface* GetLayerDebugColorMaterial() const { return LayerDebugColorMaterial; }
+	UMaterialInterface* GetSelectionColorMaterial() const { return SelectionColorMaterial; }
+	UMaterialInterface* GetSelectionRegionMaterial() const { return SelectionRegionMaterial; }
+	UMaterialInterface* GetMaskRegionMaterial() const { return MaskRegionMaterial; }
+	UMaterialInterface* GetColorMaskRegionMaterial() const { return ColorMaskRegionMaterial; }
+	UTexture2D* GetLandscapeBlackTexture() const { return LandscapeBlackTexture; }
+	UMaterialInterface* GetLandscapeLayerUsageMaterial() const { return LandscapeLayerUsageMaterial; }
+	UMaterialInterface* GetLandscapeDirtyMaterial() const { return LandscapeDirtyMaterial; }
+#endif // WITH_EDITOR
+
 private:
 	UPROPERTY(Transient, NonTransactional)
 	TArray<TObjectPtr<ULandscapeScratchRenderTarget>> ScratchRenderTargets;
+
+#if WITH_EDITORONLY_DATA
+	UPROPERTY(Transient, NonTransactional) 
+	TObjectPtr<UMaterialInterface> LayerDebugColorMaterial;
+
+	UPROPERTY(Transient, NonTransactional) 
+	TObjectPtr<UMaterialInterface> SelectionColorMaterial;
+
+	UPROPERTY(Transient, NonTransactional) 
+	TObjectPtr<UMaterialInterface> SelectionRegionMaterial;
+
+	UPROPERTY(Transient, NonTransactional) 
+	TObjectPtr<UMaterialInterface> MaskRegionMaterial;
+
+	UPROPERTY(Transient, NonTransactional) 
+	TObjectPtr<UMaterialInterface> ColorMaskRegionMaterial;
+
+	UPROPERTY(Transient, NonTransactional) 
+	TObjectPtr<UTexture2D> LandscapeBlackTexture;
+
+	UPROPERTY(Transient, NonTransactional) 
+	TObjectPtr<UMaterialInterface> LandscapeLayerUsageMaterial;
+
+	UPROPERTY(Transient, NonTransactional) 
+	TObjectPtr<UMaterialInterface> LandscapeDirtyMaterial;
+#endif // WITH_EDITORONLY_DATA
+	 
 };
diff --git a/Engine/Source/Runtime/Landscape/Public/LandscapeRender.h b/Engine/Source/Runtime/Landscape/Public/LandscapeRender.h
--- a/Engine/Source/Runtime/Landscape/Public/LandscapeRender.h
+++ b/Engine/Source/Runtime/Landscape/Public/LandscapeRender.h
@@ -38,6 +38,10 @@ LandscapeRender.h: New terrain rendering
 class FLandscapeComponentSceneProxy;
 enum EShaderPlatform : uint16;
 
+#if WITH_EDITOR
+class ULandscapeEditResourcesSubsystem;
+#endif
+
 namespace UE::Renderer::Private { class IShadowInvalidatingInstances; }
 
 #if RHI_RAYTRACING
@@ -47,7 +51,7 @@ struct FLandscapeRayTracingImpl;
 #if WITH_EDITOR
 namespace ELandscapeViewMode
 {
-	enum Type
+	enum Type: int32
 	{
 		Invalid = -1,
 		/** Color only */
@@ -83,13 +87,22 @@ namespace ELandscapeEditRenderMode
 
 LANDSCAPE_API extern bool GLandscapeEditModeActive;
 LANDSCAPE_API extern int32 GLandscapeEditRenderMode;
+
+UE_DEPRECATED(5.6, "This global has been deprecated. Please use ULandscapeEditResourcesSubsystem->GetLayerDebugColorMaterial")
 LANDSCAPE_API extern TObjectPtr<UMaterialInterface> GLayerDebugColorMaterial;
+UE_DEPRECATED(5.6, "This global has been deprecated. Please use ULandscapeEditResourcesSubsystem->GetSelectionColorMaterial")
 LANDSCAPE_API extern TObjectPtr<UMaterialInterface> GSelectionColorMaterial;
+UE_DEPRECATED(5.6, "This global has been deprecated. Please use ULandscapeEditResourcesSubsystem->GetSelectionRegionMaterial")
 LANDSCAPE_API extern TObjectPtr<UMaterialInterface> GSelectionRegionMaterial;
+UE_DEPRECATED(5.6, "This global has been deprecated. Please use ULandscapeEditResourcesSubsystem->GetMaskRegionMaterial")
 LANDSCAPE_API extern TObjectPtr<UMaterialInterface> GMaskRegionMaterial;
+UE_DEPRECATED(5.6, "This global has been deprecated. Please use ULandscapeEditResourcesSubsystem->GetColorMaskRegionMaterial")
 LANDSCAPE_API extern TObjectPtr<UMaterialInterface> GColorMaskRegionMaterial;
+UE_DEPRECATED(5.6, "This global has been deprecated. Please use ULandscapeEditResourcesSubsystem->GetLandscapeBlackTexture")
 LANDSCAPE_API extern TObjectPtr<UTexture2D> GLandscapeBlackTexture;
+UE_DEPRECATED(5.6, "This global has been deprecated. Please use ULandscapeEditResourcesSubsystem->GetLandscapeLayerUsageMaterial")
 LANDSCAPE_API extern TObjectPtr<UMaterialInterface> GLandscapeLayerUsageMaterial;
+UE_DEPRECATED(5.6, "This global has been deprecated. Please use ULandscapeEditResourcesSubsystem->GetLandscapeDirtyMaterial")
 LANDSCAPE_API extern TObjectPtr<UMaterialInterface> GLandscapeDirtyMaterial;
 #endif
 
@@ -815,6 +828,7 @@ protected:
 
 #if WITH_EDITOR
 	TArray<FLinearColor> LayerColors;
+	static ULandscapeEditResourcesSubsystem* LandscapeEditResourcesSubsystem;
 #endif
 	// Heightmap in RG and Normalmap in BA
 	UTexture2D* HeightmapTexture; 
diff --git a/Engine/Source/Runtime/Landscape/Public/LandscapeUtils.h b/Engine/Source/Runtime/Landscape/Public/LandscapeUtils.h
--- a/Engine/Source/Runtime/Landscape/Public/LandscapeUtils.h
+++ b/Engine/Source/Runtime/Landscape/Public/LandscapeUtils.h
@@ -7,6 +7,7 @@
 #include "UObject/NameTypes.h"
 #include "AssetRegistry/AssetData.h"
 #include "LandscapeTextureHash.h"
+#include "Materials/Material.h"
 
 class ULevel;
 enum EShaderPlatform : uint16;
@@ -117,6 +118,13 @@ struct FLayerInfoFinder
 	TArray<FAssetData> LayerInfoAssets;
 };
 
+/**
+ * Returns a pointer to a newly created ULandscapeMaterialInstanceConstant 
+ * @param	BaseMaterial	The base material appended to the debug name and set as parent of MaterialInstance
+ * @return  Pointer to ULandscapeMaterialInstanceConstant
+ */
+LANDSCAPE_API UMaterialInstance* CreateToolLandscapeMaterialInstanceConstant(UMaterialInterface* BaseMaterial);
+
 #endif // WITH_EDITOR
 
 } // end namespace UE::Landscape
