Landscape edit layers batched merge :
* Fixed edit layer BP brushes that would not get disabled by their parent layer's visibility state when using batched merge, unlike global merge
* Fixed weightmap alpha not being properly applied to edit layers
* Optimized renderer / render item analysis
* Fixed some confusion around CanAffect/Affects for BP brushes
* Fixed performance issue with BP brushes by trying to use intermediate (scratch) render targets of constant size from one batch/one run to another (instead of using only the size necessary for a given batch). This avoids calling Initialize repeatedly (which is very inefficient in the case of the water brush, for example, as it flushes all cached render targets)
* Fixed landscape patch manager not supporting the individual "enable heightmap/weightmap/visibility" toggles

Misc:
* Added debug names to renderer states, in order to help debugging
* Removed unused and not-publicly exposed Affects function on FLandscapeLayerBrush
* Target layers are now handled as bit arrays almost all the way through when using batched merge. The FRenderMergeContext now inherits from FMergeContext, which is setup at the very beginning of the pipe (i.e. before preparing the merge render context).
* FMergeContext is immutable and contains all the target layer names/layer info objects and is passed along to all renderer states/render items, so that they can use it to declare their target layer usage, render groups, etc. as bit arrays, rather than TSet<FName>
* This makes masking / intersecting the render states trivial, which is where a good chunk of the the batched merge preparation time was spent
* Thanks to this, we now skip heightmap-only renderers and render items in weightmap merges (and vice versA), which also reduces preparation time
* Note that RenderLayer and the "on batch rendered" callback relate to a given batch, which is only concerned by a subset of the target layers and allocate a 2D array render target whose slice count is the number of layers needed for that batch (in other words, the slice index doesn't correspond to a target layer index in the merge context). To make that more plain, those function still deal with FNames for target layers, instead of a bit array of the layers in FMergeContext. If it proves too confusing, that could easily be changed

diff --git a/Engine/Plugins/Experimental/LandscapePatch/Source/LandscapePatch/Private/LandscapeCircleHeightPatch.cpp b/Engine/Plugins/Experimental/LandscapePatch/Source/LandscapePatch/Private/LandscapeCircleHeightPatch.cpp
--- a/Engine/Plugins/Experimental/LandscapePatch/Source/LandscapePatch/Private/LandscapeCircleHeightPatch.cpp
+++ b/Engine/Plugins/Experimental/LandscapePatch/Source/LandscapePatch/Private/LandscapeCircleHeightPatch.cpp
@@ -207,10 +207,10 @@ void ULandscapeCircleHeightPatch::ApplyCirclePatch(bool bIsVisibilityLayer,
 }
 
 #if WITH_EDITOR
-void ULandscapeCircleHeightPatch::GetRendererStateInfo(const ULandscapeInfo* InLandscapeInfo,
+void ULandscapeCircleHeightPatch::GetRendererStateInfo(const UE::Landscape::EditLayers::FMergeContext* InMergeContext,
 	FEditLayerTargetTypeState& OutSupported,
 	FEditLayerTargetTypeState& OutEnabled,
-	TArray<TSet<FName>>& OutRenderGroups) const
+	TArray<TBitArray<>>& OutRenderGroups) const
 {
 	OutSupported.AddTargetType(bEditVisibility ? ELandscapeToolTargetType::Visibility
 		: ELandscapeToolTargetType::Heightmap);
@@ -226,7 +226,7 @@ FString ULandscapeCircleHeightPatch::GetEditLayerRendererDebugName() const
 	return FString::Printf(TEXT("%s:%s"), *GetOwner()->GetActorNameOrLabel(), *GetName());
 }
 
-TArray<UE::Landscape::EditLayers::FEditLayerRenderItem> ULandscapeCircleHeightPatch::GetRenderItems(const ULandscapeInfo* InLandscapeInfo) const
+TArray<UE::Landscape::EditLayers::FEditLayerRenderItem> ULandscapeCircleHeightPatch::GetRenderItems(const UE::Landscape::EditLayers::FMergeContext* InMergeContext) const
 {
 	using namespace UE::Landscape::EditLayers;
 
@@ -237,8 +237,7 @@ TArray<UE::Landscape::EditLayers::FEditLayerRenderItem> ULandscapeCircleHeightPa
 	// Figure out the extents of the patch. It will be radius + falloff + an adjustment if we're
 	//  trying to make the whole circle lie flat. The adjustment will be the size of one landscape
 	//  quad, but to be safe we'll make it two quads in each direction.
-	FVector3d LandscapeScale = InLandscapeInfo->LandscapeActor.IsValid() ?
-		InLandscapeInfo->LandscapeActor->GetActorTransform().GetScale3D() : FVector3d::Zero();
+	FVector3d LandscapeScale = InMergeContext->GetLandscape()->GetActorTransform().GetScale3D();
 	FVector2D Extents(2 * FMath::Max(LandscapeScale.X, LandscapeScale.Y) + Radius + Falloff);
 
 	FOOBox2D PatchArea(ComponentTransform, Extents);
@@ -246,11 +245,8 @@ TArray<UE::Landscape::EditLayers::FEditLayerRenderItem> ULandscapeCircleHeightPa
 	FInputWorldArea InputWorldArea = FInputWorldArea::CreateOOBox(PatchArea);
 	FOutputWorldArea OutputWorldArea = FOutputWorldArea::CreateOOBox(PatchArea);
 
-	FEditLayerTargetTypeState TargetInfo(bEditVisibility ? ELandscapeToolTargetTypeFlags::Visibility
-		: ELandscapeToolTargetTypeFlags::Heightmap);
-	FEditLayerRenderItem Item(TargetInfo, InputWorldArea, OutputWorldArea,
-		// ModifyExistingWeightmapsOnly: we want the patch to allocate the visibility layer if it needs to, in its region.
-		false);
+	FEditLayerTargetTypeState TargetInfo(InMergeContext, bEditVisibility ? ELandscapeToolTargetTypeFlags::Visibility : ELandscapeToolTargetTypeFlags::Heightmap);
+	FEditLayerRenderItem Item(TargetInfo, InputWorldArea, OutputWorldArea, /*bInModifyExistingWeightmapsOnly = */false); 
 	AffectedAreas.Add(Item);
 
 	return AffectedAreas;
diff --git a/Engine/Plugins/Experimental/LandscapePatch/Source/LandscapePatch/Private/LandscapePatchEditLayer.cpp b/Engine/Plugins/Experimental/LandscapePatch/Source/LandscapePatch/Private/LandscapePatchEditLayer.cpp
--- a/Engine/Plugins/Experimental/LandscapePatch/Source/LandscapePatch/Private/LandscapePatchEditLayer.cpp
+++ b/Engine/Plugins/Experimental/LandscapePatch/Source/LandscapePatch/Private/LandscapePatchEditLayer.cpp
@@ -312,7 +312,7 @@ bool ULandscapePatchEditLayer::SupportsTargetType(ELandscapeToolTargetType InTyp
 
 #if WITH_EDITOR
 // Called in batched merge path to apply the patches
-TArray<UE::Landscape::EditLayers::FEditLayerRendererState> ULandscapePatchEditLayer::GetEditLayerRendererStates(const ULandscapeInfo* InLandscapeInfo, bool bInSkipBrush)
+TArray<UE::Landscape::EditLayers::FEditLayerRendererState> ULandscapePatchEditLayer::GetEditLayerRendererStates(const UE::Landscape::EditLayers::FMergeContext* InMergeContext)
 {
 	UpdatePatchListIfDirty();
 
@@ -326,8 +326,8 @@ TArray<UE::Landscape::EditLayers::FEditLayerRendererState> ULandscapePatchEditLa
 			continue;
 		}
 
-		FEditLayerRendererState& RendererState = RendererStates.Emplace_GetRef(Patch, InLandscapeInfo);
-		if (bInSkipBrush || !Patch->IsEnabled())
+		FEditLayerRendererState& RendererState = RendererStates.Emplace_GetRef(InMergeContext, Patch);
+		if (InMergeContext->ShouldSkipProceduralRenderers() || !Patch->IsEnabled())
 		{
 			RendererState.DisableTargetTypeMask(ELandscapeToolTargetTypeFlags::All);
 		}
diff --git a/Engine/Plugins/Experimental/LandscapePatch/Source/LandscapePatch/Private/LandscapePatchManager.cpp b/Engine/Plugins/Experimental/LandscapePatch/Source/LandscapePatch/Private/LandscapePatchManager.cpp
--- a/Engine/Plugins/Experimental/LandscapePatch/Source/LandscapePatch/Private/LandscapePatchManager.cpp
+++ b/Engine/Plugins/Experimental/LandscapePatch/Source/LandscapePatch/Private/LandscapePatchManager.cpp
@@ -226,12 +226,19 @@ UTextureRenderTarget2D* ALandscapePatchManager::RenderLayer_Native(const FLandsc
 
 #if WITH_EDITOR
 // Called in batched merge path to apply the patches
-TArray<UE::Landscape::EditLayers::FEditLayerRendererState> ALandscapePatchManager::GetEditLayerRendererStates(const ULandscapeInfo* InLandscapeInfo, bool bInSkipBrush)
+TArray<UE::Landscape::EditLayers::FEditLayerRendererState> ALandscapePatchManager::GetEditLayerRendererStates(const UE::Landscape::EditLayers::FMergeContext* InMergeContext)
 {
+	using namespace UE::Landscape::EditLayers;
 	using namespace LandscapePatchManagerLocals;
 
 	FilterLegacyRegisteredPatches(PatchComponents, PatchToIndex, this);
 
+	// Inherit the base brush's enabled state and intersect it with all patches' state, so that the heightmap/weightmap/visibility toggle at the brush's level applies to them too : 
+	FEditLayerTargetTypeState BaseBrushSupportedState(InMergeContext);
+	FEditLayerTargetTypeState BaseBrushEnabledState(InMergeContext);
+	TArray<TBitArray<>> BaseBrushRenderGroups;
+	Super::GetRendererStateInfo(InMergeContext, BaseBrushSupportedState, BaseBrushEnabledState, BaseBrushRenderGroups);
+
 	TArray<FEditLayerRendererState> RendererStates;
 	RendererStates.Reserve(PatchComponents.Num());
 	for (TSoftObjectPtr<ULandscapePatchComponent>& PatchSoft : PatchComponents)
@@ -242,8 +249,10 @@ TArray<UE::Landscape::EditLayers::FEditLayerRendererState> ALandscapePatchManage
 			continue;
 		}
 
-		FEditLayerRendererState& RendererState = RendererStates.Emplace_GetRef(Patch, InLandscapeInfo);
-		if (bInSkipBrush || !Patch->IsEnabled())
+		FEditLayerRendererState& RendererState = RendererStates.Emplace_GetRef(InMergeContext, Patch);
+		// Disable all target types that are disabled in base brush : 
+		RendererState.DisableTargetTypeMask(~BaseBrushEnabledState.GetTargetTypeMask());
+		if (InMergeContext->ShouldSkipProceduralRenderers() || !Patch->IsEnabled())
 		{
 			RendererState.DisableTargetTypeMask(ELandscapeToolTargetTypeFlags::All);
 		}
@@ -487,7 +496,7 @@ bool ALandscapePatchManager::AffectsHeightmap() const
 {
 	using namespace LandscapePatchManagerLocals;
 
-	if (!CanAffectHeightmap())
+	if (!Super::AffectsHeightmap())
 	{
 		return false;
 	}
@@ -500,7 +509,7 @@ bool ALandscapePatchManager::AffectsWeightmap() const
 {
 	using namespace LandscapePatchManagerLocals;
 
-	if (!CanAffectWeightmap())
+	if (!Super::AffectsWeightmap())
 	{
 		return false;
 	}
@@ -514,7 +523,7 @@ bool ALandscapePatchManager::AffectsWeightmapLayer(const FName& InLayerName) con
 {
 	using namespace LandscapePatchManagerLocals;
 
-	if (!CanAffectWeightmap())
+	if (!Super::AffectsWeightmap()) // Don't call Super::AffectsWeightmapLayer(InLayerName) here as we don't want to use the AffectedWeightmapLayers list for weightmap layers
 	{
 		return false;
 	}
@@ -527,7 +536,7 @@ bool ALandscapePatchManager::AffectsVisibilityLayer() const
 {
 	using namespace LandscapePatchManagerLocals;
 
-	if (!CanAffectVisibilityLayer())
+	if (!Super::AffectsVisibilityLayer())
 	{
 		return false;
 	}
@@ -540,7 +549,7 @@ bool ALandscapePatchManager::CanAffectWeightmapLayer(const FName& InLayerName) c
 {
 	using namespace LandscapePatchManagerLocals;
 
-	if (!CanAffectWeightmap())
+	if (!Super::CanAffectWeightmap()) // Don't call Super::CanAffectWeightmapLayer(InLayerName) here as we don't want to use the AffectedWeightmapLayers list for weightmap layers
 	{
 		return false;
 	}
diff --git a/Engine/Plugins/Experimental/LandscapePatch/Source/LandscapePatch/Private/LandscapeTexturePatch.cpp b/Engine/Plugins/Experimental/LandscapePatch/Source/LandscapePatch/Private/LandscapeTexturePatch.cpp
--- a/Engine/Plugins/Experimental/LandscapePatch/Source/LandscapePatch/Private/LandscapeTexturePatch.cpp
+++ b/Engine/Plugins/Experimental/LandscapePatch/Source/LandscapePatch/Private/LandscapeTexturePatch.cpp
@@ -161,31 +161,25 @@ void ULandscapeTexturePatch::RenderLayer(FRenderParams& InRenderParams)
 	// If we got to here, we're dealing with weightmaps.
 
 	UTextureRenderTarget2DArray* TextureArray = LandscapeScratchRT->TryGetRenderTarget2DArray();
-	if (!ensure(TextureArray))
-	{
-		return;
-	}
+	check(TextureArray != nullptr);
 	
 	// Only need to transition if we get a matching weight patch
 	bool bTransitionedToSRV = false;
 
-	int32 NumLayers = InRenderParams.RenderGroupTargetLayerNames.Num();
-	if (!ensure(TextureArray->Slices == NumLayers))
-	{
-		NumLayers = FMath::Min(NumLayers, TextureArray->Slices);
-	}
+	int32 NumTargetLayersInRenderGroup = InRenderParams.RenderGroupTargetLayerNames.Num();
+	check(TextureArray->Slices == NumTargetLayersInRenderGroup);
 
-	for (int32 LayerIndex = 0; LayerIndex < NumLayers; ++LayerIndex)
+	for (int32 TargetLayerIndexInRenderGroup = 0; TargetLayerIndexInRenderGroup < NumTargetLayersInRenderGroup; ++TargetLayerIndexInRenderGroup)
 	{
-		bool bIsVisibilityLayer = ensure(LayerIndex < InRenderParams.RenderGroupTargetLayerInfos.Num()) 
-			&& UE::Landscape::IsVisibilityLayer(InRenderParams.RenderGroupTargetLayerInfos[LayerIndex]);
+		bool bIsVisibilityLayer = ensure(TargetLayerIndexInRenderGroup < InRenderParams.RenderGroupTargetLayerInfos.Num())
+			&& UE::Landscape::IsVisibilityLayer(InRenderParams.RenderGroupTargetLayerInfos[TargetLayerIndexInRenderGroup]);
 		
 		// Try to find the weight patch
 		ULandscapeWeightPatchTextureInfo* WeightPatchInfo = nullptr;
 		for (const TObjectPtr<ULandscapeWeightPatchTextureInfo>& WeightPatchEntry : WeightPatches)
 		{
 			if ((bIsVisibilityLayer && WeightPatchEntry->bEditVisibilityLayer)
-				|| (WeightPatchEntry->WeightmapLayerName == InRenderParams.RenderGroupTargetLayerNames[LayerIndex]))
+				|| (WeightPatchEntry->WeightmapLayerName == InRenderParams.RenderGroupTargetLayerNames[TargetLayerIndexInRenderGroup]))
 			{
 				WeightPatchInfo = WeightPatchEntry;
 				break;
@@ -209,13 +203,13 @@ void ULandscapeTexturePatch::RenderLayer(FRenderParams& InRenderParams)
 		{
 			WeightPatchInfo->bReinitializeOnNextRender = false;
 			ReinitializeWeightPatch(WeightPatchInfo, TextureArray->GetResource(),
-				FIntPoint(TextureArray->SizeX, TextureArray->SizeY), LayerIndex, LandscapeHeightmapToWorld);
+				FIntPoint(TextureArray->SizeX, TextureArray->SizeY), TargetLayerIndexInRenderGroup, LandscapeHeightmapToWorld);
 		}
 		else
 		{
 			ApplyToWeightmap(WeightPatchInfo, 
 				[Resource = TextureArray->GetResource()]() { return Resource->GetTexture2DArrayRHI(); },
-				LayerIndex, 
+				TargetLayerIndexInRenderGroup,
 				InRenderParams.RenderAreaSectionRect.Size(), 
 				GetHeightmapToWorld(InRenderParams.RenderAreaWorldTransform));
 		}
@@ -2020,28 +2014,33 @@ void ULandscapeTexturePatch::SetEditVisibilityLayer(const FName& InWeightmapLaye
 }
 
 #if WITH_EDITOR
-void ULandscapeTexturePatch::GetRendererStateInfo(const ULandscapeInfo* InLandscapeInfo,
+void ULandscapeTexturePatch::GetRendererStateInfo(const UE::Landscape::EditLayers::FMergeContext* InMergeContext,
 	FEditLayerTargetTypeState& OutSupported, FEditLayerTargetTypeState& OutEnabled,
-	TArray<TSet<FName>>& OutRenderGroups) const
+	TArray<TBitArray<>>& OutRenderGroups) const
 {
-	if (CanAffectHeightmap())
+	if (InMergeContext->IsHeightmapMerge())
 	{
-		OutSupported.AddTargetType(ELandscapeToolTargetType::Heightmap);
+		if (CanAffectHeightmap())
+		{
+			OutSupported.AddTargetType(ELandscapeToolTargetType::Heightmap);
+		}
 	}
-
-	for (const TObjectPtr<ULandscapeWeightPatchTextureInfo>& WeightPatch : WeightPatches)
+	else
 	{
-		if (IsValid(WeightPatch) && WeightPatchCanRender(*WeightPatch))
+		for (const TObjectPtr<ULandscapeWeightPatchTextureInfo>& WeightPatch : WeightPatches)
 		{
-			if (WeightPatch->bEditVisibilityLayer)
-			{
-				OutSupported.AddTargetType(ELandscapeToolTargetType::Visibility);
-			}
-			else
+			if (IsValid(WeightPatch) && WeightPatchCanRender(*WeightPatch))
 			{
-				OutSupported.AddTargetType(ELandscapeToolTargetType::Weightmap);
-				OutSupported.AddWeightmap(WeightPatch->WeightmapLayerName);
-				OutRenderGroups.Add(TSet<FName>({ WeightPatch->WeightmapLayerName }));
+				if (WeightPatch->bEditVisibilityLayer)
+				{
+					OutSupported.AddTargetType(ELandscapeToolTargetType::Visibility);
+				}
+				else if (InMergeContext->IsValidTargetLayerName(WeightPatch->WeightmapLayerName))
+				{
+					const int32 TargetLayerIndex = InMergeContext->GetTargetLayerIndexForNameChecked(WeightPatch->WeightmapLayerName);
+					OutSupported.AddTargetType(ELandscapeToolTargetType::Weightmap);
+					OutSupported.AddWeightmap(TargetLayerIndex);
+				}
 			}
 		}
 	}
@@ -2057,7 +2056,7 @@ FString ULandscapeTexturePatch::GetEditLayerRendererDebugName() const
 	return FString::Printf(TEXT("%s:%s"), *GetOwner()->GetActorNameOrLabel(), *GetName());
 }
 
-TArray<UE::Landscape::EditLayers::FEditLayerRenderItem> ULandscapeTexturePatch::GetRenderItems(const ULandscapeInfo* InLandscapeInfo) const
+TArray<UE::Landscape::EditLayers::FEditLayerRenderItem> ULandscapeTexturePatch::GetRenderItems(const UE::Landscape::EditLayers::FMergeContext* InMergeContext) const
 {
 	using namespace UE::Landscape::EditLayers;
 
@@ -2068,24 +2067,37 @@ TArray<UE::Landscape::EditLayers::FEditLayerRenderItem> ULandscapeTexturePatch::
 	FInputWorldArea InputWorldArea = FInputWorldArea::CreateOOBox(PatchArea);
 	FOutputWorldArea OutputWorldArea = FOutputWorldArea::CreateOOBox(PatchArea);
 
-	if (CanAffectHeightmap())
-	{
-		FEditLayerTargetTypeState TargetInfo(ELandscapeToolTargetTypeFlags::Heightmap);
-		FEditLayerRenderItem Item(TargetInfo, InputWorldArea, OutputWorldArea, false);
-		AffectedAreas.Add(Item);
-	}
-
-	for (const TObjectPtr<ULandscapeWeightPatchTextureInfo>& WeightPatch : WeightPatches)
+	if (InMergeContext->IsHeightmapMerge())
 	{
-		if (IsValid(WeightPatch) && WeightPatchCanRender(*WeightPatch))
+		if (CanAffectHeightmap())
 		{
-			FEditLayerTargetTypeState TargetInfo = WeightPatch->bEditVisibilityLayer ? 
-				FEditLayerTargetTypeState(ELandscapeToolTargetTypeFlags::Visibility)
-				: FEditLayerTargetTypeState(ELandscapeToolTargetTypeFlags::Weightmap, { WeightPatch->WeightmapLayerName });
+			FEditLayerTargetTypeState TargetInfo(InMergeContext, ELandscapeToolTargetTypeFlags::Heightmap);
 			FEditLayerRenderItem Item(TargetInfo, InputWorldArea, OutputWorldArea, false);
 			AffectedAreas.Add(Item);
 		}
 	}
+	else
+	{
+		for (const TObjectPtr<ULandscapeWeightPatchTextureInfo>& WeightPatch : WeightPatches)
+		{
+			if (IsValid(WeightPatch) && WeightPatchCanRender(*WeightPatch))
+			{
+				FEditLayerTargetTypeState TargetInfo(InMergeContext);
+				if (WeightPatch->bEditVisibilityLayer)
+				{
+					TargetInfo.AddTargetType(ELandscapeToolTargetType::Visibility);
+				}
+				else if (InMergeContext->IsValidTargetLayerName(WeightPatch->WeightmapLayerName))
+				{
+					const int32 TargetLayerIndex = InMergeContext->GetTargetLayerIndexForNameChecked(WeightPatch->WeightmapLayerName);
+					TargetInfo.AddTargetType(ELandscapeToolTargetType::Weightmap);
+					TargetInfo.AddWeightmap(TargetLayerIndex);
+				}
+				FEditLayerRenderItem Item(TargetInfo, InputWorldArea, OutputWorldArea, /*bInModifyExistingWeightmapsOnly = */false);
+				AffectedAreas.Add(Item);
+			}
+		}
+	}
 
 	return AffectedAreas;
 }
diff --git a/Engine/Plugins/Experimental/LandscapePatch/Source/LandscapePatch/Public/LandscapeCircleHeightPatch.h b/Engine/Plugins/Experimental/LandscapePatch/Source/LandscapePatch/Public/LandscapeCircleHeightPatch.h
--- a/Engine/Plugins/Experimental/LandscapePatch/Source/LandscapePatch/Public/LandscapeCircleHeightPatch.h
+++ b/Engine/Plugins/Experimental/LandscapePatch/Source/LandscapePatch/Public/LandscapeCircleHeightPatch.h
@@ -26,12 +26,12 @@ public:
 	using FEditLayerRenderItem = UE::Landscape::EditLayers::FEditLayerRenderItem;
 
 	// ILandscapeEditLayerRenderer, via ULandscapePatchComponent
-	virtual void GetRendererStateInfo(const ULandscapeInfo* InLandscapeInfo,
+	virtual void GetRendererStateInfo(const UE::Landscape::EditLayers::FMergeContext* InMergeContext,
 		FEditLayerTargetTypeState& OutSupportedTargetTypeState, 
 		FEditLayerTargetTypeState& OutEnabledTargetTypeState, 
-		TArray<TSet<FName>>& OutRenderGroups) const override;
+		TArray<TBitArray<>>& OutRenderGroups) const override;
 	virtual FString GetEditLayerRendererDebugName() const override;
-	virtual TArray<FEditLayerRenderItem> GetRenderItems(const ULandscapeInfo* InLandscapeInfo) const override;
+	virtual TArray<FEditLayerRenderItem> GetRenderItems(const UE::Landscape::EditLayers::FMergeContext* InMergeContext) const override;
 	virtual void RenderLayer(FRenderParams& InRenderParams) override;
 #endif
 
diff --git a/Engine/Plugins/Experimental/LandscapePatch/Source/LandscapePatch/Public/LandscapePatchEditLayer.h b/Engine/Plugins/Experimental/LandscapePatch/Source/LandscapePatch/Public/LandscapePatchEditLayer.h
--- a/Engine/Plugins/Experimental/LandscapePatch/Source/LandscapePatch/Public/LandscapePatchEditLayer.h
+++ b/Engine/Plugins/Experimental/LandscapePatch/Source/LandscapePatch/Public/LandscapePatchEditLayer.h
@@ -64,7 +64,7 @@ public:
 	using FEditLayerRendererState = UE::Landscape::EditLayers::FEditLayerRendererState;
 
 	// IEditLayerRendererProvider
-	virtual TArray<FEditLayerRendererState> GetEditLayerRendererStates(const ULandscapeInfo* InLandscapeInfo, bool bInSkipBrush) override;
+	virtual TArray<FEditLayerRendererState> GetEditLayerRendererStates(const UE::Landscape::EditLayers::FMergeContext* InMergeContext) override;
 
 	// UObject
 	virtual void PostLoad() override;
diff --git a/Engine/Plugins/Experimental/LandscapePatch/Source/LandscapePatch/Public/LandscapePatchManager.h b/Engine/Plugins/Experimental/LandscapePatch/Source/LandscapePatch/Public/LandscapePatchManager.h
--- a/Engine/Plugins/Experimental/LandscapePatch/Source/LandscapePatch/Public/LandscapePatchManager.h
+++ b/Engine/Plugins/Experimental/LandscapePatch/Source/LandscapePatch/Public/LandscapePatchManager.h
@@ -37,17 +37,17 @@ public:
 
 	// IEditLayerRendererProvider
 	// Called by the batched-merge application path.
-	virtual TArray<FEditLayerRendererState> GetEditLayerRendererStates(const ULandscapeInfo* InLandscapeInfo, bool bInSkipBrush) override;
+	virtual TArray<FEditLayerRendererState> GetEditLayerRendererStates(const UE::Landscape::EditLayers::FMergeContext* InMergeContext) override;
 
 	// ILandscapeEditLayerRenderer
 	//~ In batched merge, the manager relies on being a renderer provider. It does not need to have its 
 	//~  RenderLayer_Native method called, so we override the implementations of ILandscapeEditLayerRenderer
 	//~  inherited from ALandscapeBlueprintBrushBase to do nothing.
-	virtual void GetRendererStateInfo(const ULandscapeInfo* InLandscapeInfo,
+	virtual void GetRendererStateInfo(const UE::Landscape::EditLayers::FMergeContext* InMergeContext,
 		UE::Landscape::EditLayers::FEditLayerTargetTypeState& OutSupportedTargetTypeState,
 		UE::Landscape::EditLayers::FEditLayerTargetTypeState& OutEnabledTargetTypeState,
-		TArray<TSet<FName>>& OutRenderGroups) const override {}
-	virtual TArray<UE::Landscape::EditLayers::FEditLayerRenderItem> GetRenderItems(const ULandscapeInfo* InLandscapeInfo) const override
+		TArray<TBitArray<>>& OutRenderGroups) const override {}
+	virtual TArray<UE::Landscape::EditLayers::FEditLayerRenderItem> GetRenderItems(const UE::Landscape::EditLayers::FMergeContext* InMergeContext) const override
 	{
 		return {};
 	}
diff --git a/Engine/Plugins/Experimental/LandscapePatch/Source/LandscapePatch/Public/LandscapeTexturePatch.h b/Engine/Plugins/Experimental/LandscapePatch/Source/LandscapePatch/Public/LandscapeTexturePatch.h
--- a/Engine/Plugins/Experimental/LandscapePatch/Source/LandscapePatch/Public/LandscapeTexturePatch.h
+++ b/Engine/Plugins/Experimental/LandscapePatch/Source/LandscapePatch/Public/LandscapeTexturePatch.h
@@ -230,12 +230,12 @@ public:
 	using FEditLayerRenderItem = UE::Landscape::EditLayers::FEditLayerRenderItem;
 
 	// ILandscapeEditLayerRenderer, via ULandscapePatchComponent
-	virtual void GetRendererStateInfo(const ULandscapeInfo* InLandscapeInfo,
+	virtual void GetRendererStateInfo(const UE::Landscape::EditLayers::FMergeContext* InMergeContext,
 		FEditLayerTargetTypeState& OutSupportedTargetTypeState,
 		FEditLayerTargetTypeState& OutEnabledTargetTypeState,
-		TArray<TSet<FName>>& OutRenderGroups) const override;
+		TArray<TBitArray<>>& OutRenderGroups) const override;
 	virtual FString GetEditLayerRendererDebugName() const override;
-	virtual TArray<FEditLayerRenderItem> GetRenderItems(const ULandscapeInfo* InLandscapeInfo) const override;
+	virtual TArray<FEditLayerRenderItem> GetRenderItems(const UE::Landscape::EditLayers::FMergeContext* InMergeContext) const override;
 	virtual void RenderLayer(FRenderParams& InRenderParams) override;
 
 	// ULandscapePatchComponent
diff --git a/Engine/Plugins/Experimental/Water/Source/Editor/Private/WaterLandscapeBrush.cpp b/Engine/Plugins/Experimental/Water/Source/Editor/Private/WaterLandscapeBrush.cpp
--- a/Engine/Plugins/Experimental/Water/Source/Editor/Private/WaterLandscapeBrush.cpp
+++ b/Engine/Plugins/Experimental/Water/Source/Editor/Private/WaterLandscapeBrush.cpp
@@ -403,11 +403,11 @@ void AWaterLandscapeBrush::AddReferencedObjects(UObject* InThis, FReferenceColle
 }
 
 #if WITH_EDITOR
-TArray<UE::Landscape::EditLayers::FEditLayerRenderItem> AWaterLandscapeBrush::GetRenderItems(const ULandscapeInfo* InLandscapeInfo) const
+TArray<UE::Landscape::EditLayers::FEditLayerRenderItem> AWaterLandscapeBrush::GetRenderItems(const UE::Landscape::EditLayers::FMergeContext* InMergeContext) const
 {
 	using namespace UE::Landscape::EditLayers; 
 
-	TArray<FEditLayerRenderItem> RenderItems = Super::GetRenderItems(InLandscapeInfo);
+	TArray<FEditLayerRenderItem> RenderItems = Super::GetRenderItems(InMergeContext);
 	check(RenderItems.Num() == 1);
 
 	// For now, this brush requires the entire landscape to be loaded to work deterministically, so we force the input area to be "infinite" : 
diff --git a/Engine/Plugins/Experimental/Water/Source/Editor/Private/WaterLandscapeBrush.h b/Engine/Plugins/Experimental/Water/Source/Editor/Private/WaterLandscapeBrush.h
--- a/Engine/Plugins/Experimental/Water/Source/Editor/Private/WaterLandscapeBrush.h
+++ b/Engine/Plugins/Experimental/Water/Source/Editor/Private/WaterLandscapeBrush.h
@@ -31,7 +31,7 @@ public:
 
 #if WITH_EDITOR
 	//~ Begin ILandscapeEditLayerRenderer implementation
-	virtual TArray<UE::Landscape::EditLayers::FEditLayerRenderItem> GetRenderItems(const ULandscapeInfo* InLandscapeInfo) const override;
+	virtual TArray<UE::Landscape::EditLayers::FEditLayerRenderItem> GetRenderItems(const UE::Landscape::EditLayers::FMergeContext* InMergeContext) const override;
 	//~ End ILandscapeEditLayerRenderer implementation
 #endif // WITH_EDITOR
 
diff --git a/Engine/Shaders/Private/Landscape/LandscapeEditLayersWeightmaps.usf b/Engine/Shaders/Private/Landscape/LandscapeEditLayersWeightmaps.usf
--- a/Engine/Shaders/Private/Landscape/LandscapeEditLayersWeightmaps.usf
+++ b/Engine/Shaders/Private/Landscape/LandscapeEditLayersWeightmaps.usf
@@ -99,11 +99,11 @@ void MergeEditLayerPS(in float4 SVPos : SV_POSITION, out float OutWeight : SV_Ta
 	
 	if (InEditLayerPaintLayerBlendMode == EEDITLAYERWEIGHTMAPBLENDMODE_ADDITIVE)
 	{
-		OutWeight += CurrentLayerWeight;
+		OutWeight += CurrentLayerWeight * InEditLayerAlpha;
 	}
 	else if (InEditLayerPaintLayerBlendMode == EEDITLAYERWEIGHTMAPBLENDMODE_SUBTRACTIVE)
 	{
-		OutWeight -= CurrentLayerWeight;
+		OutWeight -= CurrentLayerWeight * InEditLayerAlpha;
 	}
 }
 #endif // MERGE_EDIT_LAYER
diff --git a/Engine/Source/Runtime/Landscape/Classes/Landscape.h b/Engine/Source/Runtime/Landscape/Classes/Landscape.h
--- a/Engine/Source/Runtime/Landscape/Classes/Landscape.h
+++ b/Engine/Source/Runtime/Landscape/Classes/Landscape.h
@@ -123,7 +123,6 @@ struct FLandscapeLayerBrush
 	UTextureRenderTarget2D* RenderLayer(const FLandscapeBrushParameters& InParameters);
 
 	LANDSCAPE_API ALandscapeBlueprintBrushBase* GetBrush() const;
-	bool Affects(ELandscapeToolTargetType InTargetType, const FName& InWeightmapLayerName) const;
 	bool AffectsHeightmap() const;
 	bool AffectsWeightmapLayer(const FName& InWeightmapLayerName) const;
 	bool AffectsVisibilityLayer() const;
@@ -131,7 +130,7 @@ struct FLandscapeLayerBrush
 
 	//~ Begin UE::Landscape::EditLayers::IEditLayerRendererProvider implementation
 	LANDSCAPE_EDIT_LAYERS_BATCHED_MERGE_EXPERIMENTAL
-	LANDSCAPE_API virtual TArray<UE::Landscape::EditLayers::FEditLayerRendererState> GetEditLayerRendererStates(const ULandscapeInfo* InLandscapeInfo, bool bInSkipBrush) override;
+	LANDSCAPE_API virtual TArray<UE::Landscape::EditLayers::FEditLayerRendererState> GetEditLayerRendererStates(const UE::Landscape::EditLayers::FMergeContext* InMergeContext) override;
 	//~ End UE::Landscape::EditLayers::IEditLayerRendererProvider implementation
 #endif // WITH_EDITOR
 
@@ -174,8 +173,10 @@ struct FLandscapeLayer
 #if WITH_EDITOR
 	//~ Begin UE::Landscape::EditLayers::IEditLayerRendererProvider implementation
 	LANDSCAPE_EDIT_LAYERS_BATCHED_MERGE_EXPERIMENTAL
-	LANDSCAPE_API virtual TArray<UE::Landscape::EditLayers::FEditLayerRendererState> GetEditLayerRendererStates(const ULandscapeInfo* InLandscapeInfo, bool bInSkipBrush) override;
+	LANDSCAPE_API virtual TArray<UE::Landscape::EditLayers::FEditLayerRendererState> GetEditLayerRendererStates(const UE::Landscape::EditLayers::FMergeContext* InMergeContext) override;
 	//~ End UE::Landscape::EditLayers::IEditLayerRendererProvider implementation
+
+	ELandscapeToolTargetTypeFlags GetEnabledTargetTypeMask() const;
 #endif // WITH_EDITOR
 
 	// TODO [jonathan.bard] : make those private
@@ -526,10 +527,10 @@ private:
 	void MonitorLandscapeEdModeChanges();
 	
 	//~ Begin UE::Landscape::EditLayers::IEditLayerRendererProvider implementation
-	virtual TArray<UE::Landscape::EditLayers::FEditLayerRendererState> GetEditLayerRendererStates(const ULandscapeInfo* InLandscapeInfo, bool bInSkipBrush) override;
+	virtual TArray<UE::Landscape::EditLayers::FEditLayerRendererState> GetEditLayerRendererStates(const UE::Landscape::EditLayers::FMergeContext* InMergeContext) override;
 	//~ End UE::Landscape::EditLayers::IEditLayerRendererProvider implementation
 
-	UE::Landscape::EditLayers::FMergeRenderContext PrepareEditLayersMergeRenderContext(const UE::Landscape::EditLayers::FMergeRenderParams& InParams);
+	UE::Landscape::EditLayers::FMergeRenderContext PrepareEditLayersMergeRenderContext(const UE::Landscape::EditLayers::FMergeContext& InMergeContext, const UE::Landscape::EditLayers::FMergeRenderParams& InParams);
 
 	FLandscapeBrushParameters BuildLandscapeBrushParameters(bool bInIsHeightmapMerge, const FIntRect& InRenderAreaExtents, UTextureRenderTarget2D* InRenderTarget, FName InWeightmapLayerName = NAME_None);
 
diff --git a/Engine/Source/Runtime/Landscape/Classes/LandscapeEditLayer.h b/Engine/Source/Runtime/Landscape/Classes/LandscapeEditLayer.h
--- a/Engine/Source/Runtime/Landscape/Classes/LandscapeEditLayer.h
+++ b/Engine/Source/Runtime/Landscape/Classes/LandscapeEditLayer.h
@@ -157,11 +157,7 @@ public:
 	//~ Begin IEditLayerRendererProvider implementation
 	// By default this does nothing in a landscape edit layer, but subclasses can override it if 
 	//  they would like to provide additional renderers.
-	virtual TArray<UE::Landscape::EditLayers::FEditLayerRendererState> GetEditLayerRendererStates(
-		const ULandscapeInfo* InLandscapeInfo, bool bInSkipBrush) 
-	{ 
-		return {};
-	};
+	virtual TArray<UE::Landscape::EditLayers::FEditLayerRendererState> GetEditLayerRendererStates(const UE::Landscape::EditLayers::FMergeContext* InMergeContext) {  return {};  };
 	//~ End IEditLayerRendererProvider implementation
 #endif
 
@@ -198,10 +194,10 @@ public:
 #if WITH_EDITOR
 	//~ Begin ILandscapeEditLayerRenderer implementation
 	LANDSCAPE_EDIT_LAYERS_BATCHED_MERGE_EXPERIMENTAL
-	LANDSCAPE_API virtual void GetRendererStateInfo(const ULandscapeInfo* InLandscapeInfo,
-		UE::Landscape::EditLayers::FEditLayerTargetTypeState& OutSupportedTargetTypeState, UE::Landscape::EditLayers::FEditLayerTargetTypeState& OutEnabledTargetTypeState, TArray<TSet<FName>>& OutRenderGroups) const override;
+	LANDSCAPE_API virtual void GetRendererStateInfo(const UE::Landscape::EditLayers::FMergeContext* InMergeContext,
+		UE::Landscape::EditLayers::FEditLayerTargetTypeState& OutSupportedTargetTypeState, UE::Landscape::EditLayers::FEditLayerTargetTypeState& OutEnabledTargetTypeState, TArray<TBitArray<>>& OutRenderGroups) const override;
 	LANDSCAPE_EDIT_LAYERS_BATCHED_MERGE_EXPERIMENTAL
-	LANDSCAPE_API virtual TArray<UE::Landscape::EditLayers::FEditLayerRenderItem> GetRenderItems(const ULandscapeInfo* InLandscapeInfo) const override;
+	LANDSCAPE_API virtual TArray<UE::Landscape::EditLayers::FEditLayerRenderItem> GetRenderItems(const UE::Landscape::EditLayers::FMergeContext* InMergeContext) const override;
 	LANDSCAPE_EDIT_LAYERS_BATCHED_MERGE_EXPERIMENTAL
 	LANDSCAPE_API virtual void RenderLayer(ILandscapeEditLayerRenderer::FRenderParams& InRenderParams) override;
 	LANDSCAPE_EDIT_LAYERS_BATCHED_MERGE_EXPERIMENTAL
diff --git a/Engine/Source/Runtime/Landscape/Private/LandscapeBlueprintBrushBase.cpp b/Engine/Source/Runtime/Landscape/Private/LandscapeBlueprintBrushBase.cpp
--- a/Engine/Source/Runtime/Landscape/Private/LandscapeBlueprintBrushBase.cpp
+++ b/Engine/Source/Runtime/Landscape/Private/LandscapeBlueprintBrushBase.cpp
@@ -104,7 +104,7 @@ UTextureRenderTarget2D* ALandscapeBlueprintBrushBase::Execute(const FLandscapeBr
 	}
 	UTextureRenderTarget2D* Result = nullptr;
 	{
-		RHI_BREADCRUMB_EVENT_GAMETHREAD("BP Render (%s): %s", UEnum::GetValueAsString(InParameters.LayerType), LayerDetailString);
+		RHI_BREADCRUMB_EVENT_GAMETHREAD("BP Render %s (%s): %s", GetActorNameOrLabel(), UEnum::GetDisplayValueAsText(InParameters.LayerType).ToString(), LayerDetailString);
 
 		TGuardValue<bool> AutoRestore(GAllowActorScriptExecutionInEditor, true);
 		Result = RenderLayer(InParameters);
@@ -115,7 +115,7 @@ UTextureRenderTarget2D* ALandscapeBlueprintBrushBase::Execute(const FLandscapeBr
 		((Result->SizeX != InParameters.CombinedResult->SizeX) || (Result->SizeY != InParameters.CombinedResult->SizeY)))
 	{
 		UE_LOG(LogLandscape, Warning, TEXT("In landscape %s, the BP brush %s failed to render for (%s%s). Make sure the brush properly implements RenderLayer and returns a render target of the appropriate size: expected (%i, %i), actual (%i, %i). This brush will be skipped until then."), 
-			*OwningLandscape->GetActorLabel(), *GetActorLabel(), *UEnum::GetValueAsString(InParameters.LayerType), *LayerDetailString, InParameters.CombinedResult->SizeX, InParameters.CombinedResult->SizeY, Result->SizeX, Result->SizeY );
+			*OwningLandscape->GetActorLabel(), *GetActorLabel(), *UEnum::GetDisplayValueAsText(InParameters.LayerType).ToString(), *LayerDetailString, InParameters.CombinedResult->SizeX, InParameters.CombinedResult->SizeY, Result->SizeX, Result->SizeY);
 		Result = nullptr;
 	}
 
@@ -221,8 +221,8 @@ void ALandscapeBlueprintBrushBase::SetCanAffectVisibilityLayer(bool bInCanAffect
 
 #if WITH_EDITOR
 
-void ALandscapeBlueprintBrushBase::GetRendererStateInfo(const ULandscapeInfo* InLandscapeInfo,
-	UE::Landscape::EditLayers::FEditLayerTargetTypeState& OutSupportedTargetTypeState, UE::Landscape::EditLayers::FEditLayerTargetTypeState& OutEnabledTargetTypeState, TArray<TSet<FName>>& OutRenderGroups) const
+void ALandscapeBlueprintBrushBase::GetRendererStateInfo(const UE::Landscape::EditLayers::FMergeContext* InMergeContext,
+	UE::Landscape::EditLayers::FEditLayerTargetTypeState& OutSupportedTargetTypeState, UE::Landscape::EditLayers::FEditLayerTargetTypeState& OutEnabledTargetTypeState, TArray<TBitArray<>>& OutRenderGroups) const
 {
 	// What can the brush do?
 	if (CanAffectHeightmap())
@@ -253,33 +253,32 @@ void ALandscapeBlueprintBrushBase::GetRendererStateInfo(const ULandscapeInfo* In
 	}
 
 	// Mark which weightmap is supported/enabled : 
-	if (CanAffectWeightmap())
+	if (!InMergeContext->IsHeightmapMerge() && CanAffectWeightmap())
 	{
-		for (const FLandscapeInfoLayerSettings& InfoLayerSettings : InLandscapeInfo->Layers)
-		{
-			if (InfoLayerSettings.LayerInfoObj != nullptr)
+		InMergeContext->ForEachValidTargetLayer(
+			[InMergeContext, this, &OutSupportedTargetTypeState, &OutEnabledTargetTypeState](int32 InTargetLayerIndex, const FName& InTargetLayerName, ULandscapeLayerInfoObject* InWeightmapLayerInfo)
 			{
-				FName WeightmapLayerName = InfoLayerSettings.GetLayerName();
-				if (CanAffectWeightmapLayer(WeightmapLayerName))
+				if (CanAffectWeightmapLayer(InTargetLayerName))
 				{
-					OutSupportedTargetTypeState.AddWeightmap(WeightmapLayerName);
-					if (AffectsWeightmapLayer(WeightmapLayerName))
+					OutSupportedTargetTypeState.AddWeightmap(InTargetLayerIndex);
+					if (AffectsWeightmapLayer(InTargetLayerName))
 					{
-						OutEnabledTargetTypeState.AddWeightmap(WeightmapLayerName);
+						OutEnabledTargetTypeState.AddWeightmap(InTargetLayerIndex);
 					}
 				}
-			}
-		}
+				return true;
+			});
 	}
 }
 
-TArray<UE::Landscape::EditLayers::FEditLayerRenderItem> ALandscapeBlueprintBrushBase::GetRenderItems(const ULandscapeInfo* InLandscapeInfo) const
+TArray<UE::Landscape::EditLayers::FEditLayerRenderItem> ALandscapeBlueprintBrushBase::GetRenderItems(const UE::Landscape::EditLayers::FMergeContext* InMergeContext) const
 {
 	using namespace UE::Landscape::EditLayers;
 
-	FEditLayerTargetTypeState SupportedTargetTypeState, EnabledTargetTypeState;
-	TArray<TSet<FName>> DummyRenderGroups;
-	GetRendererStateInfo(InLandscapeInfo, SupportedTargetTypeState, EnabledTargetTypeState, DummyRenderGroups);
+	FEditLayerTargetTypeState SupportedTargetTypeState(InMergeContext);
+	FEditLayerTargetTypeState EnabledTargetTypeState(InMergeContext);
+	TArray<TBitArray<>> DummyRenderGroups;
+	GetRendererStateInfo(InMergeContext, SupportedTargetTypeState, EnabledTargetTypeState, DummyRenderGroups);
 
 	// By default, for landscape BP brushes, we use FInputWorldArea::EType::Infinite, to indicate they can only reliably work when applied globally on the entire landscape
 	//  This allows full backwards-compatibility but will prevent landscapes from benefiting from batched merge. Users will be able to indicate their brush works in a local fashion
@@ -305,51 +304,50 @@ void ALandscapeBlueprintBrushBase::RenderLayer(ILandscapeEditLayerRenderer::FRen
 	ULandscapeScratchRenderTarget* CurrentLayerReadRT = InRenderParams.MergeRenderContext->GetBlendRenderTargetRead();
 
 	const bool bIsHeightmapMerge = InRenderParams.MergeRenderContext->IsHeightmapMerge();
-
-	// Because we only expose UTextureRenderTarget2D to BP, in the case of weightmaps, we need an additional scratch render target 2D that we will copy the 
-	//  current result of each paint layer into, so that the BP can use it as its source : 
+	const TArray<FName> EnabledWeightmaps = InRenderParams.RendererState.GetActiveTargetWeightmaps();
+	
 	ULandscapeEditResourcesSubsystem* LandscapeEditResourcesSubsystem = GEngine->GetEngineSubsystem<ULandscapeEditResourcesSubsystem>();
 	check(LandscapeEditResourcesSubsystem != nullptr);
 	ULandscapeScratchRenderTarget* CurrentLayerReadRT2D = CurrentLayerReadRT;
 
-	// We might require an additional scratch render target (it will get freed/recycled at the end of this function when the AdditionalScratchTexture goes out of scope) : 
-	TUniquePtr<FScratchRenderTargetScope> AdditionalScratchTexture;
-
-	TArray<FName> EnabledWeightmaps;
-	if (!bIsHeightmapMerge)
-	{
-		EnabledWeightmaps = InRenderParams.RendererState.GetEnabledTargetWeightmaps();
-
-		check(CurrentLayerReadRT->IsTexture2DArray() && WriteRT->IsTexture2DArray());
-		FScratchRenderTargetParams ScratchRenderTargetParams(TEXT("BPBrushScratchRT"), /*bInExactDimensions = */false, /*bInUseUAV = */false, /*bInTargetArraySlicesIndependently = */false,
-			InRenderParams.RenderAreaSectionRect.Size(), /*InNumSlices = */0, CurrentLayerReadRT->GetFormat(), CurrentLayerReadRT->GetClearColor(), ERHIAccess::CopyDest);
-		// We need a new scratch 2D texture in order to copy the merged result of a single target layer, so that the BP brush can use that as input (as it operates on texture 2D render targets only): 
-		AdditionalScratchTexture = MakeUnique<FScratchRenderTargetScope>(ScratchRenderTargetParams);
-		CurrentLayerReadRT2D = AdditionalScratchTexture->RenderTarget;
-
-		// The original texture array will be accessed as ERHIAccess::CopySrc all along : 
-		CurrentLayerReadRT->TransitionTo(ERHIAccess::CopySrc);
-	}
-
-	const int32 NumTargetLayers = InRenderParams.RenderGroupTargetLayerNames.Num();
-	for (int32 TargetLayerIndex = 0; TargetLayerIndex < NumTargetLayers; ++TargetLayerIndex)
+	// Because we only expose UTextureRenderTarget2D to BP, we need an additional scratch render target 2D that we will copy the current result of each target layer into, 
+	//  so that the BP can use it as its source. We'd be inclined not to do it for bIsHeightmapMerge, where it is a texture 2D (i.e. exposed to BP), rather than a texture 
+	//  2D array for weightmaps, but we still use a scratch texture for heightmaps in order to ensure the BP's render target size doesn't change too often (more details in the following comments)
+	check((bIsHeightmapMerge && (CurrentLayerReadRT->IsTexture2D() && WriteRT->IsTexture2D()))
+		|| (!bIsHeightmapMerge && (CurrentLayerReadRT->IsTexture2DArray() && WriteRT->IsTexture2DArray())));
+
+	// The ideal size for the render target to use for this brush may differ from the requested render area. The reason is that the brush's Initialize function depends on the size of the 
+	//  render target, so if, from one call to another, this size is changing because there happens to be a scratch render target big enough, we will have to call Initialize again 
+	//  on the brush, which might have performance implications. So using a size that matches what we've already used in the past here allows such brushes not to uselessly go 
+	//  through the initialization process again 
+	const FIntPoint RenderTargetIdealSize = InRenderParams.RenderAreaSectionRect.Size().ComponentMax(CurrentRenderAreaSize);
+	// Use exact dimensions for the scratch texture here for the reason explained above (see RenderTargetIdealSize) :
+	FScratchRenderTargetParams ScratchRenderTargetParams(TEXT("BPBrushScratchRT"), /*bInExactDimensions = */true, /*bInUseUAV = */false, /*bInTargetArraySlicesIndependently = */false,
+		RenderTargetIdealSize, /*InNumSlices = */0, CurrentLayerReadRT->GetFormat(), CurrentLayerReadRT->GetClearColor(), ERHIAccess::CopyDest);
+	FScratchRenderTargetScope ScratchTexture(ScratchRenderTargetParams);
+	CurrentLayerReadRT2D = ScratchTexture.RenderTarget;
+
+	// The original texture array will be accessed as ERHIAccess::CopySrc all along : 
+	CurrentLayerReadRT->TransitionTo(ERHIAccess::CopySrc);
+
+	const int32 NumTargetLayersInRenderGroup = InRenderParams.RenderGroupTargetLayerNames.Num();
+	for (int32 TargetLayerIndexInRenderGroup = 0; TargetLayerIndexInRenderGroup < NumTargetLayersInRenderGroup; ++TargetLayerIndexInRenderGroup)
 	{
-		const FName TargetLayerName = InRenderParams.RenderGroupTargetLayerNames[TargetLayerIndex];
+		const FName TargetLayerName = InRenderParams.RenderGroupTargetLayerNames[TargetLayerIndexInRenderGroup];
 		RHI_BREADCRUMB_EVENT_GAMETHREAD("Render %s", TargetLayerName);
 
-		// If necessary, copy from the texture array's slice to the scratch render target 2D : 
-		if (!bIsHeightmapMerge)
-		{
-			RHI_BREADCRUMB_EVENT_GAMETHREAD("Copy Source (slice %i) -> %s", TargetLayerIndex, CurrentLayerReadRT2D->GetDebugName());
+		// If necessary, copy from the source render target (slice) to the scratch render target 2D : 
+		RHI_BREADCRUMB_EVENT_GAMETHREAD("Copy Source (slice %i) -> %s", TargetLayerIndexInRenderGroup, CurrentLayerReadRT2D->GetDebugName());
 
+		{
 			ULandscapeScratchRenderTarget::FCopyFromScratchRenderTargetParams CopyParams(CurrentLayerReadRT);
-			// It's important to copy only the needed size, because CurrentLayerReadRT2D can come from AdditionalScratchTexture, which is allocated with bInExactDimensions = false)
+			// It's important to copy only the needed size, because CurrentLayerReadRT2D comes from ScratchTexture, whose size may be larger than InRenderParams.RenderAreaSectionRect)
 			CopyParams.CopySize = InRenderParams.RenderAreaSectionRect.Size();
-			// Copy from the proper slice in the texture array : 
-			CopyParams.SourceSliceIndex = TargetLayerIndex;
+			// Copy from the proper slice in the texture array if any : 
+			CopyParams.SourceSliceIndex = TargetLayerIndexInRenderGroup;
 			CurrentLayerReadRT2D->CopyFrom(CopyParams);
-			CurrentLayerReadRT2D->TransitionTo(ERHIAccess::SRVMask);
 		}
+		CurrentLayerReadRT2D->TransitionTo(ERHIAccess::SRVMask);
 
 		check(CurrentLayerReadRT2D->GetCurrentState() == ERHIAccess::SRVMask);
 		check(WriteRT->GetCurrentState() == ERHIAccess::CopyDest);
@@ -372,11 +370,11 @@ void ALandscapeBlueprintBrushBase::RenderLayer(ILandscapeEditLayerRenderer::FRen
 		}
 
 		// TODO: handle conversion/handling of RT not same size as internal size
-		check((OutputRT2D->SizeX == OutputRT2D->SizeX) && (OutputRT2D->SizeY == ReadRT2D->SizeY));
+		check((OutputRT2D->SizeX == ReadRT2D->SizeX) && (OutputRT2D->SizeY == ReadRT2D->SizeY));
 
 		// Resolve back to the write RT 
 		{
-			RHI_BREADCRUMB_EVENT_GAMETHREAD("Copy BP Render Result -> %s (slice %i)", WriteRT->GetDebugName(), TargetLayerIndex);
+			RHI_BREADCRUMB_EVENT_GAMETHREAD("Copy BP Render Result -> %s (slice %i)", WriteRT->GetDebugName(), TargetLayerIndexInRenderGroup);
 
 			// The RT returned by the brush is in SRV state so we need a transition: 
 			ENQUEUE_RENDER_COMMAND(TransitionToCopySrc)([Resource = OutputRT2D->GetResource()](FRHICommandListImmediate& InRHICmdList) mutable
@@ -386,9 +384,9 @@ void ALandscapeBlueprintBrushBase::RenderLayer(ILandscapeEditLayerRenderer::FRen
 
 			ULandscapeScratchRenderTarget::FCopyFromTextureParams CopyParams(OutputRT2D);
 			// It's important to copy only the needed size, because OutputRT2D is currently sized after ReadRT2D, which can be a render target larger than it needs to be (because 
-			//  it can come from AdditionalScratchTexture, which is allocated with bInExactDimensions = false)
+			//  it can come from ScratchTexture, whose size may be larger than InRenderParams.RenderAreaSectionRect)
 			CopyParams.CopySize = InRenderParams.RenderAreaSectionRect.Size();
-			CopyParams.DestSliceIndex = TargetLayerIndex;
+			CopyParams.DestSliceIndex = TargetLayerIndexInRenderGroup;
 			WriteRT->CopyFrom(CopyParams);
 
 			// It's also expected we leave the RT return by the BP as SRV :
@@ -409,7 +407,7 @@ FString ALandscapeBlueprintBrushBase::GetEditLayerRendererDebugName() const
 	return GetActorNameOrLabel();
 }
 
-TArray<UE::Landscape::EditLayers::FEditLayerRendererState> ALandscapeBlueprintBrushBase::GetEditLayerRendererStates(const ULandscapeInfo* InLandscapeInfo, bool bInSkipBrush)
+TArray<UE::Landscape::EditLayers::FEditLayerRendererState> ALandscapeBlueprintBrushBase::GetEditLayerRendererStates(const UE::Landscape::EditLayers::FMergeContext* InMergeContext)
 {
 	using namespace UE::Landscape::EditLayers;
 
@@ -418,9 +416,9 @@ TArray<UE::Landscape::EditLayers::FEditLayerRendererState> ALandscapeBlueprintBr
 		return {};
 	}
 
-	FEditLayerRendererState RendererState(this, InLandscapeInfo);
+	FEditLayerRendererState RendererState(InMergeContext, this);
 	// Force the renderer to be fully disabled in case we are asked to skip the brush :
-	if (bInSkipBrush)
+	if (InMergeContext->ShouldSkipProceduralRenderers())
 	{
 		RendererState.DisableTargetTypeMask(ELandscapeToolTargetTypeFlags::All);
 	}
@@ -496,9 +494,19 @@ bool ALandscapeBlueprintBrushBase::CanAffectWeightmapLayer(const FName& InLayerN
 	return AffectedWeightmapLayers.Contains(InLayerName);
 }
 
+bool ALandscapeBlueprintBrushBase::AffectsHeightmap() const
+{
+	return CanAffectHeightmap() && IsVisible();
+}
+
+bool ALandscapeBlueprintBrushBase::AffectsWeightmap() const
+{
+	return CanAffectWeightmap() && IsVisible();
+}
+
 bool ALandscapeBlueprintBrushBase::AffectsWeightmapLayer(const FName& InLayerName) const
 {
-	if (!CanAffectWeightmap())
+	if (!AffectsWeightmap())
 	{
 		return false;
 	}
@@ -508,6 +516,11 @@ bool ALandscapeBlueprintBrushBase::AffectsWeightmapLayer(const FName& InLayerNam
 	return AffectedWeightmapLayers.Contains(InLayerName);
 }
 
+bool ALandscapeBlueprintBrushBase::AffectsVisibilityLayer() const
+{
+	return CanAffectVisibilityLayer() && IsVisible();
+}
+
 void ALandscapeBlueprintBrushBase::PostEditMove(bool bFinished)
 {
 	Super::PostEditMove(bFinished);
diff --git a/Engine/Source/Runtime/Landscape/Private/LandscapeEditLayerRenderer.cpp b/Engine/Source/Runtime/Landscape/Private/LandscapeEditLayerRenderer.cpp
--- a/Engine/Source/Runtime/Landscape/Private/LandscapeEditLayerRenderer.cpp
+++ b/Engine/Source/Runtime/Landscape/Private/LandscapeEditLayerRenderer.cpp
@@ -113,34 +113,306 @@ FString ConvertTargetLayerNamesToString(const TArrayView<const FName>& InTargetL
 
 #if WITH_EDITOR
 
-FEditLayerTargetTypeState::FEditLayerTargetTypeState(ELandscapeToolTargetTypeFlags InTargetTypeMask, const TArrayView<const FName>& InSupportedWeightmaps)
-	: Weightmaps(InSupportedWeightmaps)
+FMergeContext::FMergeContext(ALandscape* InLandscape, bool bInIsHeightmapMerge, bool bInSkipProceduralRenderers)
+	: bIsHeightmapMerge(bInIsHeightmapMerge)
+	, bSkipProceduralRenderers(bInSkipProceduralRenderers)
+	, Landscape(InLandscape)
+	, LandscapeInfo(InLandscape->GetLandscapeInfo())
 {
+	check(LandscapeInfo != nullptr);
+
+	// Start by gathering all possible target layer names on this landscape: 
+	//  This list of all unique target layer names will help accelerate the gathering of output layers on each component (using bit arrays) as well as the target layers intersection tests:
+	if (bInIsHeightmapMerge)
+	{
+		// Only one target layer in the case of heightmap: 
+		const FName TargetLayerName = FName("Height");
+		AllTargetLayerNames = { TargetLayerName };
+		// And it's valid :
+		ValidTargetLayerBitIndices = TBitArray<>(true, 1);
+		VisibilityTargetLayerMask = TBitArray<>(false, 1);
+		AllWeightmapLayerInfos = { nullptr };
+	}
+	else
+	{
+		// Gather all target layer names and mark those that are valid layers :
+		for (const FLandscapeInfoLayerSettings& LayerSettings : LandscapeInfo->Layers)
+		{
+			check(!LayerSettings.LayerName.IsNone());
+			ULandscapeLayerInfoObject*& LayerInfo = AllWeightmapLayerInfos.Add_GetRef(LayerSettings.LayerInfoObj);
+			AllTargetLayerNames.Add(LayerSettings.LayerName);
+			ValidTargetLayerBitIndices.Add(LayerInfo != nullptr);
+		}
+
+		// Visibility is always a valid layer : 
+		VisibilityTargetLayerIndex = AllTargetLayerNames.Find(UMaterialExpressionLandscapeVisibilityMask::ParameterName);
+		if (VisibilityTargetLayerIndex == INDEX_NONE)
+		{
+			VisibilityTargetLayerIndex = AllTargetLayerNames.Num();
+			AllTargetLayerNames.Add(UMaterialExpressionLandscapeVisibilityMask::ParameterName);
+			check(!AllWeightmapLayerInfos.Contains(ALandscape::VisibilityLayer));
+			AllWeightmapLayerInfos.Add(ALandscape::VisibilityLayer);
+			check(ALandscape::VisibilityLayer != nullptr);
+			ValidTargetLayerBitIndices.Add(true);
+		}
+		check(VisibilityTargetLayerIndex != INDEX_NONE);
+
+		VisibilityTargetLayerMask = TBitArray<>(false, AllTargetLayerNames.Num());
+		VisibilityTargetLayerMask[VisibilityTargetLayerIndex] = true;
+	}
+
+	NegatedVisibilityTargetLayerMask = VisibilityTargetLayerMask;
+	NegatedVisibilityTargetLayerMask.BitwiseNOT();
+}
+
+TArray<FName> FMergeContext::GetValidTargetLayerNames() const
+{
+	return bIsHeightmapMerge ? AllTargetLayerNames : ConvertTargetLayerBitIndicesToNames(ValidTargetLayerBitIndices);
+}
+
+int32 FMergeContext::IsValidTargetLayerName(const FName& InName) const
+{
+	int32 Index = GetTargetLayerIndexForName(InName);
+	return (Index != INDEX_NONE) ? ValidTargetLayerBitIndices[Index] : false;
+}
+
+int32 FMergeContext::IsValidTargetLayerNameChecked(const FName& InName) const
+{
+	int32 Index = AllTargetLayerNames.Find(InName);
+	check(Index != INDEX_NONE);
+	return ValidTargetLayerBitIndices[Index];
+}
+
+int32 FMergeContext::IsValidTargetLayerIndex(int32 InIndex) const
+{
+	check(AllTargetLayerNames.IsValidIndex(InIndex));
+	return ValidTargetLayerBitIndices[InIndex];
+}
+
+int32 FMergeContext::GetTargetLayerIndexForName(const FName& InName) const
+{
+	return AllTargetLayerNames.Find(InName);
+}
+
+int32 FMergeContext::GetTargetLayerIndexForNameChecked(const FName& InName) const
+{
+	int32 Index = AllTargetLayerNames.Find(InName);
+	check(Index != INDEX_NONE);
+	return Index;
+}
+
+FName FMergeContext::GetTargetLayerNameForIndex(int32 InIndex) const
+{
+	return AllTargetLayerNames.IsValidIndex(InIndex) ? AllTargetLayerNames[InIndex] : NAME_None;
+}
+
+FName FMergeContext::GetTargetLayerNameForIndexChecked(int32 InIndex) const
+{
+	check(AllTargetLayerNames.IsValidIndex(InIndex));
+	return AllTargetLayerNames[InIndex];
+}
+
+int32 FMergeContext::GetTargetLayerIndexForLayerInfo(ULandscapeLayerInfoObject* InLayerInfo) const
+{
+	return AllWeightmapLayerInfos.Find(InLayerInfo);
+}
+
+int32 FMergeContext::GetTargetLayerIndexForLayerInfoChecked(ULandscapeLayerInfoObject* InLayerInfo) const
+{
+	int32 Index = AllWeightmapLayerInfos.Find(InLayerInfo);;
+	check(Index != INDEX_NONE);
+	return Index;
+}
+
+ULandscapeLayerInfoObject* FMergeContext::GetTargetLayerInfoForName(const FName& InName) const
+{
+	int32 Index = GetTargetLayerIndexForName(InName);
+	return (Index != INDEX_NONE) ? AllWeightmapLayerInfos[Index] : nullptr;
+}
+
+ULandscapeLayerInfoObject* FMergeContext::GetTargetLayerInfoForNameChecked(const FName& InName) const
+{
+	int32 Index = GetTargetLayerIndexForNameChecked(InName);
+	return AllWeightmapLayerInfos[Index];
+}
+
+ULandscapeLayerInfoObject* FMergeContext::GetTargetLayerInfoForIndex(int32 InIndex) const
+{
+	check(AllTargetLayerNames.IsValidIndex(InIndex));
+	return AllWeightmapLayerInfos[InIndex];
+}
+
+TBitArray<> FMergeContext::ConvertTargetLayerNamesToBitIndices(TConstArrayView<FName> InTargetLayerNames) const
+{
+	TBitArray<> Result(false, AllTargetLayerNames.Num());
+	for (FName Name : InTargetLayerNames)
+	{
+		if (int32 Index = GetTargetLayerIndexForName(Name); Index != INDEX_NONE)
+		{
+			Result[Index] = true;
+		}
+	}
+	return Result;
+}
+
+TBitArray<> FMergeContext::ConvertTargetLayerNamesToBitIndicesChecked(TConstArrayView<FName> InTargetLayerNames) const
+{
+	TBitArray<> Result(false, AllTargetLayerNames.Num());
+	for (FName Name : InTargetLayerNames)
+	{
+		int32 Index = GetTargetLayerIndexForNameChecked(Name);
+		Result[Index] = true;
+	}
+	return Result;
+}
+
+TArray<FName> FMergeContext::ConvertTargetLayerBitIndicesToNames(const TBitArray<>& InTargetLayerBitIndices) const
+{
+	const int32 NumNames = AllTargetLayerNames.Num();
+	check(InTargetLayerBitIndices.Num() == NumNames);
+	TArray<FName> Names;
+	Names.Reserve(NumNames);
+	for (TConstSetBitIterator It(InTargetLayerBitIndices); It; ++It)
+	{
+		Names.Add(AllTargetLayerNames[It.GetIndex()]);
+	}
+	return Names;
+}
+
+TArray<ULandscapeLayerInfoObject*> FMergeContext::ConvertTargetLayerBitIndicesToLayerInfos(const TBitArray<>& InTargetLayerBitIndices) const
+{
+	const int32 NumTargetLayerInfos = AllTargetLayerNames.Num();
+	check(InTargetLayerBitIndices.Num() == NumTargetLayerInfos);
+	TArray<ULandscapeLayerInfoObject*> LayerInfos;
+	LayerInfos.Reserve(NumTargetLayerInfos);
+	for (TConstSetBitIterator It(InTargetLayerBitIndices); It; ++It)
+	{
+		LayerInfos.Add(AllWeightmapLayerInfos[It.GetIndex()]);
+	}
+	return LayerInfos;
+}
+
+void FMergeContext::ForEachTargetLayer(const TBitArray<>& InTargetLayerBitIndices, TFunctionRef<bool(int32 /*InTargetLayerIndex*/, const FName& /*InTargetLayerName*/, ULandscapeLayerInfoObject* /*InWeightmapLayerInfo*/)> Fn) const
+{
+	check(InTargetLayerBitIndices.Num() == AllTargetLayerNames.Num());
+	for (TConstSetBitIterator It(InTargetLayerBitIndices); It; ++It)
+	{
+		const int32 TargetLayerIndex = It.GetIndex();
+		if (!AllTargetLayerNames.IsValidIndex(TargetLayerIndex))
+		{
+			return;
+		}
+
+		if (!Fn(TargetLayerIndex, AllTargetLayerNames[TargetLayerIndex], AllWeightmapLayerInfos[TargetLayerIndex]))
+		{
+			return;
+		}
+	}
+}
+
+void FMergeContext::ForEachValidTargetLayer(TFunctionRef<bool(int32 /*InTargetLayerIndex*/, const FName& /*InTargetLayerName*/, ULandscapeLayerInfoObject* /*InWeightmapLayerInfo*/)> Fn) const
+{
+	return ForEachTargetLayer(ValidTargetLayerBitIndices, Fn);
+}
+
+TBitArray<> FMergeContext::BuildTargetLayerBitIndices(bool bInBitValue) const
+{
+	return TBitArray<>(bInBitValue, AllTargetLayerNames.Num());
+}
+
+
+// ----------------------------------------------------------------------------------
+
+FEditLayerTargetTypeState::FEditLayerTargetTypeState(const FMergeContext* InMergeContext)
+	: FEditLayerTargetTypeState(InMergeContext, ELandscapeToolTargetTypeFlags::None)
+{}
+
+
+FEditLayerTargetTypeState::FEditLayerTargetTypeState(const FMergeContext* InMergeContext, ELandscapeToolTargetTypeFlags InTargetTypeMask)
+	: FEditLayerTargetTypeState(InMergeContext, InTargetTypeMask, TBitArray<>())
+{}
+
+FEditLayerTargetTypeState::FEditLayerTargetTypeState(const FMergeContext* InMergeContext, ELandscapeToolTargetTypeFlags InTargetTypeMask, const TArrayView<const FName>& InSupportedWeightmaps, bool bInChecked)
+	: FEditLayerTargetTypeState(InMergeContext, InTargetTypeMask, 
+		bInChecked ? InMergeContext->ConvertTargetLayerNamesToBitIndicesChecked(InSupportedWeightmaps) : InMergeContext->ConvertTargetLayerNamesToBitIndices(InSupportedWeightmaps))
+{}
+
+FEditLayerTargetTypeState::FEditLayerTargetTypeState(const FMergeContext* InMergeContext, ELandscapeToolTargetTypeFlags InTargetTypeMask, const TBitArray<>& InSupportedWeightmapLayerIndices)
+	: MergeContext(InMergeContext)
+{
+	if (InSupportedWeightmapLayerIndices.IsEmpty())
+	{
+		// Even if all weightmaps are turned off in this constructor, make sure to build a bit array that is dealing with as many layers as the merge context : 
+		WeightmapTargetLayerBitIndices = MergeContext->BuildTargetLayerBitIndices(false);
+	}
+	else
+	{
+		checkf(InSupportedWeightmapLayerIndices.Num() == MergeContext->GetAllTargetLayerNames().Num(), 
+			TEXT("Make sure that the target type state is dealing with the same amount of target layers as the merge context"));
+		WeightmapTargetLayerBitIndices = InSupportedWeightmapLayerIndices;
+	}
 	SetTargetTypeMask(InTargetTypeMask);
 }
 
-bool FEditLayerTargetTypeState::IsActive(ELandscapeToolTargetType InTargetType, FName InWeightmapLayerName /*= NAME_None*/) const
+bool FEditLayerTargetTypeState::IsActive(ELandscapeToolTargetType InTargetType, const FName& InWeightmapLayerName) const
+{
+	if (InWeightmapLayerName != NAME_None)
+	{
+		if (int32 TargetLayerIndex = MergeContext->GetTargetLayerIndexForName(InWeightmapLayerName); TargetLayerIndex != INDEX_NONE)
+		{
+			return IsActive(InTargetType, TargetLayerIndex);
+		}
+
+		return false;
+	}
+
+	return IsActive(InTargetType, INDEX_NONE);
+}
+
+bool FEditLayerTargetTypeState::IsActiveChecked(ELandscapeToolTargetType InTargetType, const FName& InWeightmapLayerName) const
+{
+	if (InWeightmapLayerName != NAME_None)
+	{
+		return IsActive(InTargetType, MergeContext->GetTargetLayerIndexForNameChecked(InWeightmapLayerName));
+	}
+
+	return IsActive(InTargetType, INDEX_NONE);
+}
+
+bool FEditLayerTargetTypeState::IsActive(ELandscapeToolTargetType InTargetType, int32 InWeightmapLayerIndex) const
 {
+	check((InWeightmapLayerIndex == INDEX_NONE) || MergeContext->IsValidTargetLayerIndex(InWeightmapLayerIndex));
+
 	if (EnumHasAnyFlags(TargetTypeMask, GetLandscapeToolTargetTypeAsFlags(InTargetType)))
 	{
 		if (InTargetType != ELandscapeToolTargetType::Heightmap)
 		{
-			return Weightmaps.Contains(InWeightmapLayerName);
+			return GetActiveWeightmapBitIndices()[InWeightmapLayerIndex];
 		}
+
 		return true;
 	}
 
 	return false;
 }
 
-const TArray<FName>& FEditLayerTargetTypeState::GetActiveWeightmaps() const
+TArray<FName> FEditLayerTargetTypeState::GetActiveWeightmaps() const
 {
-	if (EnumHasAnyFlags(TargetTypeMask, ELandscapeToolTargetTypeFlags::Weightmap | ELandscapeToolTargetTypeFlags::Visibility))
+	return MergeContext->ConvertTargetLayerBitIndicesToNames(GetActiveWeightmapBitIndices());
+}
+
+TBitArray<> FEditLayerTargetTypeState::GetActiveWeightmapBitIndices() const
+{
+	TBitArray<> Result = WeightmapTargetLayerBitIndices;
+	if (!EnumHasAnyFlags(TargetTypeMask, ELandscapeToolTargetTypeFlags::Visibility))
+	{
+		Result.CombineWithBitwiseAND(MergeContext->GetNegatedVisibilityTargetLayerMask(), EBitwiseOperatorFlags::MinSize);
+	}
+	if (!EnumHasAnyFlags(TargetTypeMask, ELandscapeToolTargetTypeFlags::Weightmap))
 	{
-		return Weightmaps;
+		Result.CombineWithBitwiseAND(MergeContext->GetVisibilityTargetLayerMask(), EBitwiseOperatorFlags::MinSize);
 	}
-	static TArray<FName> EmptyWeightmaps;
-	return EmptyWeightmaps;
+	return Result;
 }
 
 void FEditLayerTargetTypeState::SetTargetTypeMask(ELandscapeToolTargetTypeFlags InTargetTypeMask)
@@ -150,13 +422,17 @@ void FEditLayerTargetTypeState::SetTargetTypeMask(ELandscapeToolTargetTypeFlags
 		TargetTypeMask = InTargetTypeMask;
 
 		// Special case for the visibility weightmap, where we want to make sure the weightmap layer name is specified if visibility is supported (and vice versa) : 
-		if (EnumHasAnyFlags(TargetTypeMask, ELandscapeToolTargetTypeFlags::Visibility))
-		{
-			AddWeightmap(UMaterialExpressionLandscapeVisibilityMask::ParameterName);
-		}
-		else if (Weightmaps.Contains(UMaterialExpressionLandscapeVisibilityMask::ParameterName))
+		const int32 VisibilityTargetLayerIndex = MergeContext->GetVisibilityTargetLayerIndex();
+		if (VisibilityTargetLayerIndex != INDEX_NONE)
 		{
-			RemoveWeightmap(UMaterialExpressionLandscapeVisibilityMask::ParameterName);
+			if (EnumHasAnyFlags(TargetTypeMask, ELandscapeToolTargetTypeFlags::Visibility))
+			{
+				AddWeightmap(VisibilityTargetLayerIndex);
+			}
+			else if (WeightmapTargetLayerBitIndices[VisibilityTargetLayerIndex])
+			{
+				RemoveWeightmap(VisibilityTargetLayerIndex);
+			}
 		}
 	}
 }
@@ -183,37 +459,64 @@ void FEditLayerTargetTypeState::RemoveTargetTypeMask(ELandscapeToolTargetTypeFla
 
 FEditLayerTargetTypeState FEditLayerTargetTypeState::Intersect(const FEditLayerTargetTypeState& InOther) const
 {
-	TSet<FName> OtherWeightmaps(InOther.GetActiveWeightmaps());
-	return FEditLayerTargetTypeState(InOther.GetTargetTypeMask() & TargetTypeMask, OtherWeightmaps.Intersect(TSet<FName>(Weightmaps)).Array());
+	checkf(WeightmapTargetLayerBitIndices.Num() == InOther.WeightmapTargetLayerBitIndices.Num(), 
+		TEXT("It is assumed that the 2 target type states to intersect are from the same context and are therefore dealing with the same amount of target layers"));
+	return FEditLayerTargetTypeState(MergeContext, InOther.GetTargetTypeMask() & TargetTypeMask, TBitArray<>::BitwiseAND(InOther.WeightmapTargetLayerBitIndices, WeightmapTargetLayerBitIndices, EBitwiseOperatorFlags::MinSize));
 }
 
-void FEditLayerTargetTypeState::AddWeightmap(FName InWeightmapLayerName)
+void FEditLayerTargetTypeState::AddWeightmap(const FName& InWeightmapLayerName)
 {
-	checkf(EnumHasAnyFlags(TargetTypeMask, ELandscapeToolTargetTypeFlags::Weightmap | ELandscapeToolTargetTypeFlags::Visibility), 
-		TEXT("Cannot add weightmap %s to a target type state that doesn't support weightmaps"), *InWeightmapLayerName.ToString());
+	if (int32 TargetLayerIndex = MergeContext->GetTargetLayerIndexForName(InWeightmapLayerName); TargetLayerIndex != INDEX_NONE)
+	{
+		AddWeightmap(TargetLayerIndex);
+	}
+}
 
-	checkf(!EnumHasAnyFlags(TargetTypeMask, ELandscapeToolTargetTypeFlags::Visibility) || (InWeightmapLayerName == UMaterialExpressionLandscapeVisibilityMask::ParameterName) || Weightmaps.Contains(UMaterialExpressionLandscapeVisibilityMask::ParameterName),
-		TEXT("Visibility should always come with a weightmap named %s"), *UMaterialExpressionLandscapeVisibilityMask::ParameterName.ToString());
+void FEditLayerTargetTypeState::AddWeightmapChecked(const FName& InWeightmapLayerName)
+{
+	AddWeightmap(MergeContext->GetTargetLayerIndexForNameChecked(InWeightmapLayerName));
+}
 
-	Weightmaps.AddUnique(InWeightmapLayerName);
+void FEditLayerTargetTypeState::AddWeightmap(int32 InWeightmapLayerIndex)
+{
+	check(MergeContext->IsValidTargetLayerIndex(InWeightmapLayerIndex));
+
+	checkf((InWeightmapLayerIndex != MergeContext->GetVisibilityTargetLayerIndex()) || EnumHasAnyFlags(TargetTypeMask, ELandscapeToolTargetTypeFlags::Visibility), 
+		TEXT("The visibility layer may only be used for target type states that support visibility"));
+
+	checkf((InWeightmapLayerIndex == MergeContext->GetVisibilityTargetLayerIndex()) || EnumHasAnyFlags(TargetTypeMask, ELandscapeToolTargetTypeFlags::Weightmap), 
+		TEXT("A weightmap layer (here, %s) may only be used for target type states that support weightmaps"), *MergeContext->GetTargetLayerNameForIndexChecked(InWeightmapLayerIndex).ToString());
+
+	WeightmapTargetLayerBitIndices[InWeightmapLayerIndex] = true;
+}
+
+void FEditLayerTargetTypeState::RemoveWeightmap(const FName& InWeightmapLayerName)
+{
+	if (int32 TargetLayerIndex = MergeContext->GetTargetLayerIndexForName(InWeightmapLayerName); TargetLayerIndex != INDEX_NONE)
+	{
+		RemoveWeightmap(TargetLayerIndex);
+	}
+}
+
+void FEditLayerTargetTypeState::RemoveWeightmapChecked(const FName& InWeightmapLayerName)
+{
+	RemoveWeightmap(MergeContext->GetTargetLayerIndexForNameChecked(InWeightmapLayerName));
 }
 
-void FEditLayerTargetTypeState::RemoveWeightmap(FName InWeightmapLayerName)
+void FEditLayerTargetTypeState::RemoveWeightmap(int32 InWeightmapLayerIndex)
 {
-	checkf(EnumHasAnyFlags(TargetTypeMask, ELandscapeToolTargetTypeFlags::Weightmap | ELandscapeToolTargetTypeFlags::Visibility), 
-		TEXT("Cannot remove weightmap %s from a target type state that doesn't support weightmaps"), *InWeightmapLayerName.ToString());
+	check(MergeContext->IsValidTargetLayerIndex(InWeightmapLayerIndex));
 
-	checkf(!EnumHasAnyFlags(TargetTypeMask, ELandscapeToolTargetTypeFlags::Visibility) || (InWeightmapLayerName != UMaterialExpressionLandscapeVisibilityMask::ParameterName), 
-		TEXT("Cannot remove weightmap %s from a target type state that supports visibility"), *UMaterialExpressionLandscapeVisibilityMask::ParameterName.ToString());
+	checkf(!EnumHasAnyFlags(TargetTypeMask, ELandscapeToolTargetTypeFlags::Visibility) || (InWeightmapLayerIndex != MergeContext->GetVisibilityTargetLayerIndex()),
+		TEXT("Cannot remove weightmap %s from a target type state that supports visibility"), *MergeContext->GetTargetLayerNameForIndexChecked(InWeightmapLayerIndex).ToString());
 
-	Weightmaps.Remove(InWeightmapLayerName);
+	WeightmapTargetLayerBitIndices[InWeightmapLayerIndex] = false;
 }
 
 bool FEditLayerTargetTypeState::operator==(const FEditLayerTargetTypeState& InOther) const
 {
 	return (TargetTypeMask == InOther.TargetTypeMask) 
-		// TODO [jonathan.bard] : This is really bad for performance if called several times. This test can be replaced by a hash test of all ordered weightmaps: 
-		&& (TSet<FName>(Weightmaps).Intersect(TSet<FName>(InOther.Weightmaps)).Num() == Weightmaps.Num());
+		&& (WeightmapTargetLayerBitIndices == InOther.WeightmapTargetLayerBitIndices);
 }
 
 FString FEditLayerTargetTypeState::ToString() const
@@ -221,45 +524,51 @@ FString FEditLayerTargetTypeState::ToString() const
 	FString Result = FString::Printf(TEXT("Target types: %s"), *UE::Landscape::GetLandscapeToolTargetTypeFlagsAsString(TargetTypeMask));
 	if (EnumHasAnyFlags(TargetTypeMask, ELandscapeToolTargetTypeFlags::Weightmap | ELandscapeToolTargetTypeFlags::Visibility))
 	{
-		Result += FString::Printf(TEXT("\nWeightmaps: %s"), *ConvertTargetLayerNamesToString(Weightmaps));
+		Result += FString::Printf(TEXT("\nWeightmaps: %s"), *ConvertTargetLayerNamesToString(MergeContext->ConvertTargetLayerBitIndicesToNames(WeightmapTargetLayerBitIndices)));
 	}
 	return Result;
 }
 
+
 // ----------------------------------------------------------------------------------
 
-FEditLayerRendererState::FEditLayerRendererState(ILandscapeEditLayerRenderer* InRenderer, const ULandscapeInfo* InLandscapeInfo)
-	: Renderer(InRenderer)
+FEditLayerRendererState::FEditLayerRendererState(const FMergeContext* InMergeContext)
+	: MergeContext(InMergeContext)
+	, SupportedTargetTypeState(InMergeContext)
+	, EnabledTargetTypeState(InMergeContext)
+	, ActiveTargetTypeState(InMergeContext)
 {
-	Renderer->GetRendererStateInfo(InLandscapeInfo, SupportedTargetTypeState, EnabledTargetTypeState, RenderGroups);
-
-	// Make sure that each supported weightmap belongs to one render group and one only. For those that are in no render group, put them in their own group, that simply means this renderer 
-	//  can render them without requesting the presence of other target layers (e.g. no weight-blending)
-	for (FName TargetLayerName : SupportedTargetTypeState.GetActiveWeightmaps())
-	{
-		const int32 RenderGroupCount = Algo::CountIf(RenderGroups, [TargetLayerName](const TSet<FName>& InRenderGroup) { return InRenderGroup.Contains(TargetLayerName); });
-		checkf(RenderGroupCount < 2, TEXT("Target layer %s belongs to more than 1 render group in edit layer renderer %s. This is forbidden: in the end, it must belong to 1 and 1 only."),
-			*TargetLayerName.ToString(), *InRenderer->GetEditLayerRendererDebugName());
-		if (RenderGroupCount == 0)
-		{
-			RenderGroups.Add({ TargetLayerName });
-		}
-	}
+	UpdateActiveTargetTypeState();
 }
 
-bool FEditLayerRendererState::IsTargetSupported(ELandscapeToolTargetType InTargetType, FName InWeightmapLayerName) const
+FEditLayerRendererState::FEditLayerRendererState(const FMergeContext* InMergeContext, ILandscapeEditLayerRenderer* InRenderer)
+	: MergeContext(InMergeContext)
+	, Renderer(InRenderer)
+	, DebugName(InRenderer->GetEditLayerRendererDebugName())
+	, SupportedTargetTypeState(InMergeContext)
+	, EnabledTargetTypeState(InMergeContext)
+	, ActiveTargetTypeState(InMergeContext)
 {
-	return SupportedTargetTypeState.IsActive(InTargetType, InWeightmapLayerName);
-}
+	Renderer->GetRendererStateInfo(InMergeContext, SupportedTargetTypeState, EnabledTargetTypeState, RenderGroups);
 
-const TArray<FName>& FEditLayerRendererState::GetSupportedTargetWeightmaps() const
-{
-	return SupportedTargetTypeState.GetActiveWeightmaps();
-}
+	// Make sure that each supported weightmap belongs to one render group and one only. For those that are in no render group, put them in their own group, that simply means this renderer 
+	//  can render them without requesting the presence of other target layers (e.g. no weight-blending)
+	InMergeContext->ForEachTargetLayer(SupportedTargetTypeState.GetActiveWeightmapBitIndices(), 
+		[InMergeContext, InRenderer, this](int32 InTargetLayerIndex, const FName& InTargetLayerName, ULandscapeLayerInfoObject* InWeightmapLayerInfo)
+		{
+			const int32 RenderGroupCount = Algo::CountIf(RenderGroups, [InTargetLayerIndex](const TBitArray<>& InRenderGroup) { return InRenderGroup[InTargetLayerIndex]; });
+			checkf(RenderGroupCount < 2, TEXT("Target layer %s belongs to more than 1 render group in edit layer renderer %s. This is forbidden: in the end, it must belong to 1 and 1 only."),
+				*InTargetLayerName.ToString(), *InRenderer->GetEditLayerRendererDebugName());
+			if (RenderGroupCount == 0)
+			{
+				TBitArray<> RenderGroup = InMergeContext->BuildTargetLayerBitIndices(false);
+				RenderGroup[InTargetLayerIndex] = true;
+				RenderGroups.Add(RenderGroup);
+			}
+			return true;
+		});
 
-bool FEditLayerRendererState::IsTargetEnabled(ELandscapeToolTargetType InTargetType, FName InWeightmapLayerName) const
-{
-	return SupportedTargetTypeState.IsActive(InTargetType, InWeightmapLayerName) && EnabledTargetTypeState.IsActive(InTargetType, InWeightmapLayerName);
+	UpdateActiveTargetTypeState();
 }
 
 void FEditLayerRendererState::EnableTargetType(ELandscapeToolTargetType InTargetType)
@@ -267,6 +576,7 @@ void FEditLayerRendererState::EnableTargetType(ELandscapeToolTargetType InTarget
 	checkf(!EnumHasAnyFlags(SupportedTargetTypeState.GetTargetTypeMask(), GetLandscapeToolTargetTypeAsFlags(InTargetType)), 
 		TEXT("Target type %s cannot be enabled on this renderer state because it is not supported. Make sure that target types are supported before enabling them"), *UEnum::GetValueAsString(InTargetType));
 	EnabledTargetTypeState.AddTargetType(InTargetType);
+	UpdateActiveTargetTypeState();
 }
 
 void FEditLayerRendererState::EnableTargetTypeMask(ELandscapeToolTargetTypeFlags InTargetTypeMask)
@@ -280,34 +590,90 @@ void FEditLayerRendererState::EnableTargetTypeMask(ELandscapeToolTargetTypeFlags
 void FEditLayerRendererState::DisableTargetType(ELandscapeToolTargetType InTargetType)
 {
 	EnabledTargetTypeState.RemoveTargetType(InTargetType);
+	UpdateActiveTargetTypeState();
 }
 
 void FEditLayerRendererState::DisableTargetTypeMask(ELandscapeToolTargetTypeFlags InTargetTypeMask)
 {
 	EnabledTargetTypeState.RemoveTargetTypeMask(InTargetTypeMask);
+	UpdateActiveTargetTypeState();
 }
 
-bool FEditLayerRendererState::EnableTarget(ELandscapeToolTargetType InTargetType, FName InWeightmapLayerName)
+void FEditLayerRendererState::EnableWeightmap(ELandscapeToolTargetType InTargetType, const FName& InWeightmapLayerName)
 {
 	EnableTargetType(InTargetType);
-	EnabledTargetTypeState.AddWeightmap(InWeightmapLayerName);
-	
-	// The target has to be both supported and enabled to be considered fully enabled : 
-	return IsTargetEnabled(InTargetType, InWeightmapLayerName);
+	if (int32 TargetLayerIndex = MergeContext->GetTargetLayerIndexForName(InWeightmapLayerName); TargetLayerIndex != INDEX_NONE)
+	{
+		EnabledTargetTypeState.AddWeightmap(TargetLayerIndex);
+		UpdateActiveTargetTypeState();
+	}
+}
+
+void FEditLayerRendererState::EnableWeightmapChecked(ELandscapeToolTargetType InTargetType, const FName& InWeightmapLayerName)
+{
+	EnableWeightmap(InTargetType, (InWeightmapLayerName != NAME_None) ? MergeContext->GetTargetLayerIndexForNameChecked(InWeightmapLayerName) : INDEX_NONE);
+}
+
+void FEditLayerRendererState::EnableWeightmap(ELandscapeToolTargetType InTargetType, int32 InWeightmapLayerIndex)
+{
+	EnableTargetType(InTargetType);
+	EnabledTargetTypeState.AddWeightmap(InWeightmapLayerIndex);
+	UpdateActiveTargetTypeState();
+}
+
+void FEditLayerRendererState::DisableWeightmap(const FName& InWeightmapLayerName)
+{
+	if (int32 TargetLayerIndex = MergeContext->GetTargetLayerIndexForName(InWeightmapLayerName); TargetLayerIndex != INDEX_NONE)
+	{
+		DisableWeightmap(TargetLayerIndex);
+	}
+}
+
+void FEditLayerRendererState::DisableWeightmapChecked(const FName& InWeightmapLayerName)
+{
+	DisableWeightmap(MergeContext->GetTargetLayerIndexForNameChecked(InWeightmapLayerName));
+}
+
+void FEditLayerRendererState::DisableWeightmap(int32 InWeightmapLayerIndex)
+{
+	EnabledTargetTypeState.RemoveWeightmap(InWeightmapLayerIndex);
+	UpdateActiveTargetTypeState();
+}
+
+ELandscapeToolTargetTypeFlags FEditLayerRendererState::GetActiveTargetTypeMask() const
+{
+	return ActiveTargetTypeState.GetTargetTypeMask();
+}
+
+bool FEditLayerRendererState::IsTargetActive(ELandscapeToolTargetType InTargetType, const FName& InWeightmapLayerName) const
+{
+	return ActiveTargetTypeState.IsActive(InTargetType, InWeightmapLayerName);
 }
 
-void FEditLayerRendererState::DisableTarget(FName InWeightmapLayerName)
+bool FEditLayerRendererState::IsTargetActiveChecked(ELandscapeToolTargetType InTargetType, const FName& InWeightmapLayerName) const
 {
-	EnabledTargetTypeState.RemoveWeightmap(InWeightmapLayerName);
+	return ActiveTargetTypeState.IsActiveChecked(InTargetType, InWeightmapLayerName);
 }
 
-TArray<FName> FEditLayerRendererState::GetEnabledTargetWeightmaps() const
+bool FEditLayerRendererState::IsTargetActive(ELandscapeToolTargetType InTargetType, int32 InWeightmapLayerIndex) const
 {
-	// Find the weightmaps that are both supported and enabled : 
-	FEditLayerTargetTypeState SupportedAndEnabledState = SupportedTargetTypeState.Intersect(EnabledTargetTypeState);
-	return SupportedAndEnabledState.GetActiveWeightmaps();
+	return ActiveTargetTypeState.IsActive(InTargetType, InWeightmapLayerIndex);
 }
 
+TArray<FName> FEditLayerRendererState::GetActiveTargetWeightmaps() const
+{
+	return ActiveTargetTypeState.GetActiveWeightmaps();
+}
+
+TBitArray<> FEditLayerRendererState::GetActiveTargetWeightmapBitIndices() const
+{
+	return ActiveTargetTypeState.GetActiveWeightmapBitIndices();
+}
+
+void FEditLayerRendererState::UpdateActiveTargetTypeState()
+{
+	ActiveTargetTypeState = SupportedTargetTypeState.Intersect(EnabledTargetTypeState);
+}
 
 // ----------------------------------------------------------------------------------
 
@@ -427,9 +793,8 @@ FIntPoint FMergeRenderBatch::GetRenderTargetResolution(bool bInWithDuplicateBord
 
 // ----------------------------------------------------------------------------------
 
-FMergeRenderContext::FMergeRenderContext(ALandscape* InLandscape, bool bInIsHeightmapMerge)
-	: Landscape(InLandscape)
-	, bIsHeightmapMerge(bInIsHeightmapMerge)
+FMergeRenderContext::FMergeRenderContext(const FMergeContext& InMergeContext)
+	: FMergeContext(InMergeContext)
 {
 	for (ULandscapeScratchRenderTarget*& BlendRenderTarget : BlendRenderTargets)
 	{
@@ -521,8 +886,8 @@ void FMergeRenderContext::AllocateBatchResources(const FMergeRenderBatch& InRend
 
 	// We need a RT version of the stencil buffer, one per target layer, to let users sample it as a standard texture :
 	int32 VisibilityScratchRTIndex = 0;
-	ForEachTargetLayer(InRenderBatch.TargetLayerNameBitIndices, 
-		[this, LandscapeEditResourcesSubsystem, &VisibilityScratchRTIndex](int32 InTargetLayerIndex, FName InTargetLayerName)
+	ForEachTargetLayer(InRenderBatch.TargetLayerBitIndices,
+		[this, LandscapeEditResourcesSubsystem, &VisibilityScratchRTIndex](int32 InTargetLayerIndex, const FName& InTargetLayerName, ULandscapeLayerInfoObject* InWeightmapLayerInfo)
 		{
 			FScratchRenderTargetParams ScratchRenderTargetParams(FString::Printf(TEXT("VisibilityScratchRT(%i)"), VisibilityScratchRTIndex), /*bInExactDimensions = */false, /*bInUseUAV = */false, /*bInTargetArraySlicesIndependently = */false,
 				MaxNeededResolution, 0, ETextureRenderTargetFormat::RTF_R8, FLinearColor::Black, ERHIAccess::RTV);
@@ -629,8 +994,9 @@ void FMergeRenderContext::RenderValidityRenderTargets(const FMergeRenderBatch& I
 	};
 
 	TArray<FTextureAndRects> TexturesAndRects;
-	TexturesAndRects.Reserve(InRenderBatch.TargetLayerNameBitIndices.CountSetBits());
-	ForEachTargetLayer(InRenderBatch.TargetLayerNameBitIndices, [this, &TexturesAndRects, &InRenderBatch](int32 InTargetLayerIndex, FName InTargetLayerName)
+	TexturesAndRects.Reserve(InRenderBatch.TargetLayerBitIndices.CountSetBits());
+	ForEachTargetLayer(InRenderBatch.TargetLayerBitIndices,
+		[this, &TexturesAndRects, &InRenderBatch](int32 InTargetLayerIndex, const FName& InTargetLayerName, ULandscapeLayerInfoObject* InWeightmapLayerInfo)
 		{
 			ULandscapeScratchRenderTarget* ScratchRenderTarget = PerTargetLayerValidityRenderTargets.FindChecked(InTargetLayerName);
 			check(ScratchRenderTarget != nullptr);
@@ -816,8 +1182,8 @@ void FMergeRenderContext::Render(TFunction<void(const FOnRenderBatchTargetGroupD
 		for (int32 RenderStepIndex = 0; RenderStepIndex < NumRenderSteps; ++RenderStepIndex)
 		{
 			const FMergeRenderStep& RenderStep = RenderBatch.RenderSteps[RenderStepIndex];
-			TArray<FName> RenderGroupTargetLayerNames = ConvertTargetLayerBitIndicesToNamesChecked(RenderStep.RenderGroupBitIndices);
-			TArray<ULandscapeLayerInfoObject*> RenderGroupTargetLayerInfos = bIsHeightmapMerge ? TArray<ULandscapeLayerInfoObject*> { nullptr } : ConvertTargetLayerBitIndicesToLayerInfosChecked(RenderStep.RenderGroupBitIndices);
+			TArray<FName> RenderGroupTargetLayerNames = ConvertTargetLayerBitIndicesToNames(RenderStep.RenderGroupBitIndices);
+			TArray<ULandscapeLayerInfoObject*> RenderGroupTargetLayerInfos = bIsHeightmapMerge ? TArray<ULandscapeLayerInfoObject*> { nullptr } : ConvertTargetLayerBitIndicesToLayerInfos(RenderStep.RenderGroupBitIndices);
 
 			// Compute all necessary info about the components affected by this renderer at this step
 			TArray<FComponentMergeRenderInfo> SortedComponentMergeRenderInfos;
@@ -907,140 +1273,12 @@ const FMergeRenderBatch* FMergeRenderContext::GetCurrentRenderBatch() const
 	return RenderBatches.IsValidIndex(CurrentRenderBatchIndex) ? &RenderBatches[CurrentRenderBatchIndex] : nullptr;
 }
 
-TBitArray<> FMergeRenderContext::ConvertTargetLayerNamesToBitIndices(TConstArrayView<FName> InTargetLayerNames) const
-{
-	TBitArray<> Result(false, AllTargetLayerNames.Num());
-	for (FName Name : InTargetLayerNames)
-	{
-		if (int32 Index = GetTargetLayerIndexForName(Name); Index != INDEX_NONE)
-		{
-			Result[Index] = true;
-		}
-	}
-	return Result;
-}
-
-TBitArray<> FMergeRenderContext::ConvertTargetLayerNamesToBitIndicesChecked(TConstArrayView<FName> InTargetLayerNames) const
-{
-	TBitArray<> Result(false, AllTargetLayerNames.Num());
-	for (FName Name : InTargetLayerNames)
-	{
-		int32 Index = GetTargetLayerIndexForNameChecked(Name);
-		Result[Index] = true;
-	}
-	return Result;
-}
-
-TArray<FName> FMergeRenderContext::ConvertTargetLayerBitIndicesToNames(const TBitArray<>& InTargetLayerBitIndices) const
-{
-	const int32 NumNames = AllTargetLayerNames.Num();
-	TArray<FName> Names;
-	Names.Reserve(NumNames);
-	for (TConstSetBitIterator It(InTargetLayerBitIndices); It && (It.GetIndex() < NumNames); ++It)
-	{
-		Names.Add(AllTargetLayerNames[It.GetIndex()]);
-	}
-	return Names;
-}
-
-TArray<FName> FMergeRenderContext::ConvertTargetLayerBitIndicesToNamesChecked(const TBitArray<>& InTargetLayerBitIndices) const
-{
-	const int32 NumNames = AllTargetLayerNames.Num();
-	check(InTargetLayerBitIndices.Num() == NumNames);
-	TArray<FName> Names;
-	Names.Reserve(NumNames);
-	for (TConstSetBitIterator It(InTargetLayerBitIndices); It; ++It)
-	{
-		Names.Add(AllTargetLayerNames[It.GetIndex()]);
-	}
-	return Names;
-}
 
 bool FMergeRenderContext::IsValid() const
 {
 	return !RenderBatches.IsEmpty();
 }
 
-int32 FMergeRenderContext::GetTargetLayerIndexForName(const FName& InName) const
-{
-	return AllTargetLayerNames.Find(InName);
-}
-
-int32 FMergeRenderContext::GetTargetLayerIndexForNameChecked(const FName& InName) const
-{
-	int32 Index = AllTargetLayerNames.Find(InName); 
-	check(Index != INDEX_NONE); 
-	return Index;
-}
-
-FName FMergeRenderContext::GetTargetLayerNameForIndex(int32 InIndex) const
-{
-	return AllTargetLayerNames.IsValidIndex(InIndex) ? AllTargetLayerNames[InIndex] : NAME_None;
-}
-
-FName FMergeRenderContext::GetTargetLayerNameForIndexChecked(int32 InIndex) const
-{
-	check(AllTargetLayerNames.IsValidIndex(InIndex)); 
-	return AllTargetLayerNames[InIndex];
-}
-
-TArray<ULandscapeLayerInfoObject*> FMergeRenderContext::ConvertTargetLayerBitIndicesToLayerInfos(const TBitArray<>& InTargetLayerBitIndices) const
-{
-	const int32 NumTargetLayerInfos = AllTargetLayerNames.Num();
-	TArray<ULandscapeLayerInfoObject*> LayerInfos;
-	LayerInfos.Reserve(NumTargetLayerInfos);
-	for (TConstSetBitIterator It(InTargetLayerBitIndices); It && (It.GetIndex() < NumTargetLayerInfos); ++It)
-	{
-		LayerInfos.Add(WeightmapLayerInfos[It.GetIndex()]);
-	}
-	return LayerInfos;
-}
-
-TArray<ULandscapeLayerInfoObject*> FMergeRenderContext::ConvertTargetLayerBitIndicesToLayerInfosChecked(const TBitArray<>& InTargetLayerBitIndices) const
-{
-	const int32 NumTargetLayerInfos = AllTargetLayerNames.Num();
-	check(InTargetLayerBitIndices.Num() == NumTargetLayerInfos);
-	TArray<ULandscapeLayerInfoObject*> LayerInfos;
-	LayerInfos.Reserve(NumTargetLayerInfos);
-	for (TConstSetBitIterator It(InTargetLayerBitIndices); It; ++It)
-	{
-		LayerInfos.Add(WeightmapLayerInfos[It.GetIndex()]);
-	}
-	return LayerInfos;
-}
-
-void FMergeRenderContext::ForEachTargetLayer(const TBitArray<>& InTargetLayerBitIndices, TFunctionRef<bool(int32 /*InTargetLayerIndex*/, FName /*InTargetLayerName*/)> Fn) const
-{
-	for (TConstSetBitIterator It(InTargetLayerBitIndices); It; ++It)
-	{
-		const int32 TargetLayerIndex = It.GetIndex();
-		if (!AllTargetLayerNames.IsValidIndex(TargetLayerIndex))
-		{
-			return;
-		}
-
-		if (!Fn(TargetLayerIndex, AllTargetLayerNames[TargetLayerIndex]))
-		{
-			return;
-		}
-	}
-}
-
-void FMergeRenderContext::ForEachTargetLayerChecked(const TBitArray<>& InTargetLayerBitIndices, TFunctionRef<bool(int32 /*InTargetLayerIndex*/, FName /*InTargetLayerName*/)> Fn) const
-{
-	const int32 NumNames = AllTargetLayerNames.Num();
-	check(InTargetLayerBitIndices.Num() == NumNames);
-
-	for (TConstSetBitIterator It(InTargetLayerBitIndices); It; ++It)
-	{
-		const int32 TargetLayerIndex = It.GetIndex();
-		if (!Fn(TargetLayerIndex, AllTargetLayerNames[TargetLayerIndex]))
-		{
-			return;
-		}
-	}
-}
-
 #if ENABLE_VISUAL_LOG
 
 int32 FMergeRenderContext::GetVisualLogAlpha()
@@ -1243,18 +1481,11 @@ bool FComponentMergeRenderInfo::operator<(const FComponentMergeRenderInfo& InOth
 
 // ----------------------------------------------------------------------------------
 
-FMergeRenderParams::FMergeRenderParams(bool bInIsHeightmapMerge, ALandscape* InLandscape, TArray<ULandscapeComponent*> InComponentsToMerge, const TArrayView<FEditLayerRendererState>& InEditLayerRendererStates, const TSet<FName>& InWeightmapLayerNames)
-	: bIsHeightmapMerge(bInIsHeightmapMerge)
-	, Landscape(InLandscape)
-	, ComponentsToMerge(InComponentsToMerge)
+FMergeRenderParams::FMergeRenderParams(TArray<ULandscapeComponent*> InComponentsToMerge, const TArrayView<FEditLayerRendererState>& InEditLayerRendererStates, const TSet<FName>& InWeightmapLayerNames)
+	: ComponentsToMerge(InComponentsToMerge)
 	, EditLayerRendererStates(InEditLayerRendererStates)
 	, WeightmapLayerNames(InWeightmapLayerNames)
 {
-	if (bIsHeightmapMerge)
-	{
-		// Make sure no weightmap layer name is passed in the case of heightmap :
-		WeightmapLayerNames.Empty();
-	}
 }
 
 #endif // WITH_EDITOR
diff --git a/Engine/Source/Runtime/Landscape/Private/LandscapeEditLayerRendererPrivate.cpp b/Engine/Source/Runtime/Landscape/Private/LandscapeEditLayerRendererPrivate.cpp
--- a/Engine/Source/Runtime/Landscape/Private/LandscapeEditLayerRendererPrivate.cpp
+++ b/Engine/Source/Runtime/Landscape/Private/LandscapeEditLayerRendererPrivate.cpp
@@ -15,21 +15,21 @@
 
 // ----------------------------------------------------------------------------------
 
-void ULandscapeDefaultEditLayerRenderer::GetRendererStateInfo(const ULandscapeInfo* InLandscapeInfo,
-	UE::Landscape::EditLayers::FEditLayerTargetTypeState& OutSupportedTargetTypeState, UE::Landscape::EditLayers::FEditLayerTargetTypeState& OutEnabledTargetTypeState, TArray<TSet<FName>>& OutRenderGroups) const
+void ULandscapeDefaultEditLayerRenderer::GetRendererStateInfo(const UE::Landscape::EditLayers::FMergeContext* InMergeContext,
+	UE::Landscape::EditLayers::FEditLayerTargetTypeState& OutSupportedTargetTypeState, UE::Landscape::EditLayers::FEditLayerTargetTypeState& OutEnabledTargetTypeState, TArray<TBitArray<>>& OutRenderGroups) const
 {
+	using namespace UE::Landscape::EditLayers;
+
 	// Supports all heightmaps and weightmaps:
-	OutSupportedTargetTypeState.SetTargetTypeMask(ELandscapeToolTargetTypeFlags::All);
-	TSet<FName> WeightmapLayerNames = GatherWeightmapLayerNames(InLandscapeInfo);
-	Algo::ForEach(WeightmapLayerNames, [&OutSupportedTargetTypeState](FName InLayerName) { OutSupportedTargetTypeState.AddWeightmap(InLayerName); });
+	OutSupportedTargetTypeState = FEditLayerTargetTypeState(InMergeContext, ELandscapeToolTargetTypeFlags::All, InMergeContext->GetValidTargetLayerBitIndices());
 	OutEnabledTargetTypeState = OutSupportedTargetTypeState;
 }
 
-TArray<UE::Landscape::EditLayers::FEditLayerRenderItem> ULandscapeDefaultEditLayerRenderer::GetRenderItems(const ULandscapeInfo* InLandscapeInfo) const
+TArray<UE::Landscape::EditLayers::FEditLayerRenderItem> ULandscapeDefaultEditLayerRenderer::GetRenderItems(const UE::Landscape::EditLayers::FMergeContext* InMergeContext) const
 {
 	using namespace UE::Landscape::EditLayers;
 
-	FEditLayerTargetTypeState OutputTargetTypeState(ELandscapeToolTargetTypeFlags::All);
+	FEditLayerTargetTypeState OutputTargetTypeState(InMergeContext, ELandscapeToolTargetTypeFlags::All, InMergeContext->GetValidTargetLayerBitIndices());
 	// Standard renderer : we don't need more than the component itself to render properly:
 	FInputWorldArea InputWorldArea(FInputWorldArea::CreateLocalComponent());
 	// The renderer only writes into the component itself (i.e. it renders to the area that it's currently being asked to render to):
@@ -60,33 +60,28 @@ void ULandscapeDefaultEditLayerRenderer::RenderLayer(FRenderParams& InRenderPara
 	RenderContext->RenderValidityRenderTargets(*RenderBatch);
 }
 
-TSet<FName> ULandscapeDefaultEditLayerRenderer::GatherWeightmapLayerNames(const ULandscapeInfo* InLandscapeInfo) const
-{
-	TSet<FName> WeightmapLayerNames;
-	// This renderer supports all weightmaps:
-	Algo::TransformIf(InLandscapeInfo->Layers, WeightmapLayerNames,
-		[](const FLandscapeInfoLayerSettings& InInfoLayerSettings) { return (InInfoLayerSettings.LayerInfoObj != nullptr); },
-		[](const FLandscapeInfoLayerSettings& InInfoLayerSettings) { return InInfoLayerSettings.LayerInfoObj->LayerName; });
-	return WeightmapLayerNames;
-}
-
 
 // ----------------------------------------------------------------------------------
 
-void ULandscapeHeightmapNormalsEditLayerRenderer::GetRendererStateInfo(const ULandscapeInfo* InLandscapeInfo,
-	UE::Landscape::EditLayers::FEditLayerTargetTypeState& OutSupportedTargetTypeState, UE::Landscape::EditLayers::FEditLayerTargetTypeState& OutEnabledTargetTypeState, TArray<TSet<FName>>& OutRenderGroups) const
+void ULandscapeHeightmapNormalsEditLayerRenderer::GetRendererStateInfo(const UE::Landscape::EditLayers::FMergeContext* InMergeContext,
+	UE::Landscape::EditLayers::FEditLayerTargetTypeState& OutSupportedTargetTypeState, UE::Landscape::EditLayers::FEditLayerTargetTypeState& OutEnabledTargetTypeState, TArray<TBitArray<>>& OutRenderGroups) const
 {
 	// Only relevant for heightmaps :
 	OutSupportedTargetTypeState.SetTargetTypeMask(ELandscapeToolTargetTypeFlags::Heightmap);
 	OutEnabledTargetTypeState.SetTargetTypeMask(ELandscapeToolTargetTypeFlags::Heightmap);
 }
 
-TArray<UE::Landscape::EditLayers::FEditLayerRenderItem> ULandscapeHeightmapNormalsEditLayerRenderer::GetRenderItems(const ULandscapeInfo* InLandscapeInfo) const
+TArray<UE::Landscape::EditLayers::FEditLayerRenderItem> ULandscapeHeightmapNormalsEditLayerRenderer::GetRenderItems(const UE::Landscape::EditLayers::FMergeContext* InMergeContext) const
 {
 	using namespace UE::Landscape::EditLayers;
 
 	// Only relevant for heightmaps :
-	FEditLayerTargetTypeState OutputTargetTypeState(ELandscapeToolTargetTypeFlags::Heightmap);
+	if (!InMergeContext->IsHeightmapMerge())
+	{
+		return {};
+	}
+	
+	FEditLayerTargetTypeState OutputTargetTypeState(InMergeContext, ELandscapeToolTargetTypeFlags::Heightmap);
 	// The input is relative and its size is equal to the size of 3x3 landscape components so that we gather all neighbor landscape components around each component : 
 	FInputWorldArea InputWorldArea(FInputWorldArea::CreateLocalComponent(FIntRect(-1, -1, 1, 1)));
 	// The renderer only writes into the component itself (i.e. it renders to the area that it's currently being asked to render to):
@@ -102,33 +97,35 @@ FString ULandscapeHeightmapNormalsEditLayerRenderer::GetEditLayerRendererDebugNa
 
 // ----------------------------------------------------------------------------------
 
-void ULandscapeWeightmapWeightBlendedLayersRenderer::GetRendererStateInfo(const ULandscapeInfo* InLandscapeInfo,
-	UE::Landscape::EditLayers::FEditLayerTargetTypeState& OutSupportedTargetTypeState, UE::Landscape::EditLayers::FEditLayerTargetTypeState& OutEnabledTargetTypeState, TArray<TSet<FName>>& OutRenderGroups) const
+void ULandscapeWeightmapWeightBlendedLayersRenderer::GetRendererStateInfo(const UE::Landscape::EditLayers::FMergeContext* InMergeContext,
+	UE::Landscape::EditLayers::FEditLayerTargetTypeState& OutSupportedTargetTypeState, UE::Landscape::EditLayers::FEditLayerTargetTypeState& OutEnabledTargetTypeState, TArray<TBitArray<>>& OutRenderGroups) const
 {
+	using namespace UE::Landscape::EditLayers;
+
+	TBitArray<> WeightBlendedWeightmapLayerBitIndices = GatherWeightBlendedWeightmapLayerBitIndices(InMergeContext);
 	// Only relevant for weightmaps :
-	OutSupportedTargetTypeState.SetTargetTypeMask(ELandscapeToolTargetTypeFlags::Weightmap);
-	TSet<FName> WeightBlendedWeightmapLayerNames = GatherWeightBlendedWeightmapLayerNames(InLandscapeInfo);
-	Algo::ForEach(WeightBlendedWeightmapLayerNames, [&OutSupportedTargetTypeState](FName InLayerName) { OutSupportedTargetTypeState.AddWeightmap(InLayerName); });
+	FEditLayerTargetTypeState OutputTargetTypeState(InMergeContext, ELandscapeToolTargetTypeFlags::Weightmap, WeightBlendedWeightmapLayerBitIndices);
 	OutEnabledTargetTypeState = OutSupportedTargetTypeState;
 
-	check(Algo::AllOf(WeightBlendedWeightmapLayerNames, [InLandscapeInfo](FName InTargetLayerName)
-	{
-		return InLandscapeInfo->Layers.ContainsByPredicate([InTargetLayerName](const FLandscapeInfoLayerSettings& InInfoLayerSettings) { return (InInfoLayerSettings.GetLayerName() == InTargetLayerName); });
-	}));
-
-	// Now fill in the render groups :
-	if (!WeightBlendedWeightmapLayerNames.IsEmpty())
+	// Now fill in the render groups : one group for all weight-blended layers : 
+	if (WeightBlendedWeightmapLayerBitIndices.Contains(true))
 	{
-		OutRenderGroups.Add(WeightBlendedWeightmapLayerNames);
+		OutRenderGroups.Add(WeightBlendedWeightmapLayerBitIndices);
 	}
 }
 
-TArray<UE::Landscape::EditLayers::FEditLayerRenderItem> ULandscapeWeightmapWeightBlendedLayersRenderer::GetRenderItems(const ULandscapeInfo* InLandscapeInfo) const
+TArray<UE::Landscape::EditLayers::FEditLayerRenderItem> ULandscapeWeightmapWeightBlendedLayersRenderer::GetRenderItems(const UE::Landscape::EditLayers::FMergeContext* InMergeContext) const
 {
 	using namespace UE::Landscape::EditLayers;
 
 	// Only relevant for weightmaps :
-	FEditLayerTargetTypeState OutputTargetTypeState(ELandscapeToolTargetTypeFlags::Weightmap, GatherWeightBlendedWeightmapLayerNames(InLandscapeInfo).Array());
+	if (InMergeContext->IsHeightmapMerge())
+	{
+		return {};
+	}
+
+	TBitArray<> WeightBlendedWeightmapLayerBitIndices = GatherWeightBlendedWeightmapLayerBitIndices(InMergeContext);
+	FEditLayerTargetTypeState OutputTargetTypeState(InMergeContext, ELandscapeToolTargetTypeFlags::Weightmap, WeightBlendedWeightmapLayerBitIndices);
 	// Standard renderer : we don't need more than the component itself to render properly:
 	FInputWorldArea InputWorldArea(FInputWorldArea::CreateLocalComponent());
 	// The renderer only writes into the component itself (i.e. it renders to the area that it's currently being asked to render to):
@@ -142,13 +139,17 @@ FString ULandscapeWeightmapWeightBlendedLayersRenderer::GetEditLayerRendererDebu
 	return TEXT("Final Weight Blend");
 }
 
-TSet<FName> ULandscapeWeightmapWeightBlendedLayersRenderer::GatherWeightBlendedWeightmapLayerNames(const ULandscapeInfo* InLandscapeInfo) const
+TBitArray<> ULandscapeWeightmapWeightBlendedLayersRenderer::GatherWeightBlendedWeightmapLayerBitIndices(const UE::Landscape::EditLayers::FMergeContext* InMergeContext) const
 {
-	TSet<FName> WeightBlendedWeightmapLayerNames;
-	Algo::TransformIf(InLandscapeInfo->Layers, WeightBlendedWeightmapLayerNames,
-		[](const FLandscapeInfoLayerSettings& InInfoLayerSettings) { return ((InInfoLayerSettings.LayerInfoObj != nullptr) && !InInfoLayerSettings.LayerInfoObj->bNoWeightBlend); },
-		[](const FLandscapeInfoLayerSettings& InInfoLayerSettings) { return InInfoLayerSettings.LayerInfoObj->LayerName; });
-	return WeightBlendedWeightmapLayerNames;
+	TBitArray<> WeightBlendedWeightmapLayerBitIndices = InMergeContext->GetValidTargetLayerBitIndices();
+	InMergeContext->ForEachValidTargetLayer(
+		[&WeightBlendedWeightmapLayerBitIndices, InMergeContext] (int32 InTargetLayerIndex, const FName& InTargetLayerName, ULandscapeLayerInfoObject* InWeightmapLayerInfo)
+		{
+			check(InWeightmapLayerInfo != nullptr);
+			WeightBlendedWeightmapLayerBitIndices[InTargetLayerIndex] = !InWeightmapLayerInfo->bNoWeightBlend;
+			return true;
+		});
+	return WeightBlendedWeightmapLayerBitIndices;
 }
 
 #endif // WITH_EDITOR
\ No newline at end of file
diff --git a/Engine/Source/Runtime/Landscape/Private/LandscapeEditLayerRendererPrivate.h b/Engine/Source/Runtime/Landscape/Private/LandscapeEditLayerRendererPrivate.h
--- a/Engine/Source/Runtime/Landscape/Private/LandscapeEditLayerRendererPrivate.h
+++ b/Engine/Source/Runtime/Landscape/Private/LandscapeEditLayerRendererPrivate.h
@@ -29,15 +29,12 @@ class ULandscapeDefaultEditLayerRenderer : public UObject
 #if WITH_EDITOR
 public:
 	//~ Begin ILandscapeEditLayerRenderer implementation
-	virtual void GetRendererStateInfo(const ULandscapeInfo* InLandscapeInfo,
-		UE::Landscape::EditLayers::FEditLayerTargetTypeState& OutSupportedTargetTypeState, UE::Landscape::EditLayers::FEditLayerTargetTypeState& OutEnabledTargetTypeState, TArray<TSet<FName>>& OutRenderGroups) const override;
-	virtual TArray<UE::Landscape::EditLayers::FEditLayerRenderItem> GetRenderItems(const ULandscapeInfo* InLandscapeInfo) const override;
+	virtual void GetRendererStateInfo(const UE::Landscape::EditLayers::FMergeContext* InMergeContext,
+		UE::Landscape::EditLayers::FEditLayerTargetTypeState& OutSupportedTargetTypeState, UE::Landscape::EditLayers::FEditLayerTargetTypeState& OutEnabledTargetTypeState, TArray<TBitArray<>>& OutRenderGroups) const override;
+	virtual TArray<UE::Landscape::EditLayers::FEditLayerRenderItem> GetRenderItems(const UE::Landscape::EditLayers::FMergeContext* InMergeContext) const override;
 	virtual FString GetEditLayerRendererDebugName() const override;
 	virtual void RenderLayer(FRenderParams& InRenderParams) override;
 	//~ End ILandscapeEditLayerRenderer implementation
-
-private:
-	TSet<FName> GatherWeightmapLayerNames(const ULandscapeInfo* InLandscapeInfo) const;
 #endif // WITH_EDITOR
 };
 
@@ -57,9 +54,9 @@ class ULandscapeHeightmapNormalsEditLayerRenderer : public UObject
 #if WITH_EDITOR
 public:
 	//~ Begin ILandscapeEditLayerRenderer implementation
-	virtual void GetRendererStateInfo(const ULandscapeInfo* InLandscapeInfo,
-		UE::Landscape::EditLayers::FEditLayerTargetTypeState& OutSupportedTargetTypeState, UE::Landscape::EditLayers::FEditLayerTargetTypeState& OutEnabledTargetTypeState, TArray<TSet<FName>>& OutRenderGroups) const override;
-	virtual TArray<UE::Landscape::EditLayers::FEditLayerRenderItem> GetRenderItems(const ULandscapeInfo* InLandscapeInfo) const override;
+	virtual void GetRendererStateInfo(const UE::Landscape::EditLayers::FMergeContext* InMergeContext,
+		UE::Landscape::EditLayers::FEditLayerTargetTypeState& OutSupportedTargetTypeState, UE::Landscape::EditLayers::FEditLayerTargetTypeState& OutEnabledTargetTypeState, TArray<TBitArray<>>& OutRenderGroups) const override;
+	virtual TArray<UE::Landscape::EditLayers::FEditLayerRenderItem> GetRenderItems(const UE::Landscape::EditLayers::FMergeContext* InMergeContext) const override;
 	virtual FString GetEditLayerRendererDebugName() const override;
 	virtual void RenderLayer(FRenderParams& InRenderParams) override;
 	//~ End ILandscapeEditLayerRenderer implementation
@@ -80,14 +77,14 @@ class ULandscapeWeightmapWeightBlendedLayersRenderer : public UObject
 #if WITH_EDITOR
 public:
 	//~ Begin ILandscapeEditLayerRenderer implementation
-	virtual void GetRendererStateInfo(const ULandscapeInfo* InLandscapeInfo, 
-		UE::Landscape::EditLayers::FEditLayerTargetTypeState& OutSupportedTargetTypeState, UE::Landscape::EditLayers::FEditLayerTargetTypeState& OutEnabledTargetTypeState, TArray<TSet<FName>>& OutRenderGroups) const override;
-	virtual TArray<UE::Landscape::EditLayers::FEditLayerRenderItem> GetRenderItems(const ULandscapeInfo* InLandscapeInfo) const override;
+	virtual void GetRendererStateInfo(const UE::Landscape::EditLayers::FMergeContext* InMergeContext,
+		UE::Landscape::EditLayers::FEditLayerTargetTypeState& OutSupportedTargetTypeState, UE::Landscape::EditLayers::FEditLayerTargetTypeState& OutEnabledTargetTypeState, TArray<TBitArray<>>& OutRenderGroups) const override;
+	virtual TArray<UE::Landscape::EditLayers::FEditLayerRenderItem> GetRenderItems(const UE::Landscape::EditLayers::FMergeContext* InMergeContext) const override;
 	virtual FString GetEditLayerRendererDebugName() const override;
 	virtual void RenderLayer(FRenderParams& InRenderParams) override;
 	//~ End ILandscapeEditLayerRenderer implementation
 
-	TSet<FName> GatherWeightBlendedWeightmapLayerNames(const ULandscapeInfo* InLandscapeInfo) const;
+	TBitArray<> GatherWeightBlendedWeightmapLayerBitIndices(const UE::Landscape::EditLayers::FMergeContext* InMergeContext) const;
 #endif // WITH_EDITOR
 };
 
diff --git a/Engine/Source/Runtime/Landscape/Private/LandscapeEditLayers.cpp b/Engine/Source/Runtime/Landscape/Private/LandscapeEditLayers.cpp
--- a/Engine/Source/Runtime/Landscape/Private/LandscapeEditLayers.cpp
+++ b/Engine/Source/Runtime/Landscape/Private/LandscapeEditLayers.cpp
@@ -5374,19 +5374,28 @@ int32 ALandscape::PerformLayersHeightmapsLocalMerge(const FUpdateLayersContentCo
 	return InMergeParams.HeightmapUpdateModes;
 }
 
-TArray<UE::Landscape::EditLayers::FEditLayerRendererState> ALandscape::GetEditLayerRendererStates(const ULandscapeInfo* InLandscapeInfo, bool bInSkipBrush)
+TArray<UE::Landscape::EditLayers::FEditLayerRendererState> ALandscape::GetEditLayerRendererStates(const UE::Landscape::EditLayers::FMergeContext* InMergeContext)
 {
 	using namespace UE::Landscape::EditLayers;
 
 	TArray<FEditLayerRendererState> RendererStates;
 	for (FLandscapeLayer& Layer : LandscapeEditLayers)
 	{
-		RendererStates.Append(Layer.GetEditLayerRendererStates(InLandscapeInfo, bInSkipBrush));
+		TArray<FEditLayerRendererState> LayerRendererStates = Layer.GetEditLayerRendererStates(InMergeContext);
+		// Renderer states generated from a layer inherit the layer's state so start by computing the layer's target type mask: 
+		ELandscapeToolTargetTypeFlags LayerTargetTypeMask = Layer.GetEnabledTargetTypeMask();
+		// Then disable all types that are not in the layer's mask :
+		for (FEditLayerRendererState& LayerRendererState : LayerRendererStates)
+		{
+			LayerRendererState.DisableTargetTypeMask(~LayerTargetTypeMask);
+		}
+
+		RendererStates.Append(LayerRendererStates);
 	}
 	return RendererStates;
 }
 
-TArray<UE::Landscape::EditLayers::FEditLayerRendererState> FLandscapeLayer::GetEditLayerRendererStates(const ULandscapeInfo* InLandscapeInfo, bool bInSkipBrush)
+TArray<UE::Landscape::EditLayers::FEditLayerRendererState> FLandscapeLayer::GetEditLayerRendererStates(const UE::Landscape::EditLayers::FMergeContext* InMergeContext)
 {
 	using namespace UE::Landscape::EditLayers;
 
@@ -5398,20 +5407,40 @@ TArray<UE::Landscape::EditLayers::FEditLayerRendererState> FLandscapeLayer::GetE
 	// The edit layer itself might be a renderer:
 	if (ILandscapeEditLayerRenderer* AsRenderer = Cast<ILandscapeEditLayerRenderer>(EditLayer))
 	{
-		RendererStates.Emplace(AsRenderer, InLandscapeInfo);
+		RendererStates.Emplace(InMergeContext, AsRenderer);
 	}
 
 	// The layer can also be a renderer provider
-	RendererStates.Append(EditLayer->GetEditLayerRendererStates(InLandscapeInfo, bInSkipBrush));
+	RendererStates.Append(EditLayer->GetEditLayerRendererStates(InMergeContext));
 
 	for (FLandscapeLayerBrush& Brush : Brushes)
 	{
-		RendererStates.Append(Brush.GetEditLayerRendererStates(InLandscapeInfo, bInSkipBrush));
+		RendererStates.Append(Brush.GetEditLayerRendererStates(InMergeContext));
 	}
 	return RendererStates;
 }
 
-TArray<UE::Landscape::EditLayers::FEditLayerRendererState> FLandscapeLayerBrush::GetEditLayerRendererStates(const ULandscapeInfo* InLandscapeInfo, bool bInSkipBrush)
+ELandscapeToolTargetTypeFlags FLandscapeLayer::GetEnabledTargetTypeMask() const
+{
+	// Compute the default state of each target type : 
+	ELandscapeToolTargetTypeFlags EnabledTargetTypeMask = ELandscapeToolTargetTypeFlags::None;
+	if (bVisible)
+	{
+		if (HeightmapAlpha != 0.0f)
+		{
+			EnabledTargetTypeMask |= ELandscapeToolTargetTypeFlags::Heightmap;
+		}
+		if (WeightmapAlpha > 0.0f)
+		{
+			EnabledTargetTypeMask |= ELandscapeToolTargetTypeFlags::Weightmap;
+		}
+		// If the layer is visible, visibility is always considered enabled since it does not depend on WeightmapAlpha 
+		EnabledTargetTypeMask |= ELandscapeToolTargetTypeFlags::Visibility;
+	}
+	return EnabledTargetTypeMask;
+}
+
+TArray<UE::Landscape::EditLayers::FEditLayerRendererState> FLandscapeLayerBrush::GetEditLayerRendererStates(const UE::Landscape::EditLayers::FMergeContext* InMergeContext)
 {
 	using namespace UE::Landscape::EditLayers;
 
@@ -5421,39 +5450,27 @@ TArray<UE::Landscape::EditLayers::FEditLayerRendererState> FLandscapeLayerBrush:
 		return {};
 	}
 
-	return BlueprintBrush->GetEditLayerRendererStates(InLandscapeInfo, bInSkipBrush);
+	return BlueprintBrush->GetEditLayerRendererStates(InMergeContext);
 }
 
-void ULandscapeEditLayerPersistent::GetRendererStateInfo(const ULandscapeInfo* InLandscapeInfo, 
-	UE::Landscape::EditLayers::FEditLayerTargetTypeState& OutSupportedTargetTypeState, UE::Landscape::EditLayers::FEditLayerTargetTypeState& OutEnabledTargetTypeState, TArray<TSet<FName>>& OutRenderGroups) const
+void ULandscapeEditLayerPersistent::GetRendererStateInfo(const UE::Landscape::EditLayers::FMergeContext* InMergeContext,
+	UE::Landscape::EditLayers::FEditLayerTargetTypeState& OutSupportedTargetTypeState, UE::Landscape::EditLayers::FEditLayerTargetTypeState& OutEnabledTargetTypeState, TArray<TBitArray<>>& OutRenderGroups) const
 {
 	using namespace UE::Landscape::EditLayers;
 
 	const FLandscapeLayer* LayerStruct = GetOwningLayer();
 	check(LayerStruct != nullptr);
 
-	// A layer can support all target types :
-	OutSupportedTargetTypeState.SetTargetTypeMask(ELandscapeToolTargetTypeFlags::All);
+	// A layer can support all target types. 
 	// Add an entry for each weightmap and consider them supported because there's nothing that prevents a given edit layer to write on a given weightmap layer :
-	Algo::ForEach(InLandscapeInfo->Layers, [&OutSupportedTargetTypeState] (const FLandscapeInfoLayerSettings& InInfoLayerSettings) { OutSupportedTargetTypeState.AddWeightmap(InInfoLayerSettings.GetLayerName()); });
+	OutSupportedTargetTypeState = FEditLayerTargetTypeState(InMergeContext, ELandscapeToolTargetTypeFlags::All, InMergeContext->GetValidTargetLayerBitIndices());
 
-	// Compute the default state of each target type : 
-	ELandscapeToolTargetTypeFlags EnabledTargetTypeMask = ELandscapeToolTargetTypeFlags::None;
-	if (LayerStruct->bVisible)
-	{
-		if (LayerStruct->HeightmapAlpha > 0.0f)
-		{
-			EnabledTargetTypeMask |= ELandscapeToolTargetTypeFlags::Heightmap;
-		}
-		if (LayerStruct->WeightmapAlpha > 0.0f)
-		{
-			EnabledTargetTypeMask |= ELandscapeToolTargetTypeFlags::Weightmap | ELandscapeToolTargetTypeFlags::Visibility;
-		}
-	}
-	OutEnabledTargetTypeState = FEditLayerTargetTypeState(EnabledTargetTypeMask, MakeArrayView(OutSupportedTargetTypeState.GetActiveWeightmaps()));
+	// TODO [jonathan.bard] : for now, at this level enabled == supported (the enabled state will be corrected by FLandscapeLayer since it is the one that has the bVisibility / Heightmap/WeightmapAlpha members
+	//  but eventually, if those move to ULandscapeEditLayerBase, this is where the layer's enabled state will need to be set)
+	OutEnabledTargetTypeState = OutSupportedTargetTypeState;
 }
 
-TArray<UE::Landscape::EditLayers::FEditLayerRenderItem> ULandscapeEditLayerPersistent::GetRenderItems(const ULandscapeInfo* InLandscapeInfo) const
+TArray<UE::Landscape::EditLayers::FEditLayerRenderItem> ULandscapeEditLayerPersistent::GetRenderItems(const UE::Landscape::EditLayers::FMergeContext* InMergeContext) const
 {
 	using namespace UE::Landscape::EditLayers;
 
@@ -5461,21 +5478,23 @@ TArray<UE::Landscape::EditLayers::FEditLayerRenderItem> ULandscapeEditLayerPersi
 	check(LayerStruct != nullptr);
 
 	TArray<FEditLayerRenderItem> RenderItems;
-	RenderItems.Reserve(InLandscapeInfo->XYtoComponentMap.Num() + 1);
+	ULandscapeInfo* Info = InMergeContext->GetLandscapeInfo();
+	RenderItems.Reserve(Info->XYtoComponentMap.Num() + 1);
 
 	// Heightmaps : We only need a single render item for heightmaps, because heightmaps are always present :
+	if (InMergeContext->IsHeightmapMerge())
 	{
 		// The layer doesn't need more than the component itself to render properly:
 		FInputWorldArea InputWorldArea(FInputWorldArea::CreateLocalComponent());
 		// The layer only writes into the component itself (i.e. it renders to the area that it's currently being asked to render to):
 		FOutputWorldArea OutputWorldArea(FOutputWorldArea::CreateLocalComponent());
-		RenderItems.Emplace(FEditLayerTargetTypeState(ELandscapeToolTargetTypeFlags::Heightmap), InputWorldArea, OutputWorldArea, /*bModifyExistingWeightmapsOnly = */false);
+		RenderItems.Emplace(FEditLayerTargetTypeState(InMergeContext, ELandscapeToolTargetTypeFlags::Heightmap), InputWorldArea, OutputWorldArea, /*bModifyExistingWeightmapsOnly = */false);
 	}
-
+	else
 	// Weightmaps : add one render item per component in order to be able to indicate exactly which weightmap is needed for each of them. This avoids pre-allocating weightmaps on the merged result where we know we won't write a weightmap : 
 	{
 		TArray<ULandscapeComponent*> AllLandscapeComponents;
-		InLandscapeInfo->XYtoComponentMap.GenerateValueArray(AllLandscapeComponents);
+		Info->XYtoComponentMap.GenerateValueArray(AllLandscapeComponents);
 
 		for (ULandscapeComponent* Component : AllLandscapeComponents)
 		{
@@ -5484,7 +5503,7 @@ TArray<UE::Landscape::EditLayers::FEditLayerRenderItem> ULandscapeEditLayerPersi
 			if (FLandscapeLayerComponentData* LayerData = Component->GetLayerData(LayerStruct->Guid))
 			{
 				TArray<FName, TInlineAllocator<16>> ComponentWeightmaps;
-				FEditLayerTargetTypeState OutputTargetTypeState;
+				FEditLayerTargetTypeState OutputTargetTypeState(InMergeContext);
 
 				// Iterate through all allocated weightmaps in order to find which one we will really write to : 
 				for (const FWeightmapLayerAllocationInfo& LayerAllocationInfo : LayerData->WeightmapData.LayerAllocations)
@@ -5499,7 +5518,7 @@ TArray<UE::Landscape::EditLayers::FEditLayerRenderItem> ULandscapeEditLayerPersi
 						{
 							OutputTargetTypeState.AddTargetTypeMask(ELandscapeToolTargetTypeFlags::Weightmap);
 						}
-						OutputTargetTypeState.AddWeightmap(LayerAllocationInfo.GetLayerName());
+						OutputTargetTypeState.AddWeightmapChecked(LayerAllocationInfo.GetLayerName());
 					}
 				}
 
@@ -5548,7 +5567,7 @@ void ULandscapeEditLayerPersistent::RenderLayer(ILandscapeEditLayerRenderer::FRe
 		}
 	};
 
-	const int32 NumTargetLayers = InRenderParams.RenderGroupTargetLayerNames.Num();
+	const int32 NumTargetLayersInRenderGroup = InRenderParams.RenderGroupTargetLayerNames.Num();
 
 	// The first step is to copy all of the necessary components' textures to WriteRT and remove the duplicate borders when doing so (subsection by subsection). This is done with a 
 	//  "copy from multiple sources" shader instead of several texture copies, in order to reduce the amount of copy texture commands, which can have a big impact on render-thread performance
@@ -5570,7 +5589,7 @@ void ULandscapeEditLayerPersistent::RenderLayer(ILandscapeEditLayerRenderer::FRe
 		uint8 ChannelSwizzleMask = RGBAToRGBASwizzleMask;
 	};
 	TArray<FCopyQuadParams> CopyQuadsParams;
-	CopyQuadsParams.Reserve(InRenderParams.SortedComponentMergeRenderInfos.Num() * NumTargetLayers * Landscape->NumSubsections * Landscape->NumSubsections);
+	CopyQuadsParams.Reserve(InRenderParams.SortedComponentMergeRenderInfos.Num() * NumTargetLayersInRenderGroup* Landscape->NumSubsections * Landscape->NumSubsections);
 
 	// This is kinda hacky, but we since we reuse AddRasterizeToRectsPass and it passes the source texture's sample coordinates from VS to PS via the UVs, dividing the source rect by a single TextureSize, 
 	//  we correct each quad's source rect by a common scale factor so that it samples at the right location in each source texture: 
@@ -5580,9 +5599,9 @@ void ULandscapeEditLayerPersistent::RenderLayer(ILandscapeEditLayerRenderer::FRe
 	// Build the full list of quads to copy :
 	for (const FComponentMergeRenderInfo& ComponentMergeRenderInfo : InRenderParams.SortedComponentMergeRenderInfos)
 	{
-		for (int32 TargetLayerIndex = 0; TargetLayerIndex < NumTargetLayers; ++TargetLayerIndex)
+		for (int32 TargetLayerIndexInRenderGroup = 0; TargetLayerIndexInRenderGroup < NumTargetLayersInRenderGroup; ++TargetLayerIndexInRenderGroup)
 		{
-			const FName TargetLayerName = InRenderParams.RenderGroupTargetLayerNames[TargetLayerIndex];
+			const FName TargetLayerName = InRenderParams.RenderGroupTargetLayerNames[TargetLayerIndexInRenderGroup];
 			UTexture2D* SourceTexture = nullptr;
 			FVector2D SourceTextureBias(ForceInit);
 			// Copy all channels by default :
@@ -5642,7 +5661,7 @@ void ULandscapeEditLayerPersistent::RenderLayer(ILandscapeEditLayerRenderer::FRe
 						QuadParams.SourceTextureResource = SourceTextureResource;
 						QuadParams.SourceRect = FIntRect(SourceTexturePixelCoordinatesToQuadCoords(SourcePosition), SourceTexturePixelCoordinatesToQuadCoords(SourcePosition + FIntPoint(SubsectionSizeVerts, SubsectionSizeVerts)));
 						QuadParams.DestinationRect = FIntRect(DestinationPosition, DestinationPosition + FIntPoint(SubsectionSizeVerts, SubsectionSizeVerts));
-						QuadParams.DestinationArrayIndex = TargetLayerIndex;
+						QuadParams.DestinationArrayIndex = TargetLayerIndexInRenderGroup;
 						QuadParams.ChannelSwizzleMask = CopyTextureSwizzleMask;
 					}
 				}
@@ -5881,7 +5900,7 @@ void ULandscapeEditLayerPersistent::RenderLayer(ILandscapeEditLayerRenderer::FRe
 
 			// Prepare a list of all paint layers to merge and do it all in 1 command :
 			TArray<FMergePaintLayerParams> MergePaintLayersParams;
-			MergePaintLayersParams.Reserve(NumTargetLayers);
+			MergePaintLayersParams.Reserve(NumTargetLayersInRenderGroup);
 			for (const FName& TargetLayerName : InRenderParams.RenderGroupTargetLayerNames)
 			{
 				FMergePaintLayerParams& MergePaintLayerParams = MergePaintLayersParams.Emplace_GetRef(TargetLayerName);
@@ -6066,7 +6085,7 @@ namespace UE::Landscape::EditLayers::Private
 	/** Struct that holds all the per-renderer information needed when preparing the batched merge context */
 	struct FEditLayerRendererRenderInfo
 	{
-		FEditLayerRendererRenderInfo() = default;
+		FEditLayerRendererRenderInfo() = delete;
 		FEditLayerRendererRenderInfo(const FEditLayerRendererState& InRendererState, int32 InRendererIndex, const FTransform& InLandscapeTransform, double InMaxLocalHeight, int32 InNumComponentsToRender)
 			: RendererState(InRendererState)
 			, RendererIndex(InRendererIndex)
@@ -6528,7 +6547,7 @@ void ULandscapeWeightmapWeightBlendedLayersRenderer::RenderLayer(FRenderParams&
 	});
 }
 
-UE::Landscape::EditLayers::FMergeRenderContext ALandscape::PrepareEditLayersMergeRenderContext(const UE::Landscape::EditLayers::FMergeRenderParams& InMergeRenderParams)
+UE::Landscape::EditLayers::FMergeRenderContext ALandscape::PrepareEditLayersMergeRenderContext(const UE::Landscape::EditLayers::FMergeContext& InMergeContext, const UE::Landscape::EditLayers::FMergeRenderParams& InMergeRenderParams)
 {
 	TRACE_CPUPROFILER_EVENT_SCOPE(ALandscape::PrepareEditLayersMergeRenderContext);
 
@@ -6548,38 +6567,25 @@ UE::Landscape::EditLayers::FMergeRenderContext ALandscape::PrepareEditLayersMerg
 	check(Info != nullptr);
 	check(!LandscapeEditLayers.IsEmpty()); 
 
-	const FTransform& LandscapeTransform = GetTransform();
-	const FBox LandscapeLoadedBounds = Info->GetLoadedBounds();
-
-	FMergeRenderContext MergeRenderContext(InMergeRenderParams.Landscape, InMergeRenderParams.bIsHeightmapMerge);
-
-	// Start by gathering all possible target layer names on this landscape: 
-	//  This list of all unique target layer names will help accelerate the gathering of output layers on each component (using bit arrays) as well as the target layers intersection tests:
-	if (InMergeRenderParams.bIsHeightmapMerge)
+	// Warn if invalid layer names are requested : 
+	if (!InMergeContext.IsHeightmapMerge())
 	{
-		// Only one target layer in the case of heightmap: 
-		const FName TargetLayerName = FName("Height");
-		MergeRenderContext.AllTargetLayerNames = { TargetLayerName };
-		// And it's valid and final :
-		MergeRenderContext.ValidTargetLayerBitIndices = TBitArray<>(true, 1);
-		MergeRenderContext.FinalTargetLayerBitIndices = TBitArray<>(true, 1);
-	}
-	else
-	{
-		// Gather all target layer names and mark those that are valid layers :
-		for (const FLandscapeInfoLayerSettings& LayerSettings : Info->Layers)
-		{ 
-			check(!LayerSettings.LayerName.IsNone());
-			ULandscapeLayerInfoObject*& LayerInfo = MergeRenderContext.WeightmapLayerInfos.Add_GetRef(LayerSettings.LayerInfoObj);
-
-			MergeRenderContext.AllTargetLayerNames.Add(LayerSettings.LayerName);
-			MergeRenderContext.ValidTargetLayerBitIndices.Add(LayerInfo != nullptr);
-			MergeRenderContext.FinalTargetLayerBitIndices.Add(false); // will be computed later...
+		for (FName TargetLayerName : InMergeRenderParams.WeightmapLayerNames)
+		{
+			if (!InMergeContext.IsValidTargetLayerName(TargetLayerName))
+			{
+				UE_LOG(LogLandscape, Warning, TEXT("Target layer \"%s\" was requested by the merge but is invalid (missing its layer info object asset). It will be ignored"), *TargetLayerName.ToString());
+			}
 		}
 	}
+
+	const FTransform& LandscapeTransform = GetTransform();
+	const FBox LandscapeLoadedBounds = Info->GetLoadedBounds();
+
+	FMergeRenderContext MergeRenderContext(InMergeContext);
+	const ELandscapeToolTargetTypeFlags MergeTypeMask = MergeRenderContext.IsHeightmapMerge() ? ELandscapeToolTargetTypeFlags::Heightmap : (ELandscapeToolTargetTypeFlags::Weightmap | ELandscapeToolTargetTypeFlags::Visibility);
 	const int32 NumAllTargetLayerNames = MergeRenderContext.AllTargetLayerNames.Num();
 	check(MergeRenderContext.ValidTargetLayerBitIndices.Num() == NumAllTargetLayerNames);
-	check(MergeRenderContext.FinalTargetLayerBitIndices.Num() == NumAllTargetLayerNames);
 
 	// InMergeRenderParams.EditLayerRendererStates contains a list of renderers that is not quite final : ULandscapeDefaultEditLayerRenderer is always inserted at the beginning to make sure we always render at 
 	//  least the default value and ULandscapeWeightmapWeightBlendedLayersRenderer can optionally be inserted at the end too :
@@ -6590,33 +6596,47 @@ UE::Landscape::EditLayers::FMergeRenderContext ALandscape::PrepareEditLayersMerg
 		// We always have at least 1 renderer at the start : the default one, whose job is to both provide the default value and act as the final "gatherer" of the component dependencies 
 		//  from all the renderers above (see class comment for more details) :
 		ULandscapeDefaultEditLayerRenderer* InitialEditLayerRenderer = ULandscapeDefaultEditLayerRenderer::StaticClass()->GetDefaultObject<ULandscapeDefaultEditLayerRenderer>();
-		CandidateEditLayerRendererStates.Emplace(InitialEditLayerRenderer, Info);
+		CandidateEditLayerRendererStates.Emplace(&MergeRenderContext, InitialEditLayerRenderer);
 
 		// Then append all the provided renderer states
 		CandidateEditLayerRendererStates.Append(InMergeRenderParams.EditLayerRendererStates);
 
 		// Legacy weight-blending requires an additional renderer at the end of the stack, to weight-blend the weightmaps needing it :
-		if (!InMergeRenderParams.bIsHeightmapMerge)
+		if (!MergeRenderContext.IsHeightmapMerge())
 		{
 			ULandscapeWeightmapWeightBlendedLayersRenderer* WeightmapWeightBlendedLayersRenderer = ULandscapeWeightmapWeightBlendedLayersRenderer::StaticClass()->GetDefaultObject<ULandscapeWeightmapWeightBlendedLayersRenderer>();
-			if (!WeightmapWeightBlendedLayersRenderer->GatherWeightBlendedWeightmapLayerNames(Info).IsEmpty())
+			if (WeightmapWeightBlendedLayersRenderer->GatherWeightBlendedWeightmapLayerBitIndices(&MergeRenderContext).Contains(true))
 			{
-				CandidateEditLayerRendererStates.Emplace(WeightmapWeightBlendedLayersRenderer, Info);
+				CandidateEditLayerRendererStates.Emplace(&MergeRenderContext, WeightmapWeightBlendedLayersRenderer);
 			}
 		}
 	}
 
+	// Only retain renderers that are relevant for this merge : 
+	CandidateEditLayerRendererStates = CandidateEditLayerRendererStates.FilterByPredicate([&InMergeRenderParams, MergeTypeMask](const FEditLayerRendererState& InRendererState)
+		{
+			return EnumHasAnyFlags(InRendererState.GetActiveTargetTypeMask(), MergeTypeMask);
+		});
+
 	// FinalEditLayerRendererStates will contain the renderer states that are actually relevant to this merge :
 	TArray<FEditLayerRendererState> FinalEditLayerRendererStates;
 	FinalEditLayerRendererStates.Reserve(CandidateEditLayerRendererStates.Num());
 
 	// Within each render batch, elements can be processed group by group. For heightmap/visibility, there's only one such group. For weightmaps, there's one group per list of weightmaps
 	//  that need to be processed together for weight-blending. Each group is composed of a list of (weightmap layer) names (it's only a debug name in the case of heightmaps)
-
-	// First, let's work out the weightmaps inter-dependencies (i.e. horizontal dependencies) : weight-blending requires some weightmaps to be processed together: 
 	TArray<TBitArray<>> FinalRenderGroups;
-	if (!InMergeRenderParams.bIsHeightmapMerge)
+	if (MergeRenderContext.IsHeightmapMerge())
+	{
+		// All candidates are valid in heightmap merge, the ones that don't affect heightmap have already been filtered from CandidateEditLayerRendererStates
+		FinalEditLayerRendererStates = CandidateEditLayerRendererStates;
+
+		// Only one group in the case of heightmap: 
+		const FName TargetLayerName = FName("Height");
+		FinalRenderGroups = { TBitArray<>(true, 1) };
+	}
+	else
 	{
+		// First, let's work out the weightmaps inter-dependencies (i.e. horizontal dependencies) : weight-blending requires some weightmaps to be processed together: 
 		TRACE_CPUPROFILER_EVENT_SCOPE(PrepareRenderGroups);
 
 		TArray<TArray<TBitArray<>>> RelevantRenderGroupsPerRenderer;
@@ -6626,16 +6646,16 @@ UE::Landscape::EditLayers::FMergeRenderContext ALandscape::PrepareEditLayersMerg
 		// Iterate through all renderers to find out what render group they render (only keep those that are relevant for the current state and request) : 
 		for (const FEditLayerRendererState& EditLayerRendererState : CandidateEditLayerRendererStates)
 		{
-			TBitArray<> RendererStateEnabledTargetLayerBitIndices(MergeRenderContext.ConvertTargetLayerNamesToBitIndices(EditLayerRendererState.GetEnabledTargetWeightmaps()));
+			TBitArray<> RendererStateEnabledTargetLayerBitIndices(EditLayerRendererState.GetActiveTargetWeightmapBitIndices());
 			// Retain only the target layer names that are enabled and valid :
 			RendererStateEnabledTargetLayerBitIndices.CombineWithBitwiseAND(MergeRenderContext.ValidTargetLayerBitIndices, EBitwiseOperatorFlags::MinSize);
 
 			// List of all supported render groups for this renderer :
 			TArray<TBitArray<>> RelevantRenderGroupsForThisRenderer;
-			for (const TSet<FName>& RenderGroup : EditLayerRendererState.GetWeightmapRenderGroups())
+			for (const TBitArray<>& WeightmapRenderGroup : EditLayerRendererState.GetWeightmapRenderGroups())
 			{
 				// Only retain the target layers that have passed the previous test (valid and enabled):
-				TBitArray<> RenderGroupBitIndices = MergeRenderContext.ConvertTargetLayerNamesToBitIndices(RenderGroup.Array());
+				TBitArray<> RenderGroupBitIndices = WeightmapRenderGroup;
 				RenderGroupBitIndices.CombineWithBitwiseAND(RendererStateEnabledTargetLayerBitIndices, EBitwiseOperatorFlags::MinSize);
 				if (RenderGroupBitIndices.Find(true) != INDEX_NONE)
 				{
@@ -6730,18 +6750,6 @@ UE::Landscape::EditLayers::FMergeRenderContext ALandscape::PrepareEditLayersMerg
 #endif // !SUPPORTS_LANDSCAPE_EDITORONLY_UBER_MATERIAL
 		}
 	}
-	else
-	{
-		// For heightmap, there's no other dependency so simply trim all renderers that are not enabled : 
-		FinalEditLayerRendererStates = CandidateEditLayerRendererStates.FilterByPredicate([&InMergeRenderParams](const FEditLayerRendererState& InRendererState)
-		{
-			return InRendererState.IsTargetEnabled(ELandscapeToolTargetType::Heightmap);
-		});
-
-		// Only one group in the case of heightmap: 
-		const FName TargetLayerName = FName("Height");
-		FinalRenderGroups = { TBitArray<>(true, 1) };
-	}
 
 	const int32 FinalNumRenderers = FinalEditLayerRendererStates.Num();
 	
@@ -6751,8 +6759,6 @@ UE::Landscape::EditLayers::FMergeRenderContext ALandscape::PrepareEditLayersMerg
 		return MergeRenderContext;
 	}
 
-	const ELandscapeToolTargetTypeFlags MergeTypeMask = InMergeRenderParams.bIsHeightmapMerge ? ELandscapeToolTargetTypeFlags::Heightmap : (ELandscapeToolTargetTypeFlags::Weightmap | ELandscapeToolTargetTypeFlags::Visibility);
-
 	FLandscapeComponent2DIndexer Component2DIndexer = CreateLandscapeComponent2DIndexer(Info);
 	const TArray<ULandscapeComponent*> AllComponents = Component2DIndexer.GetAllValues();
 	const TBitArray<> AllValidComponentBitIndices = Component2DIndexer.GetValidValueBitIndices();
@@ -6794,7 +6800,7 @@ UE::Landscape::EditLayers::FMergeRenderContext ALandscape::PrepareEditLayersMerg
 				const FComponentDependenciesVisLogHelper::EShowNodeInfo ShowNodeInfo = (VisualLogShowComponentDependenciesIndex != INDEX_NONE)
 					? FComponentDependenciesVisLogHelper::EShowNodeInfo::Detailed
 					: static_cast<FComponentDependenciesVisLogHelper::EShowNodeInfo>(VisualLogShowComponentDependencies);
-				VisualLogDependencyHelper.Emplace(this, InMergeRenderParams.bIsHeightmapMerge, ShowNodeInfo, MergeRenderContext);
+				VisualLogDependencyHelper.Emplace(this, MergeRenderContext.IsHeightmapMerge(), ShowNodeInfo, MergeRenderContext);
 			}
 		);
 	auto VisLogDependency = [VisualLogShowComponentDependenciesIndex, &VisualLogDependencyHelper, &AllComponentsToRenderInfos, &OrderedEditLayerRendererRenderInfos]
@@ -6854,19 +6860,22 @@ UE::Landscape::EditLayers::FMergeRenderContext ALandscape::PrepareEditLayersMerg
 			TBitArray<> ComponentsToRenderInfoBitIndicesAfterRenderer = FinalComponentsToRenderInfoBitIndices;
 
 			// Intersect this renderer's render items' outputs with the components to render to find out which ones should participate to the merge :
-			for (const FEditLayerRenderItem& RenderItem : EditLayerRendererRenderInfo.RendererState.GetRenderer()->GetRenderItems(Info))
+			for (const FEditLayerRenderItem& RenderItem : EditLayerRendererRenderInfo.RendererState.GetRenderer()->GetRenderItems(&MergeRenderContext))
 			{
-				const FEditLayerTargetTypeState& RenderItemTargetTypeState = RenderItem.GetTargetTypeState();
-				checkf(RenderItemTargetTypeState.Intersect(EditLayerRendererRenderInfo.RendererState.GetSupportedTargetTypeState()) == RenderItemTargetTypeState,
+				checkf(RenderItem.GetTargetTypeState().Intersect(EditLayerRendererRenderInfo.RendererState.GetSupportedTargetTypeState()) == RenderItem.GetTargetTypeState(),
 					TEXT("One of edit layer renderer %s's render items target type state is not supported by the renderer's supported target type state. \
 					This is invalid : the renderer's supported target type state should be a superset of its render item's target type state. \n\
 					- Render item state: %s \n\
-					- Renderer supported state: %s \n"), *EditLayerRendererDebugName, *RenderItemTargetTypeState.ToString(), *EditLayerRendererRenderInfo.RendererState.GetSupportedTargetTypeState().ToString());
+					- Renderer supported state: %s \n"), *EditLayerRendererDebugName, *RenderItem.GetTargetTypeState().ToString(), *EditLayerRendererRenderInfo.RendererState.GetSupportedTargetTypeState().ToString());
+
+				// Only consider render items which intersect with the renderer's active state (e.g. one a given renderer, a render item (A) might affect heightmaps only and another one (B) weightmaps only, 
+				//  then if performing a heightmap merge, the RendererState's mask here will be ELandscapeToolTargetTypeFlags::Heightmap, so this allows to filter out item B: 
+				const FEditLayerTargetTypeState FilteredRenderItemTargetTypeState = RenderItem.GetTargetTypeState().Intersect(EditLayerRendererRenderInfo.RendererState.GetActiveTargetTypeState());
 
 				bool bOutputRenderItem = false;
 				const FOutputWorldArea& OutputWorldArea = RenderItem.GetOutputWorldArea();
-				// Only consider render items which affect the type of merge being performed: 
-				const bool bRenderItemAffectsMergeType = EnumHasAnyFlags(RenderItemTargetTypeState.GetTargetTypeMask(), MergeTypeMask);
+				
+				const bool bRenderItemAffectsMergeType = EnumHasAnyFlags(FilteredRenderItemTargetTypeState.GetTargetTypeMask(), MergeTypeMask);
 				if (bRenderItemAffectsMergeType)
 				{
 					switch (OutputWorldArea.GetType())
@@ -6925,7 +6934,7 @@ UE::Landscape::EditLayers::FMergeRenderContext ALandscape::PrepareEditLayersMerg
 							if (bVisualLogShowRenderItemsOutput && (bOutputRenderItem || bVisualLogShowAllRenderItems))
 							{
 								TArray<ULandscapeComponent*> AllComponentsToVisLog = bVisualLogShowAllRenderItems ? Component2DIndexer.GetValidValues() : Component2DIndexer.GetValidValuesForBitIndices(FinalComponentsToRenderInfoBitIndices);
-								VisLogRenderItemOutput(this, MergeRenderContext.IsHeightmapMerge(), bAffectsOutputLayerBitIndices, MakeArrayView(RenderItemTargetTypeState.GetActiveWeightmaps()),
+								VisLogRenderItemOutput(this, MergeRenderContext.IsHeightmapMerge(), bAffectsOutputLayerBitIndices, MergeRenderContext.ConvertTargetLayerBitIndicesToNames(FilteredRenderItemTargetTypeState.GetActiveWeightmapBitIndices()),
 									OutputWorldArea, EditLayerRendererRenderInfo, LandscapeTransform, LandscapeLoadedBounds, MakeArrayView(AllComponentsToVisLog));
 							}
 						}
@@ -6989,7 +6998,7 @@ UE::Landscape::EditLayers::FMergeRenderContext ALandscape::PrepareEditLayersMerg
 						}
 						else
 						{
-							OutputLayerBitIndices = MergeRenderContext.ConvertTargetLayerNamesToBitIndices(RenderItemTargetTypeState.GetActiveWeightmaps());
+							OutputLayerBitIndices = FilteredRenderItemTargetTypeState.GetActiveWeightmapBitIndices();
 						}
 					}
 
@@ -7095,7 +7104,7 @@ UE::Landscape::EditLayers::FMergeRenderContext ALandscape::PrepareEditLayersMerg
 			// We'll reuse the same merge render targets in order to generate the mips, which include the duplicate borders, so we need to expand the render target's size to accommodate for this : 
 			MergeRenderBatch.Resolution = (MergeRenderBatch.MaxComponentKey - MergeRenderBatch.MinComponentKey + 1) * NumSubsections * (SubsectionSizeQuads + 1);
 			MergeRenderBatch.TargetLayersToComponents.AddDefaulted(MergeRenderContext.AllTargetLayerNames.Num());
-			MergeRenderBatch.TargetLayerNameBitIndices.SetNum(MergeRenderContext.AllTargetLayerNames.Num(), false);
+			MergeRenderBatch.TargetLayerBitIndices.SetNum(MergeRenderContext.AllTargetLayerNames.Num(), false);
 
 			MergeRenderContext.MaxNeededResolution = MergeRenderContext.MaxNeededResolution.ComponentMax(MergeRenderBatch.Resolution);
 		
@@ -7119,7 +7128,7 @@ UE::Landscape::EditLayers::FMergeRenderContext ALandscape::PrepareEditLayersMerg
 					// Inform the render batch and context of the target layer names associated with each component :
 
 					// Declare these target layers as being in use for the batch :
-					MergeRenderBatch.TargetLayerNameBitIndices.CombineWithBitwiseOR(ComponentToRenderInfo.TargetLayerBitIndices, EBitwiseOperatorFlags::MinSize);
+					MergeRenderBatch.TargetLayerBitIndices.CombineWithBitwiseOR(ComponentToRenderInfo.TargetLayerBitIndices, EBitwiseOperatorFlags::MinSize);
 					MergeRenderBatch.ComponentToTargetLayerBitIndices.FindOrAdd(ComponentToRenderInfo.Component).CombineWithBitwiseOR(ComponentToRenderInfo.TargetLayerBitIndices, EBitwiseOperatorFlags::MaxSize); // Use EBitwiseOperatorFlags::MaxSize here in order to allocate NumTargetLayerNames entries to the resulting bit array in case FindOrAdd is an add
 					MergeRenderContext.ComponentToTargetLayerBitIndices.FindOrAdd(ComponentToRenderInfo.Component).CombineWithBitwiseOR(ComponentToRenderInfo.TargetLayerBitIndices, EBitwiseOperatorFlags::MaxSize); // Use EBitwiseOperatorFlags::MaxSize here in order to allocate NumTargetLayerNames entries to the resulting bit array in case FindOrAdd is an add
 					for (TConstSetBitIterator ItTargetLayer(ComponentToRenderInfo.TargetLayerBitIndices); ItTargetLayer; ++ItTargetLayer)
@@ -7151,7 +7160,7 @@ UE::Landscape::EditLayers::FMergeRenderContext ALandscape::PrepareEditLayersMerg
 				}
 
 				// Finally, signal the group as done for this batch
-				MergeRenderBatch.RenderSteps.Emplace(FMergeRenderStep::EType::SignalBatchMergeGroupDone, FEditLayerRendererState(), RenderGroup, MergeRenderBatch.ComponentsToRender.Array());
+				MergeRenderBatch.RenderSteps.Emplace(FMergeRenderStep::EType::SignalBatchMergeGroupDone, FEditLayerRendererState(&MergeRenderContext), RenderGroup, MergeRenderBatch.ComponentsToRender.Array());
 
 				MergeRenderContext.MaxNeededNumSlices = FMath::Max(MergeRenderContext.MaxNeededNumSlices, RenderGroup.CountSetBits());
 			}
@@ -7460,18 +7469,17 @@ int32 ALandscape::PerformLayersHeightmapsBatchedMerge(const FUpdateLayersContent
 	TRACE_CPUPROFILER_EVENT_SCOPE(ALandscape::PerformLayersHeightmapsBatchedMerge);
 	RHI_BREADCRUMB_EVENT_GAMETHREAD("PerformLayersHeightmapsBatchedMerge");
 
-	ULandscapeInfo* Info = GetLandscapeInfo();
-	check(Info != nullptr);
+	FMergeContext MergeContext(/*InLandscape = */this, /*bInIsHeightmapMerge = */true, InMergeParams.bSkipBrush);
 
 	// Prepare the heightmap merge operations for all components that need to be updated :
-	TArray<FEditLayerRendererState> RendererStates = GetEditLayerRendererStates(Info, InMergeParams.bSkipBrush);
+	TArray<FEditLayerRendererState> RendererStates = GetEditLayerRendererStates(&MergeContext);
 	// Add an edit layer renderer at the top of the stack in order to add dependencies between each landscape component and its immediate neighbors in order to ensure they end up
 	//  in the same render batch. The renderer is responsible for computing the normals at the end of the batch : 
 	ULandscapeHeightmapNormalsEditLayerRenderer* HeightmapNormalsRenderer = ULandscapeHeightmapNormalsEditLayerRenderer::StaticClass()->GetDefaultObject<ULandscapeHeightmapNormalsEditLayerRenderer>();
-	RendererStates.Emplace(HeightmapNormalsRenderer, Info);
+	RendererStates.Emplace(&MergeContext, HeightmapNormalsRenderer);
 
-	FMergeRenderParams MergeRenderParams(/*bInIsHeightmapMerge = */true, /*InLandscape = */this, InUpdateLayersContentContext.LandscapeComponentsHeightmapsToRender, RendererStates);
-	FMergeRenderContext MergeRenderContext = PrepareEditLayersMergeRenderContext(MergeRenderParams);
+	FMergeRenderParams MergeRenderParams(InUpdateLayersContentContext.LandscapeComponentsHeightmapsToRender, RendererStates);
+	FMergeRenderContext MergeRenderContext = PrepareEditLayersMergeRenderContext(MergeContext, MergeRenderParams);
 	if (!MergeRenderContext.IsValid())
 	{
 		// Nothing to do : 
@@ -10333,17 +10341,15 @@ int32 ALandscape::PerformLayersWeightmapsBatchedMerge(FUpdateLayersContentContex
 	TRACE_CPUPROFILER_EVENT_SCOPE(ALandscape::PerformLayersWeightmapsBatchedMerge);
 	RHI_BREADCRUMB_EVENT_GAMETHREAD("PerformLayersWeightmapsBatchedMerge");
 
-	ULandscapeInfo* Info = GetLandscapeInfo();
-	check(Info != nullptr);
-
-	TArray<FEditLayerRendererState> RendererStates = GetEditLayerRendererStates(Info, InMergeParams.bSkipBrush);
+	FMergeContext MergeContext(/*InLandscape = */this, /*bInIsHeightmapMerge = */false, InMergeParams.bSkipBrush);
+	TArray<FEditLayerRendererState> RendererStates = GetEditLayerRendererStates(&MergeContext);
 
 	const TSet<FName> RequestedWeightmapLayerNames(GetTargetLayerNames(/*bInIncludeVisibilityLayer = */ true));
 
 	// Prepare the merge : 
 	TSet<ULandscapeComponent*> LandscapeComponentsWeightmapsToRenderBefore(InUpdateLayersContentContext.LandscapeComponentsWeightmapsToRender);
-	FMergeRenderParams MergeRenderParams(/*bInIsHeightmapMerge = */false, /*InLandscape = */this, InUpdateLayersContentContext.LandscapeComponentsWeightmapsToRender, RendererStates, RequestedWeightmapLayerNames);
-	FMergeRenderContext MergeRenderContext = PrepareEditLayersMergeRenderContext(MergeRenderParams); 
+	FMergeRenderParams MergeRenderParams(InUpdateLayersContentContext.LandscapeComponentsWeightmapsToRender, RendererStates, RequestedWeightmapLayerNames);
+	FMergeRenderContext MergeRenderContext = PrepareEditLayersMergeRenderContext(MergeContext, MergeRenderParams);
 	if (!MergeRenderContext.IsValid())
 	{
 		return InMergeParams.WeightmapUpdateModes;
@@ -13350,37 +13356,10 @@ void FLandscapeLayerBrush::SetOwner(ALandscape* InOwner)
 #endif
 }
 
-bool FLandscapeLayerBrush::Affects(ELandscapeToolTargetType InTargetType, const FName& InWeightmapLayerName) const
-{
-#if WITH_EDITORONLY_DATA
-	if ((BlueprintBrush != nullptr) && BlueprintBrush->IsVisible())
-	{
-		switch (InTargetType)
-		{
-			case ELandscapeToolTargetType::Heightmap:
-			{
-				return BlueprintBrush->AffectsHeightmap();
-			}
-			case ELandscapeToolTargetType::Weightmap:
-			{
-				return BlueprintBrush->AffectsWeightmap() && ((InWeightmapLayerName == NAME_None) || BlueprintBrush->AffectsWeightmapLayer(InWeightmapLayerName));
-			}
-			case ELandscapeToolTargetType::Visibility:
-			{
-				return BlueprintBrush->AffectsVisibilityLayer();
-			}
-			default:
-				check(false);
-		}
-	}
-#endif // WITH_EDITORONLY_DATA
-	return false;
-}
-
 bool FLandscapeLayerBrush::AffectsHeightmap() const
 {
 #if WITH_EDITORONLY_DATA
-	return BlueprintBrush && BlueprintBrush->IsVisible() && BlueprintBrush->AffectsHeightmap();
+	return BlueprintBrush && BlueprintBrush->AffectsHeightmap();
 #else
 	return false;
 #endif
@@ -13389,7 +13368,7 @@ bool FLandscapeLayerBrush::AffectsHeightmap() const
 bool FLandscapeLayerBrush::AffectsWeightmapLayer(const FName& InWeightmapLayerName) const
 {
 #if WITH_EDITORONLY_DATA
-	return BlueprintBrush && BlueprintBrush->IsVisible() && BlueprintBrush->AffectsWeightmap() && BlueprintBrush->AffectsWeightmapLayer(InWeightmapLayerName);
+	return BlueprintBrush && BlueprintBrush->AffectsWeightmapLayer(InWeightmapLayerName);
 #else
 	return false;
 #endif
@@ -13398,7 +13377,7 @@ bool FLandscapeLayerBrush::AffectsWeightmapLayer(const FName& InWeightmapLayerNa
 bool FLandscapeLayerBrush::AffectsVisibilityLayer() const
 {
 #if WITH_EDITORONLY_DATA
-	return BlueprintBrush && BlueprintBrush->IsVisible() && BlueprintBrush->AffectsVisibilityLayer();
+	return BlueprintBrush && BlueprintBrush->AffectsVisibilityLayer();
 #else
 	return false;
 #endif
diff --git a/Engine/Source/Runtime/Landscape/Private/LandscapeUtils.cpp b/Engine/Source/Runtime/Landscape/Private/LandscapeUtils.cpp
--- a/Engine/Source/Runtime/Landscape/Private/LandscapeUtils.cpp
+++ b/Engine/Source/Runtime/Landscape/Private/LandscapeUtils.cpp
@@ -55,7 +55,9 @@ FString GetLandscapeToolTargetTypeFlagsAsString(ELandscapeToolTargetTypeFlags In
 {
 	TArray<FString> TargetTypeStrings;
 	Algo::Transform(MakeFlagsRange(InTargetTypeFlags), TargetTypeStrings, [](ELandscapeToolTargetTypeFlags InTargetTypeFlag) 
-		{ return UEnum::GetValueAsString(GetLandscapeToolTargetTypeSingleFlagAsType(InTargetTypeFlag)); });
+		{ 
+			return UEnum::GetDisplayValueAsText(GetLandscapeToolTargetTypeSingleFlagAsType(InTargetTypeFlag)).ToString(); 
+		});
 	return *FString::Join(TargetTypeStrings, TEXT(","));
 }
 
diff --git a/Engine/Source/Runtime/Landscape/Public/LandscapeBlueprintBrushBase.h b/Engine/Source/Runtime/Landscape/Public/LandscapeBlueprintBrushBase.h
--- a/Engine/Source/Runtime/Landscape/Public/LandscapeBlueprintBrushBase.h
+++ b/Engine/Source/Runtime/Landscape/Public/LandscapeBlueprintBrushBase.h
@@ -111,10 +111,10 @@ public:
 #if WITH_EDITOR
 	//~ Begin ILandscapeEditLayerRenderer implementation
 	LANDSCAPE_EDIT_LAYERS_BATCHED_MERGE_EXPERIMENTAL
-	LANDSCAPE_API virtual void GetRendererStateInfo(const ULandscapeInfo* InLandscapeInfo, 
-		UE::Landscape::EditLayers::FEditLayerTargetTypeState& OutSupportedTargetTypeState, UE::Landscape::EditLayers::FEditLayerTargetTypeState& OutEnabledTargetTypeState, TArray<TSet<FName>>& OutRenderGroups) const override;
+	LANDSCAPE_API virtual void GetRendererStateInfo(const UE::Landscape::EditLayers::FMergeContext* InMergeContext,
+		UE::Landscape::EditLayers::FEditLayerTargetTypeState& OutSupportedTargetTypeState, UE::Landscape::EditLayers::FEditLayerTargetTypeState& OutEnabledTargetTypeState, TArray<TBitArray<>>& OutRenderGroups) const override;
 	LANDSCAPE_EDIT_LAYERS_BATCHED_MERGE_EXPERIMENTAL
-	LANDSCAPE_API virtual TArray<UE::Landscape::EditLayers::FEditLayerRenderItem> GetRenderItems(const ULandscapeInfo* InLandscapeInfo) const override;
+	LANDSCAPE_API virtual TArray<UE::Landscape::EditLayers::FEditLayerRenderItem> GetRenderItems(const UE::Landscape::EditLayers::FMergeContext* InMergeContext) const override;
 	LANDSCAPE_EDIT_LAYERS_BATCHED_MERGE_EXPERIMENTAL
 	LANDSCAPE_API virtual void RenderLayer(ILandscapeEditLayerRenderer::FRenderParams& InRenderParams) override;
 	LANDSCAPE_EDIT_LAYERS_BATCHED_MERGE_EXPERIMENTAL
@@ -122,7 +122,7 @@ public:
 	//~ End ILandscapeEditLayerRenderer implementation
 
 	//~ Begin UE::Landscape::EditLayers::IEditLayerRendererProvider implementation
-	LANDSCAPE_API TArray<UE::Landscape::EditLayers::FEditLayerRendererState> GetEditLayerRendererStates(const ULandscapeInfo* InLandscapeInfo, bool bInSkipBrush);
+	LANDSCAPE_API TArray<UE::Landscape::EditLayers::FEditLayerRendererState> GetEditLayerRendererStates(const UE::Landscape::EditLayers::FMergeContext* InMergeContext);
 	//~ End UE::Landscape::EditLayers::IEditLayerRendererProvider implementation
 
 	UTextureRenderTarget2D* Execute(const FLandscapeBrushParameters& InParameters);
@@ -141,10 +141,10 @@ public:
 	bool CanAffectWeightmap() const { return AffectWeightmap; }
 	bool CanAffectVisibilityLayer() const { return AffectVisibilityLayer; }
 	LANDSCAPE_API virtual bool CanAffectWeightmapLayer(const FName& InLayerName) const;
-	LANDSCAPE_API virtual bool AffectsHeightmap() const { return CanAffectHeightmap(); }
-	LANDSCAPE_API virtual bool AffectsWeightmap() const { return CanAffectWeightmap(); }
+	LANDSCAPE_API virtual bool AffectsHeightmap() const;
+	LANDSCAPE_API virtual bool AffectsWeightmap() const;
 	LANDSCAPE_API virtual bool AffectsWeightmapLayer(const FName& InLayerName) const;
-	LANDSCAPE_API virtual bool AffectsVisibilityLayer() const { return CanAffectVisibilityLayer(); }
+	LANDSCAPE_API virtual bool AffectsVisibilityLayer() const;
 
 	bool IsVisible() const { return bIsVisible; }
 	LANDSCAPE_API bool IsLayerUpdatePending() const;
diff --git a/Engine/Source/Runtime/Landscape/Public/LandscapeEditLayerRenderer.h b/Engine/Source/Runtime/Landscape/Public/LandscapeEditLayerRenderer.h
--- a/Engine/Source/Runtime/Landscape/Public/LandscapeEditLayerRenderer.h
+++ b/Engine/Source/Runtime/Landscape/Public/LandscapeEditLayerRenderer.h
@@ -59,6 +59,232 @@ struct FWeightmapPaintLayerInfo
 // ----------------------------------------------------------------------------------
 
 #if WITH_EDITOR
+
+/**
+ * Global info about the landscape being merged that can be used throughout the merge. 
+ */
+class FMergeContext
+{
+public:
+	friend class ::ALandscape;
+
+	FMergeContext(ALandscape* InLandscape, bool bInIsHeightmapMerge, bool bInSkipProceduralRenderers);
+	virtual ~FMergeContext() = default;
+	FMergeContext(const FMergeContext& Other) = default;
+	FMergeContext(FMergeContext&& Other) = default;
+	FMergeContext& operator=(const FMergeContext& Other) = default;
+	FMergeContext& operator=(FMergeContext&& Other) = default;
+
+	LANDSCAPE_EDIT_LAYERS_BATCHED_MERGE_EXPERIMENTAL
+	inline ALandscape* GetLandscape() const { return Landscape; }
+
+	LANDSCAPE_EDIT_LAYERS_BATCHED_MERGE_EXPERIMENTAL
+	inline ULandscapeInfo* GetLandscapeInfo() const { return LandscapeInfo; }
+	
+	/**
+	 * @return all target layer names declared in this context, including ones that are potentially invalid (i.e. no layer info associated)
+	 */
+	LANDSCAPE_EDIT_LAYERS_BATCHED_MERGE_EXPERIMENTAL
+	inline const TArray<FName>& GetAllTargetLayerNames() const { return AllTargetLayerNames; }
+
+	/**
+	 * @return all target layer names that have a valid layer info object. Will return the height layer in the case of a heightmap merge, even though there's no layer info for those : 
+	 */
+	LANDSCAPE_EDIT_LAYERS_BATCHED_MERGE_EXPERIMENTAL
+	TArray<FName> GetValidTargetLayerNames() const;
+
+	/**
+	 * @return all layer info objects declared in this context (can contain null entries for invalid layers or in the case of a heightmap merge)
+	 */
+	LANDSCAPE_EDIT_LAYERS_BATCHED_MERGE_EXPERIMENTAL
+	inline const TArray<ULandscapeLayerInfoObject*>& GetAllWeightmapLayerInfos() const { return AllWeightmapLayerInfos; }
+
+	/**
+	 * @return true if this merge context is used for a heightmap merge, false otherwise (visibility / weightmap)
+	 */
+	LANDSCAPE_EDIT_LAYERS_BATCHED_MERGE_EXPERIMENTAL
+	inline bool IsHeightmapMerge() const { return bIsHeightmapMerge; }
+
+	/**
+	 * @return true if this merge context is used for a heightmap merge, false otherwise (visibility / weightmap)
+	 */
+	LANDSCAPE_EDIT_LAYERS_BATCHED_MERGE_EXPERIMENTAL
+	inline bool ShouldSkipProceduralRenderers() const { return bSkipProceduralRenderers; }
+
+	/**
+	 * @return true if the target layer name is registered to this merge context and has an associated layer info object
+	 */
+	LANDSCAPE_EDIT_LAYERS_BATCHED_MERGE_EXPERIMENTAL
+	LANDSCAPE_API int32 IsValidTargetLayerName(const FName& InName) const;
+
+	/**
+	 * @return true if the target layer name has an associated layer info object. Asserts if the layer name isn't known to the merge context
+	 */
+	LANDSCAPE_EDIT_LAYERS_BATCHED_MERGE_EXPERIMENTAL
+	LANDSCAPE_API int32 IsValidTargetLayerNameChecked(const FName& InName) const;
+
+	/**
+	 * @return true if the target layer name has an associated layer info object. Asserts if the layer index isn't known to the merge context
+	 */
+	LANDSCAPE_EDIT_LAYERS_BATCHED_MERGE_EXPERIMENTAL
+	LANDSCAPE_API int32 IsValidTargetLayerIndex(int32 InIndex) const;
+
+	/**
+	 * @return the index of this target layer name in this context, INDEX_NONE if not found
+	 */
+	LANDSCAPE_EDIT_LAYERS_BATCHED_MERGE_EXPERIMENTAL
+	LANDSCAPE_API int32 GetTargetLayerIndexForName(const FName& InName) const;
+
+	/**
+	 * @return the index of this target layer name in this context. Asserts if the layer name isn't known to the merge context
+	 */
+	LANDSCAPE_EDIT_LAYERS_BATCHED_MERGE_EXPERIMENTAL
+	LANDSCAPE_API int32 GetTargetLayerIndexForNameChecked(const FName& InName) const;
+
+	/**
+	 * @return the target layer name corresponding to the index in this context, NAME_None if not found
+	 */
+	LANDSCAPE_EDIT_LAYERS_BATCHED_MERGE_EXPERIMENTAL
+	LANDSCAPE_API FName GetTargetLayerNameForIndex(int32 InIndex) const;
+
+	/**
+	 * @return the target layer name corresponding to the index in this context, NAME_None if not found. Asserts if the layer index isn't known to the merge context
+	 */
+	LANDSCAPE_EDIT_LAYERS_BATCHED_MERGE_EXPERIMENTAL
+	LANDSCAPE_API FName GetTargetLayerNameForIndexChecked(int32 InIndex) const;
+
+	/**
+	 * @return the index of this landscape layer info in this context, INDEX_NONE if not found
+	 */
+	LANDSCAPE_EDIT_LAYERS_BATCHED_MERGE_EXPERIMENTAL
+	LANDSCAPE_API int32 GetTargetLayerIndexForLayerInfo(ULandscapeLayerInfoObject* InLayerInfo) const;
+
+	/**
+	 * @return the index of this landscape layer info in this context. Asserts if the layer info isn't known to the merge context
+	 */
+	LANDSCAPE_EDIT_LAYERS_BATCHED_MERGE_EXPERIMENTAL
+	LANDSCAPE_API int32 GetTargetLayerIndexForLayerInfoChecked(ULandscapeLayerInfoObject* InLayerInfo) const;
+
+	/**
+	 * @return the landscape layer object corresponding to the target layer name in this context, nullptr if not found
+	 */
+	LANDSCAPE_EDIT_LAYERS_BATCHED_MERGE_EXPERIMENTAL
+	LANDSCAPE_API ULandscapeLayerInfoObject* GetTargetLayerInfoForName(const FName& InName) const;
+
+	/**
+	 * @return the landscape layer object corresponding to the target layer name in this context. Asserts if the layer name isn't known to the merge context
+	 */
+	LANDSCAPE_EDIT_LAYERS_BATCHED_MERGE_EXPERIMENTAL
+	LANDSCAPE_API ULandscapeLayerInfoObject* GetTargetLayerInfoForNameChecked(const FName& InName) const;
+
+	/**
+	 * @return the landscape layer object corresponding to the index in this context, nullptr if not found. Asserts if the layer index isn't known to the merge context
+	 */
+	LANDSCAPE_EDIT_LAYERS_BATCHED_MERGE_EXPERIMENTAL
+	LANDSCAPE_API ULandscapeLayerInfoObject* GetTargetLayerInfoForIndex(int32 InIndex) const;
+
+	/**
+	 * @return a TBitArray<> where, for each layer from InTargetLayerNames that is found in this context, the corresponding bit is set
+	 */
+	LANDSCAPE_EDIT_LAYERS_BATCHED_MERGE_EXPERIMENTAL
+	LANDSCAPE_API TBitArray<> ConvertTargetLayerNamesToBitIndices(TConstArrayView<FName> InTargetLayerNames) const;
+
+	/**
+	 * @return a TBitArray<> where, for each layer from InTargetLayerNames that is found in this context, the corresponding bit is set. Asserts if one of the layer names isn't known to the merge context
+	 */
+	LANDSCAPE_EDIT_LAYERS_BATCHED_MERGE_EXPERIMENTAL
+	LANDSCAPE_API TBitArray<> ConvertTargetLayerNamesToBitIndicesChecked(TConstArrayView<FName> InTargetLayerNames) const;
+
+	/**
+	 * @return a list of target layer names for each set bit in InTargetLayerBitIndices. Asserts if InTargetLayerBitIndices isn't the same size as AllTargetLayerNames
+	 */
+	LANDSCAPE_EDIT_LAYERS_BATCHED_MERGE_EXPERIMENTAL
+	LANDSCAPE_API TArray<FName> ConvertTargetLayerBitIndicesToNames(const TBitArray<>& InTargetLayerBitIndices) const;
+
+	/**
+	 * @return a list of landscape layer info objects for each set bit in InTargetLayerBitIndices. Asserts if InTargetLayerBitIndices isn't the same size as AllTargetLayerNames
+	 */
+	LANDSCAPE_EDIT_LAYERS_BATCHED_MERGE_EXPERIMENTAL
+	LANDSCAPE_API TArray<ULandscapeLayerInfoObject*> ConvertTargetLayerBitIndicesToLayerInfos(const TBitArray<>& InTargetLayerBitIndices) const;
+
+	/**
+	 * Runs the given function for each all valid target layer in the bit indices in parameters, with the possibility of early exit
+	 * Most easily used with a lambda as follows:
+	 * ForEachTargetLayer([](int32 InTargetLayerIndex, FName InTargetLayerName) -> bool
+	 * {
+	 *     return continueLoop ? true : false;
+	 * });
+	 */
+	LANDSCAPE_EDIT_LAYERS_BATCHED_MERGE_EXPERIMENTAL
+	LANDSCAPE_API void ForEachTargetLayer(const TBitArray<>& InTargetLayerBitIndices, TFunctionRef<bool(int32 /*InTargetLayerIndex*/, const FName& /*InTargetLayerName*/, ULandscapeLayerInfoObject* /*InWeightmapLayerInfo*/)> Fn) const;
+
+	/** Same as ForEachTargetLayer but skips over invalid target layers */
+	LANDSCAPE_EDIT_LAYERS_BATCHED_MERGE_EXPERIMENTAL
+	LANDSCAPE_API void ForEachValidTargetLayer(TFunctionRef<bool(int32 /*InTargetLayerIndex*/, const FName& /*InTargetLayerName*/, ULandscapeLayerInfoObject* /*InWeightmapLayerInfo*/)> Fn) const;
+
+	/**
+	 * @return the index corresponding to the visibility target layer in AllTargetLayerNames, INDEX_NONE if it's not in the list of layers
+	 */
+	inline int32 GetVisibilityTargetLayerIndex() const { return VisibilityTargetLayerIndex; }
+
+	/**
+	 * @return a bit mask the size of AllTargetLayerNames, where all bits are set to zero except the one corresponding to the visibility target layer
+	 */
+	inline const TBitArray<>& GetVisibilityTargetLayerMask() const { return VisibilityTargetLayerMask; }
+
+	/** Same as GetVisibilityTargetLayerMask() but negated */
+	inline const TBitArray<>& GetNegatedVisibilityTargetLayerMask() const { return NegatedVisibilityTargetLayerMask; }
+
+	/**
+	 * @return a bit mask for all target layers from AllTargetLayerNames that have a valid landscape layer info object
+	 */
+	inline const TBitArray<>& GetValidTargetLayerBitIndices() const { return ValidTargetLayerBitIndices; }
+
+	/**
+	 * @return a bit mask the size of AllTargetLayerNames, where all bits are set to bBitValue
+	 */
+	inline TBitArray<> BuildTargetLayerBitIndices(bool bInBitValue) const;
+
+protected:
+	// COMMENT [jonathan.bard] : purposefully don't use ELandscapeToolTargetType here as ELandscapeToolTargetType::Weightmap and ELandscapeToolTargetType::Visibility have to be processed together (because of weightmap packing, 
+	//  which means a visibility weightmap could be another channel of a texture which contains weightmap up to 3 other weightmaps, so we have to resolve the 4 channels altogether). 
+	//  Note: this could change when SUPPORTS_LANDSCAPE_EDITORONLY_UBER_MATERIAL is done...
+	/** Type of merge being requested */
+	bool bIsHeightmapMerge = false;
+
+	/** Allows to skip all non-default (i.e. persistent edit layer) layer renderers for this merge. Mostly there for backwards compatibility and debugging purposes. */
+	bool bSkipProceduralRenderers = false;
+
+	/** Landscape being merged */
+	ALandscape* Landscape = nullptr;
+
+	/** Landscape info associated with the landscape being merged */
+	ULandscapeInfo* LandscapeInfo = nullptr;
+
+	/** List of all target layer names that can possibly be rendered on this landscape (even invalid ones). The index of the target layer in that list is important :
+	  it allows to use a TBitArray instead of a TSet to designate a list of target layers all throughout the merge, which greatly accelerates all the intersection tests being made.
+	  If merging heightmaps, it contains a single name (that is only meaningful for debug display). */
+	TArray<FName> AllTargetLayerNames;
+
+	/** Bit mask that corresponds to only the visibility layer */
+	TBitArray<> VisibilityTargetLayerMask;
+
+	/** Bit mask that corresponds to all target layers except the visibility layer */
+	TBitArray<> NegatedVisibilityTargetLayerMask;
+
+	/** Per-target layer info. Same size as AllTargetLayerNames (only meaningful when rendering weightmaps) */
+	TArray<ULandscapeLayerInfoObject*> AllWeightmapLayerInfos;
+
+	/** List of valid target layers. If a target layer name is present here, it's because it has a valid landscape layer info object. Each bit in that bit array corresponds to an entry in AllTargetLayerNames */
+	TBitArray<> ValidTargetLayerBitIndices;
+
+	// Index of the visibility layer in AllTargetLayerNames (always valid for weightmap merges)
+	int32 VisibilityTargetLayerIndex = INDEX_NONE;
+};
+
+
+// ----------------------------------------------------------------------------------
+
 /** 
  * FEditLayerTargetTypeState fully describes the state of an edit layer renderer wrt its target types. It's named after the enum "ELandscapeToolTargetType" in order to tell
  *  whether the renderer's heightmaps and/or visibility and/or weightmaps are enabled (and if so, which weightmap is enabled exactly)  
@@ -68,11 +294,48 @@ class FEditLayerTargetTypeState
 {
 	friend class ILandscapeEditLayerRenderer;
 
+private:
+	FEditLayerTargetTypeState() = delete;
+
 public:
+	/**
+	 * Constructs an empty FEditLayerTargetTypeState (all target types turned off, no active weightmaps)
+	 * 
+	 * @param InMergeContext context containing all sorts of information related to this merge operation
+	 */
+	LANDSCAPE_EDIT_LAYERS_BATCHED_MERGE_EXPERIMENTAL
+	LANDSCAPE_API FEditLayerTargetTypeState(const FMergeContext* InMergeContext);
+
+	/**
+	 * Initializes a FEditLayerTargetTypeState with the target types passed in parameter 
+	 * 
+	 * @param InMergeContext context containing all sorts of information related to this merge operation
+	 * @param InTargetTypeMask target type mask supported by this target type state (no active weightmaps)
+	 */
+	LANDSCAPE_EDIT_LAYERS_BATCHED_MERGE_EXPERIMENTAL
+	LANDSCAPE_API FEditLayerTargetTypeState(const FMergeContext* InMergeContext, ELandscapeToolTargetTypeFlags InTargetTypeMask);
+
+	/**
+	 * Initializes a FEditLayerTargetTypeState with the target types and the weightmaps passed in parameter (by name)
+	 * 
+	 * @param InMergeContext context containing all sorts of information related to this merge operation
+	 * @param InTargetTypeMask target type mask supported by this target type state
+	 * @param InSupportedWeightmaps names of weightmaps supported by this target type state 
+	 * @param bInChecked, if true, will validate that every weightmap requested is part of the merge context (if false, only the known weightmaps will be marked as supported)
+	 */
 	LANDSCAPE_EDIT_LAYERS_BATCHED_MERGE_EXPERIMENTAL
-	LANDSCAPE_API FEditLayerTargetTypeState() = default;
+	LANDSCAPE_API FEditLayerTargetTypeState(const FMergeContext* InMergeContext, ELandscapeToolTargetTypeFlags InTargetTypeMask, const TArrayView<const FName>& InSupportedWeightmaps, bool bInChecked);
+
+	/**
+	 * Initializes a FEditLayerTargetTypeState with the target types and (optionally) the weightmaps passed in parameter (by bit index)
+	 * 
+	 * @param InMergeContext context containing all sorts of information related to this merge operation
+	 * @param InTargetTypeMask target type mask supported by this target type state
+	 * @param InSupportedWeightmaps bit indices of the weightmaps supported by this target type state (corresponds to the FMergeContext's AllTargetLayerNames)
+	 * Asserts if one of the layer indices isn't valid for the merge context (except when InSupportedWeightmapLayerIndices is == TBitArray<>)
+	 */
 	LANDSCAPE_EDIT_LAYERS_BATCHED_MERGE_EXPERIMENTAL
-	LANDSCAPE_API FEditLayerTargetTypeState(ELandscapeToolTargetTypeFlags InTargetTypeMask, const TArrayView<const FName>& InSupportedWeightmaps = TArrayView<const FName>());
+	LANDSCAPE_API FEditLayerTargetTypeState(const FMergeContext* InMergeContext, ELandscapeToolTargetTypeFlags InTargetTypeMask, const TBitArray<>& InSupportedWeightmapLayerIndices);
 	
 	/**
 	 * Indicates whether a given target type is currently active in this state
@@ -80,12 +343,40 @@ public:
 	 * @param InWeightmapLayerName (optional) is the requested weightmap, only relevant for the ELandscapeToolTargetType::Weightmap case
 	 */
 	LANDSCAPE_EDIT_LAYERS_BATCHED_MERGE_EXPERIMENTAL
-	LANDSCAPE_API bool IsActive(ELandscapeToolTargetType InTargetType, FName InWeightmapLayerName = NAME_None) const;
+	LANDSCAPE_API bool IsActive(ELandscapeToolTargetType InTargetType, const FName& InWeightmapLayerName = NAME_None) const;
+
+		/**
+	 * Indicates whether a given target type is currently active in this state
+	 * @param InTargetType the requested target type (heightmap/weightmap/visibility)
+	 * @param InWeightmapLayerName (optional) is the requested weightmap, only relevant for the ELandscapeToolTargetType::Weightmap case. Asserts if the layer name isn't valid for the merge context (except for NAME_None)
+	 */
+	LANDSCAPE_EDIT_LAYERS_BATCHED_MERGE_EXPERIMENTAL
+	LANDSCAPE_API bool IsActiveChecked(ELandscapeToolTargetType InTargetType, const FName& InWeightmapLayerName = NAME_None) const;
 
-	/** Returns the currently active weightmaps, if Weightmap is amongst the supported target types */
+	/**
+	 * Indicates whether a given target type is currently active in this state
+	 * @param InTargetType is the requested target type (heightmap/weightmap/visibility)
+	 * @param int32 (optional) is the requested weightmap index (in FMergeContex), only relevant for the ELandscapeToolTargetType::Weightmap case. Asserts if the layer index isn't valid for the merge context (except for INDEX_NONE)
+	 */
 	LANDSCAPE_EDIT_LAYERS_BATCHED_MERGE_EXPERIMENTAL
-	LANDSCAPE_API const TArray<FName>& GetActiveWeightmaps() const;
-	
+	LANDSCAPE_API bool IsActive(ELandscapeToolTargetType InTargetType, int32 InWeightmapLayerIndex = INDEX_NONE) const;
+
+	/** Returns the currently active weightmaps : 
+	 *  - If Weightmap is amongst the target types it will return all the weightmaps
+	 *  - Additionally, if Visibility is amongst the target types, it will also return the visibility weightmap
+	 *  - If neither Weightmap nor Visibility is amongst the target types, it will return an empty array
+	 */
+	LANDSCAPE_EDIT_LAYERS_BATCHED_MERGE_EXPERIMENTAL
+	LANDSCAPE_API TArray<FName> GetActiveWeightmaps() const;
+
+	/** Returns the currently active weightmaps :
+	 *  - If Weightmap is amongst the target types it will return all the weightmaps
+	 *  - Additionally, if Visibility is amongst the target types, it will also return the visibility weightmap
+	 *  - If neither Weightmap nor Visibility is amongst the target types, it will return an empty array
+	 */
+	LANDSCAPE_EDIT_LAYERS_BATCHED_MERGE_EXPERIMENTAL
+	LANDSCAPE_API TBitArray<> GetActiveWeightmapBitIndices() const;
+
 	/** Returns the target type mask (i.e. same as ELandscapeToolTargetType, but as bit flags) */
 	LANDSCAPE_EDIT_LAYERS_BATCHED_MERGE_EXPERIMENTAL
 	LANDSCAPE_API ELandscapeToolTargetTypeFlags GetTargetTypeMask() const { return TargetTypeMask; }
@@ -112,11 +403,27 @@ public:
 
 	/** Adds a weightmap to the list of active weightmaps (make sure ELandscapeToolTargetType::Weightmap is amongst the supported target types) */
 	LANDSCAPE_EDIT_LAYERS_BATCHED_MERGE_EXPERIMENTAL
-	LANDSCAPE_API void AddWeightmap(FName InWeightmapLayerName);
+	LANDSCAPE_API void AddWeightmap(const FName& InWeightmapLayerName);
+
+	/** Adds a weightmap to the list of active weightmaps (make sure ELandscapeToolTargetType::Weightmap is amongst the supported target types). Asserts if the layer name isn't known to the merge context.  */
+	LANDSCAPE_EDIT_LAYERS_BATCHED_MERGE_EXPERIMENTAL
+	LANDSCAPE_API void AddWeightmapChecked(const FName& InWeightmapLayerName);
+
+	/** Adds a weightmap to the list of active weightmaps (make sure ELandscapeToolTargetType::Weightmap is amongst the supported target types). Asserts if the layer index isn't valid for the merge context */
+	LANDSCAPE_EDIT_LAYERS_BATCHED_MERGE_EXPERIMENTAL
+	LANDSCAPE_API void AddWeightmap(int32 InWeightmapLayerIndex);
 
 	/** Removes a weightmap from the list of active weightmaps */
 	LANDSCAPE_EDIT_LAYERS_BATCHED_MERGE_EXPERIMENTAL
-	LANDSCAPE_API void RemoveWeightmap(FName InWeightmapLayerName);
+	LANDSCAPE_API void RemoveWeightmap(const FName& InWeightmapLayerName);
+
+	/** Removes a weightmap from the list of active weightmaps. Asserts if the layer name isn't known to the merge context. */
+	LANDSCAPE_EDIT_LAYERS_BATCHED_MERGE_EXPERIMENTAL
+	LANDSCAPE_API void RemoveWeightmapChecked(const FName& InWeightmapLayerName);
+
+	/** Removes a weightmap from the list of active weightmaps. Asserts if the layer index isn't valid for the merge context. */
+	LANDSCAPE_EDIT_LAYERS_BATCHED_MERGE_EXPERIMENTAL
+	LANDSCAPE_API void RemoveWeightmap(int32 InWeightmapLayerIndex);
 
 	/**
 	 * Returns the "intersection" (AND operation) between the target type state and the one in parameter. e.g. if state S0 is:
@@ -131,11 +438,17 @@ public:
 	FString ToString() const;
 
 private:
+	/** Global context being used for this merge : contains generic information about the landscape, its available layer names, etc. */
+	const FMergeContext* MergeContext = nullptr;
+
 	/** Bitmask of the target types that are supported  */
 	ELandscapeToolTargetTypeFlags TargetTypeMask = ELandscapeToolTargetTypeFlags::None;
 	
-	/** List of weightmaps that are supported for the ELandscapeToolTargetType::Weightmap/ELandscapeToolTargetType::Visibility type */
-	TArray<FName> Weightmaps;
+	/** 
+	 * List of weightmaps that are supported for the ELandscapeToolTargetType::Weightmap/ELandscapeToolTargetType::Visibility type. 
+	 *  Each bit in that bit array corresponds to an entry in FMergeContext's AllTargetLayerNames
+	 */
+	TBitArray<> WeightmapTargetLayerBitIndices;
 };
 
 
@@ -150,35 +463,30 @@ private:
  */
 class FEditLayerRendererState
 {
+private:
+	FEditLayerRendererState() = delete;
+
 public:
 	LANDSCAPE_EDIT_LAYERS_BATCHED_MERGE_EXPERIMENTAL
-	LANDSCAPE_API FEditLayerRendererState() = default;
+	LANDSCAPE_API explicit FEditLayerRendererState(const FMergeContext* InMergeContext);
 	LANDSCAPE_EDIT_LAYERS_BATCHED_MERGE_EXPERIMENTAL
-	LANDSCAPE_API explicit FEditLayerRendererState(ILandscapeEditLayerRenderer* InRenderer, const ULandscapeInfo* InLandscapeInfo);
+	LANDSCAPE_API explicit FEditLayerRendererState(const FMergeContext* InMergeContext, ILandscapeEditLayerRenderer* InRenderer);
 
 	/** Returns the edit layer renderer which this state relates to */
 	LANDSCAPE_EDIT_LAYERS_BATCHED_MERGE_EXPERIMENTAL
 	LANDSCAPE_API ILandscapeEditLayerRenderer* GetRenderer() const { return Renderer; }
 
-	/**
-	 * Indicates whether a given target type is currently supported by this renderer
-	 * @param InTargetType the requested target type (heightmap/weightmap/visibility)
-	 * @param InWeightmapLayerName (optional) is the requested weightmap, only relevant for the ELandscapeToolTargetType::Weightmap case
-	 */
-	LANDSCAPE_EDIT_LAYERS_BATCHED_MERGE_EXPERIMENTAL
-	LANDSCAPE_API bool IsTargetSupported(ELandscapeToolTargetType InTargetType, FName InWeightmapLayerName = NAME_None) const;
-
-	/** Returns a mask of all target types supported by this renderer */
+	/** Returns a mask of all target types / weightmaps supported by this renderer */
 	LANDSCAPE_EDIT_LAYERS_BATCHED_MERGE_EXPERIMENTAL
 	LANDSCAPE_API const FEditLayerTargetTypeState& GetSupportedTargetTypeState() const { return SupportedTargetTypeState; }
 
-	/** Returns a mask of all target types supported by this renderer */
+	/** Returns a mask of all target types / weightmaps enabled on this renderer */
 	LANDSCAPE_EDIT_LAYERS_BATCHED_MERGE_EXPERIMENTAL
 	LANDSCAPE_API const FEditLayerTargetTypeState& GetEnabledTargetTypeState() const { return EnabledTargetTypeState; }
 
-	/** Returns a list of all weightmaps supported by this renderer (only relevant for ELandscapeToolTargetType::Weightmap (and ELandscapeToolTargetType::Visibility))*/
+	/** Returns a mask of all target types / weightmaps supported and enabled on this renderer */
 	LANDSCAPE_EDIT_LAYERS_BATCHED_MERGE_EXPERIMENTAL
-	LANDSCAPE_API const TArray<FName>& GetSupportedTargetWeightmaps() const;
+	LANDSCAPE_API FEditLayerTargetTypeState GetActiveTargetTypeState() const { return ActiveTargetTypeState; }
 
 	/** Mutates the EnabledTargetTypeState by adding the target type in parameter to the mask of active target types */
 	LANDSCAPE_EDIT_LAYERS_BATCHED_MERGE_EXPERIMENTAL
@@ -197,32 +505,99 @@ public:
 	LANDSCAPE_API void DisableTargetTypeMask(ELandscapeToolTargetTypeFlags InTargetTypeMask);
 
 	/**
-	 * Indicates whether a given target type is currently enabled by this renderer
+	 * Mutates the EnabledTargetTypeState by adding the weightmap in parameter to the list of enabled weightmaps, if it is known to the merge context
+	 */
+	LANDSCAPE_EDIT_LAYERS_BATCHED_MERGE_EXPERIMENTAL
+	LANDSCAPE_API void EnableWeightmap(ELandscapeToolTargetType InTargetType, const FName& InWeightmapLayerName);
+
+	/**
+	 * Mutates the EnabledTargetTypeState by adding the weightmap in parameter to the list of enabled weightmaps. Asserts if the layer name isn't known to the merge context
+	 */
+	LANDSCAPE_EDIT_LAYERS_BATCHED_MERGE_EXPERIMENTAL
+	LANDSCAPE_API void EnableWeightmapChecked(ELandscapeToolTargetType InTargetType, const FName& InWeightmapLayerName);
+
+	/** 
+	 * Mutates the EnabledTargetTypeState by adding the weightmap in parameter to the list of enabled weightmaps. Asserts if the layer index isn't known to the merge context
+	 */
+	LANDSCAPE_EDIT_LAYERS_BATCHED_MERGE_EXPERIMENTAL
+	LANDSCAPE_API void EnableWeightmap(ELandscapeToolTargetType InTargetType, int32 InWeightmapLayerIndex);
+
+	/**
+	 * Mutates the EnabledTargetTypeState by removing the weightmap in parameter from the list of enabled weightmaps, if it is known to the merge context
+	 */
+	LANDSCAPE_EDIT_LAYERS_BATCHED_MERGE_EXPERIMENTAL
+	LANDSCAPE_API void DisableWeightmap(const FName& InWeightmapLayerName);
+
+	/**
+	 * Mutates the EnabledTargetTypeState by removing the weightmap in parameter from the list of enabled weightmaps. Asserts if the layer index isn't known to the merge context
+	 */
+	LANDSCAPE_EDIT_LAYERS_BATCHED_MERGE_EXPERIMENTAL
+	LANDSCAPE_API void DisableWeightmapChecked(const FName& InWeightmapLayerName);
+
+	/** 
+	 * Mutates the EnabledTargetTypeState by removing the weightmap in parameter from the list of enabled weightmaps. Asserts if the layer index isn't known to the merge context
+	 */
+	LANDSCAPE_EDIT_LAYERS_BATCHED_MERGE_EXPERIMENTAL
+	LANDSCAPE_API void DisableWeightmap(int32 InWeightmapLayerIndex);
+
+	/** Returns the mask of target types that are both supported and enabled by this renderer */
+
+	LANDSCAPE_EDIT_LAYERS_BATCHED_MERGE_EXPERIMENTAL
+	LANDSCAPE_API ELandscapeToolTargetTypeFlags GetActiveTargetTypeMask() const;
+
+	/**
+	 * Indicates whether a given target type and weightmap layer name is currently supported and enabled by this renderer. 
 	 * @param InTargetType the requested target type (heightmap/weightmap/visibility)
 	 * @param InWeightmapLayerName (optional) is the requested weightmap, only relevant for the ELandscapeToolTargetType::Weightmap case
 	 */
 	LANDSCAPE_EDIT_LAYERS_BATCHED_MERGE_EXPERIMENTAL
-	LANDSCAPE_API bool IsTargetEnabled(ELandscapeToolTargetType InTargetType, FName InWeightmapLayerName = NAME_None) const;
+	LANDSCAPE_API bool IsTargetActive(ELandscapeToolTargetType InTargetType, const FName& InWeightmapLayerName = NAME_None) const;
 
-	/** Mutates the EnabledTargetTypeState by adding the weightmap in parameter to the list of enabled weightmaps */
+	/**
+	 * Indicates whether a given target type and weightmap layer name is currently supported and enabled by this renderer. Asserts if the layer name isn't known to the merge context (except if NAME_None)
+	 * @param InTargetType the requested target type (heightmap/weightmap/visibility)
+	 * @param InWeightmapLayerName (optional) is the requested weightmap, only relevant for the ELandscapeToolTargetType::Weightmap case
+	 */
+	LANDSCAPE_EDIT_LAYERS_BATCHED_MERGE_EXPERIMENTAL
+	LANDSCAPE_API bool IsTargetActiveChecked(ELandscapeToolTargetType InTargetType, const FName& InWeightmapLayerName = NAME_None) const;
+
+	/**
+	 * Indicates whether a given target type and weightmap layer name is currently supported and enabled by this renderer. Asserts if the layer index isn't known to the merge context (except if INDEX_NONE)
+	 * @param InTargetType the requested target type (heightmap/weightmap/visibility)
+	 * @param InWeightmapLayerIndex (optional) is the requested weightmap's index in the merge index, only relevant for the ELandscapeToolTargetType::Weightmap case
+	 */
 	LANDSCAPE_EDIT_LAYERS_BATCHED_MERGE_EXPERIMENTAL
-	LANDSCAPE_API bool EnableTarget(ELandscapeToolTargetType InTargetType, FName InWeightmapLayerName = NAME_None);
+	LANDSCAPE_API bool IsTargetActive(ELandscapeToolTargetType InTargetType, int32 InWeightmapLayerIndex = INDEX_NONE) const;
 
-	/** Mutates the EnabledTargetTypeState by removing the weightmap in parameter from the list of enabled weightmaps */
+	/**
+	 * Returns a list of all weightmaps supported and enabled by this renderer (only relevant for ELandscapeToolTargetType::Weightmap (and ELandscapeToolTargetType::Visibility))
+	 */
 	LANDSCAPE_EDIT_LAYERS_BATCHED_MERGE_EXPERIMENTAL
-	LANDSCAPE_API void DisableTarget(FName InWeightmapLayerName);
+	LANDSCAPE_API TArray<FName> GetActiveTargetWeightmaps() const;
 
-	/** Returns the list of all weightmaps currently enabled */
+	/** 
+	 * Returns a list of all weightmaps supported and enabled by this renderer (only relevant for ELandscapeToolTargetType::Weightmap (and ELandscapeToolTargetType::Visibility))
+	 */
 	LANDSCAPE_EDIT_LAYERS_BATCHED_MERGE_EXPERIMENTAL
-	LANDSCAPE_API TArray<FName> GetEnabledTargetWeightmaps() const;
+	LANDSCAPE_API TBitArray<> GetActiveTargetWeightmapBitIndices() const;
 
 	/** Returns the render groups associated with this renderer. A render group is a set of target layers (weightmaps) that depend on one another in order to produce the output target layers.
 	 This allows to implement "horizontal blending", where weightmaps can be blended with one another at each step of the landscape edit layers merge algorithm */
 	LANDSCAPE_EDIT_LAYERS_BATCHED_MERGE_EXPERIMENTAL
-	LANDSCAPE_API TArray<TSet<FName>> GetWeightmapRenderGroups() const { return RenderGroups; }
+	LANDSCAPE_API TArray<TBitArray<>> GetWeightmapRenderGroups() const { return RenderGroups; }
 
 private:
+	void UpdateActiveTargetTypeState();
+
+private:
+	/** Global context being used for this merge : contains generic information about the landscape, its available layer names, etc. */
+	const FMergeContext* MergeContext = nullptr;
+
+	/** Renderer associated with this state */
 	ILandscapeEditLayerRenderer* Renderer = nullptr;
+
+	/** For debug purposes : this is the same as Renderer->GetEditLayerRendererDebugName() but having a member makes it easier to debug in the watch window */
+	FString DebugName;
 	
 	/** Struct that contains the supported target types and the which weightmap are supported by this renderer. Immutable. */
 	FEditLayerTargetTypeState SupportedTargetTypeState;
@@ -230,8 +605,13 @@ private:
 	/** Struct that contains the enabled target types and the which weightmap are currently enabled by this renderer. Can be set by the user using EnableTargetType. */
 	FEditLayerTargetTypeState EnabledTargetTypeState;
 
-	/* List of groups of target layers that this renderer requires to be rendered together. All target layers listed in SupportedTargetTypeState must belong to 1 (and 1 only) render group. */
-	TArray<TSet<FName>> RenderGroups;
+	/** Intersection of SupportedTargetTypeState and EnabledTargetTypeState (defines what is both supported and enabled on this renderer) */
+	FEditLayerTargetTypeState ActiveTargetTypeState;
+
+	/* List of groups of target layers that this renderer requires to be rendered together. All target layers listed in SupportedTargetTypeState must belong to 1 (and 1 only) render group.
+	 *  Each render group is a bit array for which each bit corresponds to an entry in FMergeContext's AllTargetLayerNames
+	 */
+	TArray<TBitArray<>> RenderGroups;
 };
 
 
@@ -244,16 +624,7 @@ private:
  */
 struct FMergeRenderParams
 {
-	FMergeRenderParams(bool bInIsHeightmapMerge, ALandscape* InLandscape, TArray<ULandscapeComponent*> InComponentsToMerge, const TArrayView<FEditLayerRendererState>& InEditLayerRendererStates, const TSet<FName>& InWeightmapLayerNames = {});
-
-	// COMMENT [jonathan.bard] : purposefully don't use ELandscapeToolTargetType here as ELandscapeToolTargetType::Weightmap and ELandscapeToolTargetType::Visibility have to be processed together (because of weightmap packing, 
-	//  which means a visibility weightmap could be another channel of a texture which contains weightmap up to 3 other weightmaps, so we have to resolve the 4 channels altogether). 
-	//  Note: this could change when SUPPORTS_LANDSCAPE_EDITORONLY_UBER_MATERIAL is done...
-	/** Type of merge being requested */
-	bool bIsHeightmapMerge = true;
-	
-	/** Parent landscape actor to which all components to merge belong */
-	ALandscape* Landscape = nullptr;
+	FMergeRenderParams(TArray<ULandscapeComponent*> InComponentsToMerge, const TArrayView<FEditLayerRendererState>& InEditLayerRendererStates, const TSet<FName>& InWeightmapLayerNames = {});
 
 	/** List of components that need merging */
 	TArray<ULandscapeComponent*> ComponentsToMerge;
@@ -290,7 +661,7 @@ struct FMergeRenderStep
 	/** Renderer state to be used this step (only when Type == EType::RenderLayer). This includes the renderer as well as its precise step (e.g. which weightmap are supported? which are enabled?) */
 	FEditLayerRendererState RendererState;
 	
-	/** List of target layers being involved in this step. Each bit in that bit array corresponds to an entry in FMergeRenderContext's AllTargetLayerNames */
+	/** List of target layers being involved in this step. Each bit in that bit array corresponds to an entry in FMergeContext's AllTargetLayerNames */
 	TBitArray<> RenderGroupBitIndices;
 
 	/** List of components involved in this step */
@@ -327,10 +698,10 @@ struct FMergeRenderBatch
 	TSet<ULandscapeComponent*> ComponentsToRender;
 
 	/** List of all target layers being rendered in this batch (i.e. bitwise OR of all of the render steps' RenderGroupBitIndices). Each bit in that bit array corresponds to an entry in FMergeRenderContext's AllTargetLayerNames  */
-	TBitArray<> TargetLayerNameBitIndices;
+	TBitArray<> TargetLayerBitIndices;
 
 	/** List of components involved in this batch and the target layers they're writing to (redundant with ComponentsToRender but we keep the latter for convenience)
-	 (each bit corresponds to a target layer name in FMergeRenderContext's AllTargetLayerNames) */
+	 (each bit corresponds to a target layer name in FMergeContext's AllTargetLayerNames) */
 	TMap<ULandscapeComponent*, TBitArray<>> ComponentToTargetLayerBitIndices;
 
 	/** Reverse lookup of ComponentToTargetLayerBitIndices : one entry per element in ComponentToTargetLayerBitIndices, each entry containing all of the components involved in this merge for this target layer */
@@ -393,12 +764,12 @@ struct FComponentMergeRenderInfo
 /** 
  * Utility class that contains everything necessary to perform the batched merge : scratch render targets, list of batches, etc.
  */
-class FMergeRenderContext
+class FMergeRenderContext : public FMergeContext
 {
 public:
 	friend class ::ALandscape;
 
-	FMergeRenderContext(ALandscape* InLandscape, bool bInIsHeightmapMerge);
+	FMergeRenderContext(const FMergeContext& InMergeContext);
 	virtual ~FMergeRenderContext();
 	FMergeRenderContext(const FMergeRenderContext& Other) = default;
 	FMergeRenderContext(FMergeRenderContext&& Other) = default;
@@ -459,11 +830,6 @@ public:
 	};
 	void Render(TFunction<void(const FOnRenderBatchTargetGroupDoneParams&)> OnRenderBatchTargetGroupDone);
 
-	LANDSCAPE_EDIT_LAYERS_BATCHED_MERGE_EXPERIMENTAL
-	inline ALandscape* GetLandscape() const { return Landscape; }
-	LANDSCAPE_EDIT_LAYERS_BATCHED_MERGE_EXPERIMENTAL
-	inline bool IsHeightmapMerge() const { return bIsHeightmapMerge; }
-
 	LANDSCAPE_EDIT_LAYERS_BATCHED_MERGE_EXPERIMENTAL
 	inline FIntPoint GetMaxNeededResolution() const { return MaxNeededResolution; }
 	LANDSCAPE_EDIT_LAYERS_BATCHED_MERGE_EXPERIMENTAL
@@ -491,35 +857,6 @@ public:
 	/** Duplicates the vertex data from the (sub-)sections of the batch, assuming GetBlendRenderTargetRead() is the RT that is read from and GetBlendRenderTargetWrite() the one that is written to */
 	void RenderExpandedRenderTarget(const FMergeRenderBatch& InRenderBatch);
 
-	const TArray<FName>& GetAllTargetLayerNames() const { return AllTargetLayerNames; }
-	int32 GetTargetLayerIndexForName(const FName& InName) const;
-	int32 GetTargetLayerIndexForNameChecked(const FName& InName) const;
-	FName GetTargetLayerNameForIndex(int32 InIndex) const;
-	FName GetTargetLayerNameForIndexChecked(int32 InIndex) const;
-	TBitArray<> ConvertTargetLayerNamesToBitIndices(TConstArrayView<FName> InTargetLayerNames) const;
-	TBitArray<> ConvertTargetLayerNamesToBitIndicesChecked(TConstArrayView<FName> InTargetLayerNames) const;
-	LANDSCAPE_EDIT_LAYERS_BATCHED_MERGE_EXPERIMENTAL
-	LANDSCAPE_API TArray<FName> ConvertTargetLayerBitIndicesToNames(const TBitArray<>& InTargetLayerBitIndices) const;
-	LANDSCAPE_EDIT_LAYERS_BATCHED_MERGE_EXPERIMENTAL
-	LANDSCAPE_API TArray<FName> ConvertTargetLayerBitIndicesToNamesChecked(const TBitArray<>& InTargetLayerBitIndices) const;
-	LANDSCAPE_EDIT_LAYERS_BATCHED_MERGE_EXPERIMENTAL
-	LANDSCAPE_API TArray<ULandscapeLayerInfoObject*> ConvertTargetLayerBitIndicesToLayerInfos(const TBitArray<>& InTargetLayerBitIndices) const;
-	LANDSCAPE_EDIT_LAYERS_BATCHED_MERGE_EXPERIMENTAL
-	LANDSCAPE_API TArray<ULandscapeLayerInfoObject*> ConvertTargetLayerBitIndicesToLayerInfosChecked(const TBitArray<>& InTargetLayerBitIndices) const;
-
-	/**
-	 * Runs the given function for each all valid target layer in the bit indices in parameters, with the possibility of early exit
-	 * Most easily used with a lambda as follows:
-	 * ForEachTargetLayer([](int32 InTargetLayerIndex, FName InTargetLayerName) -> bool
-	 * {
-	 *     return continueLoop ? true : false;
-	 * });
-	 */
-	LANDSCAPE_EDIT_LAYERS_BATCHED_MERGE_EXPERIMENTAL
-	LANDSCAPE_API void ForEachTargetLayer(const TBitArray<>& InTargetLayerBitIndices, TFunctionRef<bool(int32 /*InTargetLayerIndex*/, FName /*InTargetLayerName*/)> Fn) const;
-	LANDSCAPE_EDIT_LAYERS_BATCHED_MERGE_EXPERIMENTAL
-	LANDSCAPE_API void ForEachTargetLayerChecked(const TBitArray<>& InTargetLayerBitIndices, TFunctionRef<bool(int32 /*InTargetLayerIndex*/, FName /*InTargetLayerName*/)> Fn) const;
-
 	const TBitArray<>& GetFinalTargetLayerBitIndices() const 
 	{ 
 		return FinalTargetLayerBitIndices; 
@@ -543,23 +880,6 @@ private:
 	/** Render targets that are used throughout the blending operations (they could be texture arrays in the case of multiple weightmaps) */
 	TStaticArray<ULandscapeScratchRenderTarget*, NumBlendRenderTargets> BlendRenderTargets;
 	int32 CurrentBlendRenderTargetWriteIndex = -1; 
-		
-	/** The landscape on which the merge occurs*/
-	ALandscape* Landscape = nullptr;
-
-	/** Type of merge being requested */
-	bool bIsHeightmapMerge = false;
-
-	/** List of all target layer names that can possibly be rendered on this landscape (even invalid ones). The index of the target layer in that list is important : 
-	  it allows to use a TBitArray instead of a TSet to designate a list of target layers all throughout the merge, which greatly accelerates all the intersection tests being made.
-	  If merging heightmaps, it contains a single name (that is only meaningful for debug display). */
-	TArray<FName> AllTargetLayerNames;
-
-	/** Per-target layer info. Same size as AllTargetLayerNames (only meaningful when rendering weightmaps) */
-	TArray<ULandscapeLayerInfoObject*> WeightmapLayerInfos;
-
-	/** List of valid target layers. If a target layer name is present here, it's because it has a valid landscape layer info object.Each bit in that bit array corresponds to an entry in AllTargetLayerNames */
-	TBitArray<> ValidTargetLayerBitIndices;
 
 	/** Final list of target layer names being involved in this merge context. If a target layer name is present here, it's because it's a valid target layer and it needs to be rendered
 	  because it has been requested or one of the target layers that have been requested needs it to be present (e.g. weight-blending). Each bit in that bit array corresponds to an entry in AllTargetLayerNames */
@@ -776,14 +1096,12 @@ public:
 	virtual ~IEditLayerRendererProvider() {}
 
 	/**
-	 * GetEditLayerRendererStates returns a list of renderers that this provider can provide, along with their current state. The state can then be manipulated at will by the user to enable/disable renderers selectively
-	 *
-	 * @param InLandscapeInfo ULandscapeInfo containing information about the landscape affected by this renderer
+	 * @param InMergeContext context containing all sorts of information related to this merge operation
 	 *
 	 * @return a list of renderer states (i.e. a ILandscapeEditLayerRenderer and its current state) to be processed in that order by the merge operation
 	 */
 	LANDSCAPE_EDIT_LAYERS_BATCHED_MERGE_EXPERIMENTAL
-	LANDSCAPE_API virtual TArray<FEditLayerRendererState> GetEditLayerRendererStates(const ULandscapeInfo* InLandscapeInfo, bool bInSkipBrush) PURE_VIRTUAL(IEditLayerRendererProvider::GetEditLayerRendererStates, return {}; );
+	LANDSCAPE_API virtual TArray<FEditLayerRendererState> GetEditLayerRendererStates(const FMergeContext* InMergeContext) PURE_VIRTUAL(IEditLayerRendererProvider::GetEditLayerRendererStates, return {}; );
 };
 #endif // WITH_EDITOR
 
@@ -861,7 +1179,7 @@ public:
 	 *  A target type must be both supported and enabled in order to have this renderer affect it and the "enabled" state can be changed at will by the user (e.g. to temporarily disable 
 	 *  a given edit layer just for the duration of the merge) : see FEditLayerRendererState
 	 * 
-	 * @param InLandscapeInfo ULandscapeInfo containing information about the landscape affected by this renderer
+	 * @param InMergeContext context containing all sorts of information related to this merge operation
 	 * 
 	 * @param OutSupportedTargetTypeState List of all target types / weightmaps that this renderer supports (e.g. if, say ELandscapeToolTargetType::Weightmap, is not included, the renderer
 	 *  will *not* be used at all when rendering any kind of weightmap)
@@ -876,8 +1194,8 @@ public:
 	 *  and if this renderer doesn't act on one of these layers, it will simply do nothing with it in its RenderLayer function
 	 */
 	LANDSCAPE_EDIT_LAYERS_BATCHED_MERGE_EXPERIMENTAL
-	LANDSCAPE_API virtual void GetRendererStateInfo(const ULandscapeInfo* InLandscapeInfo,
-		UE::Landscape::EditLayers::FEditLayerTargetTypeState& OutSupportedTargetTypeState, UE::Landscape::EditLayers::FEditLayerTargetTypeState& OutEnabledTargetTypeState, TArray<TSet<FName>>& OutRenderGroups) const
+	LANDSCAPE_API virtual void GetRendererStateInfo(const UE::Landscape::EditLayers::FMergeContext* InMergeContext,
+		UE::Landscape::EditLayers::FEditLayerTargetTypeState& OutSupportedTargetTypeState, UE::Landscape::EditLayers::FEditLayerTargetTypeState& OutEnabledTargetTypeState, TArray<TBitArray<>>& OutRenderGroups) const
 		PURE_VIRTUAL(ILandscapeEditLayerRenderer::GetRendererStateInfo, );
 
 	/**
@@ -890,12 +1208,12 @@ public:
 	/**
 	 * GetRenderItems retrieves information about the areas this renderer renders to and specifically what respective input area they require to render properly
 	 *
-	 * @param InLandscapeInfo ULandscapeInfo containing information about the landscape affected by this renderer
+	 * @param InMergeContext context containing all sorts of information related to this merge operation
 	 *
 	 * @return list of all render items that affect this renderer
 	*/
 	LANDSCAPE_EDIT_LAYERS_BATCHED_MERGE_EXPERIMENTAL
-	LANDSCAPE_API virtual TArray<UE::Landscape::EditLayers::FEditLayerRenderItem> GetRenderItems(const ULandscapeInfo* InLandscapeInfo) const 
+	LANDSCAPE_API virtual TArray<UE::Landscape::EditLayers::FEditLayerRenderItem> GetRenderItems(const UE::Landscape::EditLayers::FMergeContext* InMergeContext) const
 		PURE_VIRTUAL(ILandscapeEditLayerRenderer::GetRenderItems, return { }; );
 
 	/**
diff --git a/Engine/Source/Runtime/Landscape/Public/LandscapeEditResourcesSubsystem.h b/Engine/Source/Runtime/Landscape/Public/LandscapeEditResourcesSubsystem.h
--- a/Engine/Source/Runtime/Landscape/Public/LandscapeEditResourcesSubsystem.h
+++ b/Engine/Source/Runtime/Landscape/Public/LandscapeEditResourcesSubsystem.h
@@ -115,7 +115,10 @@ public:
 	/** @return the internal render target's effective number of slices (can be different than the actual number of slices if CurrentRenderTargetParams.bExactDimensions is false) while this render target is in use */
 	LANDSCAPE_EDIT_LAYERS_BATCHED_MERGE_EXPERIMENTAL
 	LANDSCAPE_API int32 GetEffectiveNumSlices() const;
-	bool IsTexture2DArray() const { return (TryGetRenderTarget2DArray() != nullptr); }
+	LANDSCAPE_EDIT_LAYERS_BATCHED_MERGE_EXPERIMENTAL
+	LANDSCAPE_API bool IsTexture2D() const { return (TryGetRenderTarget2D() != nullptr); }
+	LANDSCAPE_EDIT_LAYERS_BATCHED_MERGE_EXPERIMENTAL
+	LANDSCAPE_API bool IsTexture2DArray() const { return (TryGetRenderTarget2DArray() != nullptr); }
 	LANDSCAPE_EDIT_LAYERS_BATCHED_MERGE_EXPERIMENTAL
 	LANDSCAPE_API ETextureRenderTargetFormat GetFormat() const;
 	LANDSCAPE_EDIT_LAYERS_BATCHED_MERGE_EXPERIMENTAL
