Mostly fix landscape corruption when smooth tool used with nanite landscape.

diff --git a/Engine/Source/Runtime/Landscape/Private/LandscapeEditLayers.cpp b/Engine/Source/Runtime/Landscape/Private/LandscapeEditLayers.cpp
--- a/Engine/Source/Runtime/Landscape/Private/LandscapeEditLayers.cpp
+++ b/Engine/Source/Runtime/Landscape/Private/LandscapeEditLayers.cpp
@@ -8264,12 +8264,18 @@ bool ALandscape::ResolveLayersTexture(
 					NewHash = ULandscapeTextureHash::CalculateTextureHash64(NewMipData, MipTexels, TextureType);
 
 					// If the last readback bypassed the hash test because it was an intermediate render, force a lock and copy to restore the "correct" non-intermediate data.  The regular
-					// significance test and conditional notification can proceed as normal.
+					// conditional notification can proceed as normal, based on the stored hash.
 					bool bForceUpdate = false;
 					if (InCPUReadback->GetLastReadbackWasIntermediate())
 					{
 						bForceUpdate = true;
 						InCPUReadback->SetLastReadbackWasIntermediate(false);
+
+						if (ChangedChannelsMask.IsSet())
+						{
+							// The previous data isn't available to compare against. Assume the worst case: all channels changed.
+							*ChangedChannelsMask = 0xf;
+						}
 					}
 
 					if (NewHash != OldHash || bForceUpdate)
@@ -8279,7 +8285,7 @@ bool ALandscape::ResolveLayersTexture(
 						const FColor* OldMipDataReadOnly = (const FColor*) InOutputTexture->Source.LockMipReadOnly(MipIndex);
 						check(OldMipDataReadOnly);
 
-						if (ULandscapeTextureHash::DoesTextureDataChangeExceedThreshold(OldMipDataReadOnly, NewMipData, MipTexels, bIsWeightmap ? ELandscapeTextureType::Weightmap : ELandscapeTextureType::Heightmap, OldHash, NewHash, ChangedChannelsMask))
+						if (bForceUpdate || ULandscapeTextureHash::DoesTextureDataChangeExceedThreshold(OldMipDataReadOnly, NewMipData, MipTexels, bIsWeightmap ? ELandscapeTextureType::Weightmap : ELandscapeTextureType::Heightmap, OldHash, NewHash, ChangedChannelsMask))
 						{
 							// convert the lock to a read/write lock (first have to release the existing read-only lock, then acquire a new read/write lock)
 							InOutputTexture->Source.UnlockMip(MipIndex);
@@ -8288,26 +8294,28 @@ bool ALandscape::ResolveLayersTexture(
 							check(!bLockedForWrite);
 							bLockedForWrite = true;
 
-							InCPUReadback->SetHash(NewHash);
-							bChanged = true;
-
-							// We're about to modify the texture's source data, the texture needs to know so that it can handle properly update cached platform data (additionally, the package needs to be dirtied) :
-							ULandscapeInfo* LandscapeInfo = GetLandscapeInfo();
-							if (GetDefault<ULandscapeSettings>()->LandscapeDirtyingMode == ELandscapeDirtyingMode::InLandscapeModeAndUserTriggeredChanges)
+							// Note that the hash might not change in the bForceUpdate case.  We skip the dirty notice if it hasn't changed.
+							bChanged |= InCPUReadback->SetHash(NewHash);
+							if (bChanged)
 							{
-								FLandscapeDirtyOnlyInModeScope Scope(LandscapeInfo, !bUserTriggered);
-								LandscapeInfo->ModifyObject(InOutputTexture);
-								if (ULandscapeTextureHash* TextureHash = InOutputTexture->GetAssetUserData<ULandscapeTextureHash>())
+								// We're about to modify the texture's source data, the texture needs to know so that it can handle properly update cached platform data (additionally, the package needs to be dirtied) :
+								ULandscapeInfo* LandscapeInfo = GetLandscapeInfo();
+								if (GetDefault<ULandscapeSettings>()->LandscapeDirtyingMode == ELandscapeDirtyingMode::InLandscapeModeAndUserTriggeredChanges)
 								{
-									LandscapeInfo->ModifyObject(TextureHash);
+									FLandscapeDirtyOnlyInModeScope Scope(LandscapeInfo, !bUserTriggered);
+									LandscapeInfo->ModifyObject(InOutputTexture);
+									if (ULandscapeTextureHash* TextureHash = InOutputTexture->GetAssetUserData<ULandscapeTextureHash>())
+									{
+										LandscapeInfo->ModifyObject(TextureHash);
+									}
 								}
-							}
-							else
-							{
-								LandscapeInfo->ModifyObject(InOutputTexture);
-								if (ULandscapeTextureHash* TextureHash = InOutputTexture->GetAssetUserData<ULandscapeTextureHash>())
+								else
 								{
-									LandscapeInfo->ModifyObject(TextureHash);
+									LandscapeInfo->ModifyObject(InOutputTexture);
+									if (ULandscapeTextureHash* TextureHash = InOutputTexture->GetAssetUserData<ULandscapeTextureHash>())
+									{
+										LandscapeInfo->ModifyObject(TextureHash);
+									}
 								}
 							}
 						}
