Added Landscape Spline scale and z-offset variables configurable in project settings to improve Spline icon visibility

diff --git a/Engine/Source/Runtime/Landscape/Private/LandscapeSplines.cpp b/Engine/Source/Runtime/Landscape/Private/LandscapeSplines.cpp
--- a/Engine/Source/Runtime/Landscape/Private/LandscapeSplines.cpp
+++ b/Engine/Source/Runtime/Landscape/Private/LandscapeSplines.cpp
@@ -49,6 +49,7 @@
 #include "LandscapeLayerInfoObject.h"
 #include "LandscapeInfoMap.h"
 #include "LandscapeSplineActor.h"
+#include "LandscapeSettings.h"
 #endif
 
 IMPLEMENT_HIT_PROXY(HLandscapeSplineProxy, HHitProxy);
@@ -72,6 +73,13 @@ static FAutoConsoleVariableRef CVarLandscapeSplineToSplineComponentMaxIterations
 	TEXT("Max possible iterations when converting a landscape spline into a spline component")
 );
 
+#if WITH_EDITOR
+static FAutoConsoleVariable CVarLandscapeSplineEnableFixedSpriteIcons(
+	TEXT("landscape.Splines.EnableFixedSpriteIcons"),
+	true,
+	TEXT("Enables fixed scaling and Z-offset for spline icons. If disabled, configure dynamic scale and offset in project settings."));
+#endif // WITH_EDITOR 
+
 //////////////////////////////////////////////////////////////////////////
 // LANDSCAPE SPLINES SCENE PROXY
 
@@ -157,6 +165,8 @@ private:
 		TRefCountPtr<HHitProxy> HitProxy;
 		FVector Location;
 		TArray<FLandscapeSplineInterpPoint> Points;
+		// Controls sprite scaling based on segment width, may be removed in favor of project settings as sole scale factor
+		// Scale usage toggled by CVar command landscape.Splines.EnableFixedSpriteIcons
 		float SpriteScale;
 		uint32 bSelected : 1;
 	};
@@ -225,7 +235,7 @@ public:
 		const FMatrix& MyLocalToWorld = GetLocalToWorld();
 
 		const FLinearColor SelectedSplineColor = GEngine->GetSelectedMaterialColor();
-		const FLinearColor SelectedControlPointSpriteColor = FLinearColor::White + (GEngine->GetSelectedMaterialColor() * GEngine->SelectionHighlightIntensityBillboards * 10); 
+		const FLinearColor SelectedControlPointSpriteColor = FLinearColor::White + (GEngine->GetSelectedMaterialColor() * GEngine->SelectionHighlightIntensityBillboards * 20);
 
 		for (int32 ViewIndex = 0; ViewIndex < Views.Num(); ViewIndex++)
 		{
@@ -287,21 +297,35 @@ public:
 					// Draw Sprite
 					if (bDrawControlPointSprite)
 					{
-						double ControlPointSpriteScale = MyLocalToWorld.GetScaleVector().X * ControlPoint.SpriteScale;
+						// By default use values from Landscape Settings for scale and z-offset
+						const ULandscapeSettings* Settings = GetDefault<ULandscapeSettings>();
+						bool bEnableFixedSpriteIcons = CVarLandscapeSplineEnableFixedSpriteIcons->GetBool();
+
+						float ControlPointSpriteZOffset = Settings->GetSplineIconWorldZOffset();
+						float ControlPointSpriteScale = bEnableFixedSpriteIcons ? ControlPoint.SpriteScale : Settings->GetSplineIconScale();
+						FVector ControlPointSpriteLocation = ControlPointLocation + FVector(0, 0, ControlPointSpriteZOffset);
+					
+						if (bEnableFixedSpriteIcons)
+						{
+							ControlPointSpriteScale *= MyLocalToWorld.GetScaleVector().X;
+							const FMatrix& ProjectionMatrix = View->ViewMatrices.GetProjectionMatrix();
+							const double ZoomFactor = FMath::Min<double>(ProjectionMatrix.M[0][0], ProjectionMatrix.M[1][1]);
+							const double Scale = View->WorldToScreen(ControlPointLocation).W * (4.0f / View->UnscaledViewRect.Width() / ZoomFactor);
+							ControlPointSpriteScale *= Scale;
+						}
+						
 						const FLinearColor ControlPointSpriteColor = ControlPoint.bSelected ? SelectedControlPointSpriteColor : FLinearColor::White;
 						PDI->SetHitProxy(ControlPoint.HitProxy);
-						const FMatrix& ProjectionMatrix = View->ViewMatrices.GetProjectionMatrix();
-						const double ZoomFactor = FMath::Min<double>(ProjectionMatrix.M[0][0], ProjectionMatrix.M[1][1]);
-						const double Scale = View->WorldToScreen(ControlPointLocation).W * (4.0f / View->UnscaledViewRect.Width() / ZoomFactor);
-						ControlPointSpriteScale *= Scale;
-
 #if WITH_EDITORONLY_DATA
 						ControlPointSpriteScale *= UBillboardComponent::EditorScale;
 #endif
+						if (bEnableFixedSpriteIcons)
+						{
+							// Clamping the scale between 10 and the ControlPoint initial scale 
+							ControlPointSpriteScale = FMath::Clamp<double>(ControlPointSpriteScale, 10, ControlPoint.SpriteScale);
+							ControlPointSpriteLocation = ControlPointLocation + FVector(0, 0, ControlPointSpriteScale * 0.75f);
+						}
 
-						// Clamping the scale between 10 and the ControlPoint initial scale 
-						ControlPointSpriteScale = FMath::Clamp<double>(ControlPointSpriteScale, 10, ControlPoint.SpriteScale);
-						const FVector ControlPointSpriteLocation = ControlPointLocation + FVector(0, 0, ControlPointSpriteScale * 0.75f);
 						PDI->DrawSprite(
 							ControlPointSpriteLocation,
 							static_cast<float>(ControlPointSpriteScale),
diff --git a/Engine/Source/Runtime/Landscape/Public/LandscapeSettings.h b/Engine/Source/Runtime/Landscape/Public/LandscapeSettings.h
--- a/Engine/Source/Runtime/Landscape/Public/LandscapeSettings.h
+++ b/Engine/Source/Runtime/Landscape/Public/LandscapeSettings.h
@@ -54,6 +54,9 @@ public:
 
 	int32 GetHLODMaxTextureSize() const { return HLODMaxTextureSize; }
 
+	float GetSplineIconWorldZOffset() const { return SplineIconWorldZOffset; }
+	float GetSplineIconScale() const { return SplineIconScale;  }
+
 	/** Returns the default landscape material that should be used when creating a new landscape. */
 	TSoftObjectPtr<UMaterialInterface> GetDefaultLandscapeMaterial() const { return DefaultLandscapeMaterial; }
 
@@ -110,5 +113,12 @@ protected:
 
 	UPROPERTY(transient)
 	bool bRestrictiveMode = false;
+
+	/** Offset in Z for the landscape spline icon in world-space. */
+	UPROPERTY(EditAnywhere, config, Category = "Spline", meta = (ClampMin = "0", UIMax = "4096"))
+	float SplineIconWorldZOffset = 75.0f;
+
+	UPROPERTY(EditAnywhere, config, Category = "Spline", meta = (ClampMin = "32", UIMax = "2048"))
+	float SplineIconScale = 50.0f;
 };
 
