Flush pending work related to landscape layers and physical material when saving.

diff --git a/Engine/Source/Runtime/Landscape/Classes/Landscape.h b/Engine/Source/Runtime/Landscape/Classes/Landscape.h
--- a/Engine/Source/Runtime/Landscape/Classes/Landscape.h
+++ b/Engine/Source/Runtime/Landscape/Classes/Landscape.h
@@ -485,6 +485,7 @@ public:
 	
 	LANDSCAPE_API void ToggleCanHaveLayersContent();
 	LANDSCAPE_API void ForceUpdateLayersContent(bool bIntermediateRender = false);
+	void FlushLayerContentThisFrame();
 
 	UFUNCTION(BlueprintCallable, Category = "Landscape")
 	LANDSCAPE_API void ForceLayersFullUpdate();
@@ -749,6 +750,8 @@ private:
 	
 	// Used in packing the material layer data contained into CombinedLayersWeightmapAllMaterialLayersResource to be set again for each component weightmap (size of the landscape)
 	class FLandscapeTexture2DResource* WeightmapScratchPackLayerTextureResource;
+
+	uint32 LastFlushedLayerUpdateFrame = 0;
 #endif
 };
 
diff --git a/Engine/Source/Runtime/Landscape/Private/Landscape.cpp b/Engine/Source/Runtime/Landscape/Private/Landscape.cpp
--- a/Engine/Source/Runtime/Landscape/Private/Landscape.cpp
+++ b/Engine/Source/Runtime/Landscape/Private/Landscape.cpp
@@ -3746,6 +3746,13 @@ void ALandscapeProxy::PreSave(FObjectPreSaveContext ObjectSaveContext)
 #if WITH_EDITOR
 	if (!ObjectSaveContext.IsProceduralSave()) // only finalize grass in a true editor save (when a GPU is available).
 	{
+		ALandscape* Landscape = GetLandscapeActor();
+		if (Landscape)
+		{
+			check(ObjectSaveContext.IsFirstConcurrentSave());  //if PreSave ever actually becomes concurrent, this will need some change to make it safe.
+			Landscape->FlushLayerContentThisFrame();
+		}
+
 		// It would be nice to strip grass data at editor save time to reduce asset size on disk.
 		// Unfortunately we can't easily know if there is a platform out there that may need to use the serialized grass map path.
 		// And future cook processes may not have a GPU available to build the grass data themselves.
@@ -3794,11 +3801,16 @@ void ALandscapeProxy::PreSave(FObjectPreSaveContext ObjectSaveContext)
 
 	for (ULandscapeComponent* LandscapeComponent : LandscapeComponents)
 	{
-		// Warn if there are any active Physical Material Renders on this Proxy
-		// don.boogert todo: We should block or wait for this to complete within some timeout period.
-		if (LandscapeComponent->PhysicalMaterialTask.IsInProgress())
+		if (LandscapeComponent->CanUpdatePhysicalMaterial() && LandscapeComponent->PhysicalMaterialTask.IsValid())
 		{
-			UE_LOG(LogLandscape, Warning, TEXT("Physical material render on component: '%s' in progress."), *LandscapeComponent->GetFullName());
+			UE_LOG(LogLandscape, Display, TEXT("Completing landscape physical material before saving.  %s"), *LandscapeComponent->GetFullName());
+			if (LandscapeComponent->PhysicalMaterialTask.IsInProgress())
+			{
+				LandscapeComponent->PhysicalMaterialTask.Flush();
+			}
+			check(LandscapeComponent->PhysicalMaterialTask.IsComplete());
+			// IsInProgress tests the render thread status.  Now finish the last steps on the game thread.
+			LandscapeComponent->UpdatePhysicalMaterialTasks();
 		}
 
 		// Ensure the component's cached bounds are correct
@@ -3846,8 +3858,8 @@ void ALandscapeProxy::PreSave(FObjectPreSaveContext ObjectSaveContext)
 		const bool bIsStreamingProxy = this->IsA<ALandscapeStreamingProxy>();
 		const bool bUseEdgeFixups = bIsStreamingProxy;
 		const bool bUpdateSnapshotNow = true;
-			for (ULandscapeComponent* LandscapeComponent : LandscapeComponents)
-			{
+		for (ULandscapeComponent* LandscapeComponent : LandscapeComponents)
+		{
 			LandscapeComponent->InstallOrUpdateTextureUserDatas(bUseEdgeFixups, bShouldCompressHeightmap, bUpdateSnapshotNow);
 		}
 	}
diff --git a/Engine/Source/Runtime/Landscape/Private/LandscapeEditLayers.cpp b/Engine/Source/Runtime/Landscape/Private/LandscapeEditLayers.cpp
--- a/Engine/Source/Runtime/Landscape/Private/LandscapeEditLayers.cpp
+++ b/Engine/Source/Runtime/Landscape/Private/LandscapeEditLayers.cpp
@@ -11535,11 +11535,14 @@ void ALandscape::UpdateLayersContent(bool bInWaitForStreaming, bool bInSkipMonit
 	// Poll and complete any outstanding resolve work
 	// If bIntermediateRender then we want to flush all work here before we do the intermediate render later on
 	// if bFlushRender then we skip this because we will flush later anyway
-	if (bProcessReadbacks && (bIntermediateRender || !bFlushRender))
+	if (bProcessReadbacks)
 	{
 		// These flags might look like they're being mixed up but they're not!
 		const bool bDoIntermediateRender = false; // bIntermediateRender flag is for the work queued up this frame not the delayed resolves
-		const bool bDoFlushRender = bIntermediateRender; // Flush before we do an intermediate render later in this frame
+		// bIntermediateRender needs to flush before the intermediate render happens.  If flush is requested without intermediate render, this function could
+		// quit early because LayerContentUpdateModes==0, so this is still the place to flush.  Flushing twice would only happen in the rare case when flush
+		// is requested, pre-existing readbacks are pending (bProcessReadbacks), and new work is also being initiated via LayerContentUpdateModes.
+		const bool bDoFlushRender = bIntermediateRender || bFlushRender;
 
 		// Flushing once all readback tasks is much faster than asking each to do it so start by doing just this :
 		if (bDoFlushRender)
@@ -11776,7 +11779,20 @@ void ALandscape::InitializeLayers(bool bUseNormalCapture)
 
 void ALandscape::OnPreSave()
 {
-	ForceUpdateLayersContent();
+	// Note:: This is only called if the outer level is saved.
+	FlushLayerContentThisFrame();
+}
+
+// If any work is pending in the layer update system, wait for it to complete.  Won't flush more than once per frame.
+void ALandscape::FlushLayerContentThisFrame()
+{
+	// Only call ForceUpdateLayersContent once per frame. Each proxy might trigger this, so only do it for the first one.
+	uint32 CurrentFrame = GFrameNumber;
+	if (LastFlushedLayerUpdateFrame != CurrentFrame)
+	{
+		LastFlushedLayerUpdateFrame = CurrentFrame;
+		ForceUpdateLayersContent(/*bIntermediateRender = */ false);
+	}
 }
 
 void ALandscape::ForceUpdateLayersContent(bool bIntermediateRender)
