Restore reverted assert and fix conflict with PIE tests.  If certain tests run without loading their level in advance, they can start PIE without first ticking normally.
UpdateLayersContent is blocked from running on the normal Editor world while a PIE world is active.  Change UpdatePhysicalMaterialTasks to match that condition.

diff --git a/Engine/Source/Runtime/Landscape/Private/LandscapeEdit.cpp b/Engine/Source/Runtime/Landscape/Private/LandscapeEdit.cpp
--- a/Engine/Source/Runtime/Landscape/Private/LandscapeEdit.cpp
+++ b/Engine/Source/Runtime/Landscape/Private/LandscapeEdit.cpp
@@ -2268,6 +2268,9 @@ void ULandscapeComponent::UpdatePhysicalMaterialTasks()
 		// Rendering the physical material in this state will silently produce the wrong result.
 		bool bWeightMapsReady = Algo::AllOf(WeightmapTextures, [](UTexture2D* Weightmap)
 			{
+				// We expect FLandscapeTextureStreamingManager to have already started managing the data lifetime of the weightmaps, via UpdateLayersContent.
+				check(!Weightmap || Weightmap->bForceMiplevelsToBeResident);
+
 				return Weightmap && !Weightmap->IsDefaultTexture() && !Weightmap->HasPendingInitOrStreaming() && Weightmap->IsFullyStreamedIn();
 			});
 
diff --git a/Engine/Source/Runtime/Landscape/Private/LandscapeSubsystem.cpp b/Engine/Source/Runtime/Landscape/Private/LandscapeSubsystem.cpp
--- a/Engine/Source/Runtime/Landscape/Private/LandscapeSubsystem.cpp
+++ b/Engine/Source/Runtime/Landscape/Private/LandscapeSubsystem.cpp
@@ -826,18 +826,14 @@ void ULandscapeSubsystem::Tick(float DeltaTime)
 	for (ALandscapeProxy* Proxy : ActiveProxies)
 	{
 #if WITH_EDITOR
-		if (GIsEditor)
+		if (GIsEditor && !World->IsPlayInEditor() && GEditor->PlayWorld == nullptr)
 		{
 			if (ALandscape* Landscape = Cast<ALandscape>(Proxy))
 			{
 				Landscape->TickLayers(DeltaTime);
 			}
 
-			// editor-only
-			if (!World->IsPlayInEditor())
-			{
-				Proxy->UpdatePhysicalMaterialTasks();
-			}
+			Proxy->UpdatePhysicalMaterialTasks();
 		}
 
 		Proxy->GetAsyncWorkMonitor().Tick(DeltaTime);
