Fixed regression with legacy weight-blending not being applied when using landscape edit layers batched merge

diff --git a/Engine/Source/Runtime/Landscape/Private/LandscapeEditLayerRendererPrivate.cpp b/Engine/Source/Runtime/Landscape/Private/LandscapeEditLayerRendererPrivate.cpp
--- a/Engine/Source/Runtime/Landscape/Private/LandscapeEditLayerRendererPrivate.cpp
+++ b/Engine/Source/Runtime/Landscape/Private/LandscapeEditLayerRendererPrivate.cpp
@@ -105,7 +105,7 @@ void ULandscapeWeightmapWeightBlendedLayersRenderer::GetRendererStateInfo(const
 	TBitArray<> WeightBlendedWeightmapLayerBitIndices = GatherWeightBlendedWeightmapLayerBitIndices(InMergeContext);
 	// Only relevant for weightmaps :
 	FEditLayerTargetTypeState OutputTargetTypeState(InMergeContext, ELandscapeToolTargetTypeFlags::Weightmap, WeightBlendedWeightmapLayerBitIndices);
-	OutEnabledTargetTypeState = OutSupportedTargetTypeState;
+	OutEnabledTargetTypeState = OutSupportedTargetTypeState = OutputTargetTypeState;
 
 	// Now fill in the render groups : one group for all weight-blended layers : 
 	if (WeightBlendedWeightmapLayerBitIndices.Contains(true))
