* Refactored the landscape edit layer renderer API to be able to run via a minimal amount of render commands (and FRDGBuilder)
* Fixed crash with water brush when running with batched merge due to an uninitialized BP brush render target size

Misc:
* This change introduces a generic FRDGBuilderRecorder which is a "RDG-based render command recorder" that can be used in a similar fashion to standard render commands, except that the recorded commands are run on the render thread, using a single render command / FRDGBuilder.
** It has 2 modes : Immediate (issues render commands as soon as they're enqueued) / Recording (records the commands until Flush is called, a render command is then issued and all recorded command, executed)
** RDG_RENDER_COMMAND_RECORDER_BREADCRUMB_EVENT can be used to tag the recorded commands appropriately into single entry
** It can enqueue RDG-based commands or standard render commands
** RDG-based commands also support being passed a list of external texture with the RHIAccess they should have when executing the FRDGBuilder. This is to prevent the RDG from auto-transitioning to SRVMask at the end
* Landscape edit layer renderers can now return render flags. RenderMode_Recorded tells the merge algorithm it supports FRDGBuilderRecorder while RenderMode_Immediate requires all previous recorded commands to be flushed and will therefore enqueue one or several standard render commands. Currently only BP brushes (and a special case in landscape patches) requires this, which means a landscape with no such brush can run within a single FDGBuilder command

diff --git a/Engine/Plugins/Experimental/LandscapePatch/Source/LandscapePatch/Private/LandscapeCircleHeightPatch.cpp b/Engine/Plugins/Experimental/LandscapePatch/Source/LandscapePatch/Private/LandscapeCircleHeightPatch.cpp
--- a/Engine/Plugins/Experimental/LandscapePatch/Source/LandscapePatch/Private/LandscapeCircleHeightPatch.cpp
+++ b/Engine/Plugins/Experimental/LandscapePatch/Source/LandscapePatch/Private/LandscapeCircleHeightPatch.cpp
@@ -45,26 +45,31 @@ void ULandscapeCircleHeightPatch::OnComponentCreated()
 
 UTextureRenderTarget2D* ULandscapeCircleHeightPatch::RenderLayer_Native(const FLandscapeBrushParameters& InParameters, const FTransform& HeightmapCoordsToWorld)
 {
+	using namespace UE::Landscape;
+
 	// Circle height patch doesn't affect regular weightmap layers.
 	if ((bEditVisibility && (InParameters.LayerType != ELandscapeToolTargetType::Visibility)) || (!bEditVisibility && (InParameters.LayerType != ELandscapeToolTargetType::Heightmap)))
 	{
 		return InParameters.CombinedResult;
 	}
 
-	ApplyCirclePatch(InParameters.LayerType == ELandscapeToolTargetType::Visibility,
-		[CombinedResult = InParameters.CombinedResult] { return CombinedResult->GetResource()->GetTexture2DRHI(); }, 0, 
+	// We render in immediate mode for the legacy path so we use an immediate recorder :  
+	FRDGBuilderRecorder RDGBuilderRecorderImmediate;
+	ApplyCirclePatch(RDGBuilderRecorderImmediate, InParameters.LayerType == ELandscapeToolTargetType::Visibility,
+		InParameters.CombinedResult->GetResource(), 0,
 		FIntPoint(InParameters.CombinedResult->SizeX, InParameters.CombinedResult->SizeY),
-		HeightmapCoordsToWorld);
+		HeightmapCoordsToWorld, /*OutputAccess = */ERHIAccess::SRVMask); // in the legacy path the RT leaves as a SRV
 
 	return InParameters.CombinedResult;
 }
 
 #if WITH_EDITOR
-void ULandscapeCircleHeightPatch::RenderLayer(FRenderParams& InRenderParams)
+void ULandscapeCircleHeightPatch::RenderLayer(FRenderParams& InRenderParams, UE::Landscape::FRDGBuilderRecorder& RDGBuilderRecorder)
 {
 	using namespace UE::Landscape::PatchUtil;
 	using namespace UE::Landscape::EditLayers;
 
+	checkf(RDGBuilderRecorder.IsRecording(), TEXT("ERenderFlags::RenderMode_Recorded means the command recorder should be recording at this point"));
 	FTransform HeightmapCoordsToWorld = GetHeightmapToWorld(InRenderParams.RenderAreaWorldTransform);
 
 	ULandscapeScratchRenderTarget* LandscapeRT = InRenderParams.MergeRenderContext->GetBlendRenderTargetWrite();
@@ -75,10 +80,7 @@ void ULandscapeCircleHeightPatch::RenderLayer(FRenderParams& InRenderParams)
 		UTextureRenderTarget2D* OutputToBlendInto = LandscapeRT->TryGetRenderTarget2D();
 		if (ensure(OutputToBlendInto))
 		{
-			// graph builder expects external textures to start as SRV
-			LandscapeRT->TransitionTo(ERHIAccess::SRVMask); 
-
-			ApplyCirclePatch(false, [OutputToBlendInto] { return OutputToBlendInto->GetResource()->GetTexture2DRHI(); },
+			ApplyCirclePatch(RDGBuilderRecorder, /*bIsVisibilityLayer = */false, OutputToBlendInto->GetResource(),
 				0, InRenderParams.RenderAreaSectionRect.Size(), HeightmapCoordsToWorld);
 		}
 		return;
@@ -92,35 +94,28 @@ void ULandscapeCircleHeightPatch::RenderLayer(FRenderParams& InRenderParams)
 	}
 
 	UTextureRenderTarget2DArray* TextureArray = LandscapeRT->TryGetRenderTarget2DArray();
-	if (!ensure(TextureArray))
-	{
-		return;
-	}
+	check(TextureArray != nullptr);
 
 	int32 NumLayers = InRenderParams.RenderGroupTargetLayerInfos.Num();
-	if (!ensure(TextureArray->Slices == NumLayers))
-	{
-		NumLayers = FMath::Min(NumLayers, TextureArray->Slices);
-	}
+	check(TextureArray->Slices == NumLayers);
 
 	for (int32 LayerIndex = 0; LayerIndex < NumLayers; ++LayerIndex)
 	{
 		if (UE::Landscape::IsVisibilityLayer(InRenderParams.RenderGroupTargetLayerInfos[LayerIndex]))
 		{
-			// graph builder expects external textures to start as SRV
-			LandscapeRT->TransitionTo(ERHIAccess::SRVMask);
-
-			ApplyCirclePatch(true, [TextureArray] { return TextureArray->GetResource()->GetTexture2DArrayRHI(); }, 
+			ApplyCirclePatch(RDGBuilderRecorder, /*bIsVisibilityLayer = */true, TextureArray->GetResource(),
 				LayerIndex, InRenderParams.RenderAreaSectionRect.Size(), HeightmapCoordsToWorld);
 		}
 	}
 }
 #endif // WITH_EDITOR
 
-void ULandscapeCircleHeightPatch::ApplyCirclePatch(bool bIsVisibilityLayer,
-	const TFunction<FRHITexture* ()>& RenderThreadLandscapeTextureGetter, int32 LandscapeTextureSliceIndex,
-	const FIntPoint& DestinationResolution, const FTransform & HeightmapCoordsToWorld)
+void ULandscapeCircleHeightPatch::ApplyCirclePatch(UE::Landscape::FRDGBuilderRecorder& RDGBuilderRecorder, bool bIsVisibilityLayer,
+	FTextureResource* InMergedLandscapeTextureResource, int32 LandscapeTextureSliceIndex,
+	const FIntPoint& DestinationResolution, const FTransform & HeightmapCoordsToWorld, ERHIAccess OutputAccess)
 {
+	using namespace UE::Landscape;
+
 	if (bEditVisibility != bIsVisibilityLayer)
 	{
 		return;
@@ -151,59 +146,59 @@ void ULandscapeCircleHeightPatch::ApplyCirclePatch(bool bIsVisibilityLayer,
 		return;
 	}
 
-	ENQUEUE_RENDER_COMMAND(LandscapeCircleHeightPatch)([RenderThreadLandscapeTextureGetter, LandscapeTextureSliceIndex, DestinationBounds,
-		CenterInHeightmapCoordinates, HeightmapRadius, HeightmapFalloff, bEditVisibility = bEditVisibility](FRHICommandListImmediate& RHICmdList)
-	{
-		TRACE_CPUPROFILER_EVENT_SCOPE(LandscapeCircleHeightPatch);
-
-		const TCHAR* OutputName = bEditVisibility ? TEXT("LandscapeCircleVisibilityPatchOutput") : TEXT("LandscapeCircleHeightPatchOutput");
-		const TCHAR* InputCopyName = bEditVisibility ? TEXT("LandscapeCircleVisibilityPatchInputCopy") : TEXT("LandscapeCircleHeightPatchInputCopy");
-
-		FRDGBuilder GraphBuilder(RHICmdList, RDG_EVENT_NAME("ApplyLandscapeCirclePatch"));
-		
-		TRefCountPtr<IPooledRenderTarget> RenderTarget = CreateRenderTarget(RenderThreadLandscapeTextureGetter(), OutputName);
-		FRDGTextureRef DestinationTexture = GraphBuilder.RegisterExternalTexture(RenderTarget);
-
-		// Make a copy of the portion of our weightmap input that we're writing to so that we can 
-		//  read and write at the same time (needed for blending)
-		FRDGTextureDesc InputCopyDescription = DestinationTexture->Desc;
-		InputCopyDescription.Dimension = ETextureDimension::Texture2D;
-		InputCopyDescription.ArraySize = 1;
-		InputCopyDescription.NumMips = 1;
-		InputCopyDescription.Extent = DestinationBounds.Size();
-		FRDGTextureRef InputCopy = GraphBuilder.CreateTexture(InputCopyDescription, InputCopyName);
-
-		FRHICopyTextureInfo CopyTextureInfo;
-		CopyTextureInfo.SourceMipIndex = 0;
-		CopyTextureInfo.NumMips = 1;
-		CopyTextureInfo.SourceSliceIndex = LandscapeTextureSliceIndex;
-		CopyTextureInfo.NumSlices = 1;
-		CopyTextureInfo.SourcePosition = FIntVector(DestinationBounds.Min.X, DestinationBounds.Min.Y, 0);
-		CopyTextureInfo.Size = FIntVector(InputCopyDescription.Extent.X, InputCopyDescription.Extent.Y, 0);
-		AddCopyTexturePass(GraphBuilder, DestinationTexture, InputCopy, CopyTextureInfo);
-
-		FRDGTextureSRVRef InputCopySRV = GraphBuilder.CreateSRV(FRDGTextureSRVDesc::CreateForMipLevel(InputCopy, 0));
-
-		FLandscapeCircleHeightPatchPS::FParameters* ShaderParams = GraphBuilder.AllocParameters<FLandscapeCircleHeightPatchPS::FParameters>();
-		ShaderParams->InCenter = (FVector3f)CenterInHeightmapCoordinates;
-		ShaderParams->InRadius = HeightmapRadius;
-		ShaderParams->InFalloff = HeightmapFalloff;
-		ShaderParams->InSourceTexture = InputCopySRV;
-		ShaderParams->InSourceTextureOffset = DestinationBounds.Min;
-		ShaderParams->RenderTargets[0] = FRenderTargetBinding(DestinationTexture, ERenderTargetLoadAction::ENoAction, 
-			/*InMipIndex = */0, LandscapeTextureSliceIndex);
-
-		if (bEditVisibility)
+	auto RDGCommand =
+		[InMergedLandscapeTextureResource, LandscapeTextureSliceIndex, DestinationBounds,
+		CenterInHeightmapCoordinates, HeightmapRadius, HeightmapFalloff, bEditVisibility = bEditVisibility](FRDGBuilder& GraphBuilder)
 		{
-			FLandscapeCircleVisibilityPatchPS::AddToRenderGraph(GraphBuilder, ShaderParams, DestinationBounds);
-		}
-		else
-		{
-			FLandscapeCircleHeightPatchPS::AddToRenderGraph(GraphBuilder, ShaderParams, DestinationBounds);
-		}
-
-		GraphBuilder.Execute();
-	});
+			TRACE_CPUPROFILER_EVENT_SCOPE(LandscapeCircleHeightPatch);
+
+			const TCHAR* OutputName = bEditVisibility ? TEXT("LandscapeCircleVisibilityPatchOutput") : TEXT("LandscapeCircleHeightPatchOutput");
+			const TCHAR* InputCopyName = bEditVisibility ? TEXT("LandscapeCircleVisibilityPatchInputCopy") : TEXT("LandscapeCircleHeightPatchInputCopy");
+
+			TRefCountPtr<IPooledRenderTarget> RenderTarget = CreateRenderTarget(InMergedLandscapeTextureResource->GetTextureRHI(), OutputName);
+			FRDGTextureRef DestinationTexture = GraphBuilder.RegisterExternalTexture(RenderTarget);
+
+			// Make a copy of the portion of our weightmap input that we're writing to so that we can 
+			//  read and write at the same time (needed for blending)
+			FRDGTextureDesc InputCopyDescription = DestinationTexture->Desc;
+			InputCopyDescription.Dimension = ETextureDimension::Texture2D;
+			InputCopyDescription.ArraySize = 1;
+			InputCopyDescription.NumMips = 1;
+			InputCopyDescription.Extent = DestinationBounds.Size();
+			FRDGTextureRef InputCopy = GraphBuilder.CreateTexture(InputCopyDescription, InputCopyName);
+
+			FRHICopyTextureInfo CopyTextureInfo;
+			CopyTextureInfo.SourceMipIndex = 0;
+			CopyTextureInfo.NumMips = 1;
+			CopyTextureInfo.SourceSliceIndex = LandscapeTextureSliceIndex;
+			CopyTextureInfo.NumSlices = 1;
+			CopyTextureInfo.SourcePosition = FIntVector(DestinationBounds.Min.X, DestinationBounds.Min.Y, 0);
+			CopyTextureInfo.Size = FIntVector(InputCopyDescription.Extent.X, InputCopyDescription.Extent.Y, 0);
+			AddCopyTexturePass(GraphBuilder, DestinationTexture, InputCopy, CopyTextureInfo);
+
+			FRDGTextureSRVRef InputCopySRV = GraphBuilder.CreateSRV(FRDGTextureSRVDesc::CreateForMipLevel(InputCopy, 0));
+
+			FLandscapeCircleHeightPatchPS::FParameters* ShaderParams = GraphBuilder.AllocParameters<FLandscapeCircleHeightPatchPS::FParameters>();
+			ShaderParams->InCenter = (FVector3f)CenterInHeightmapCoordinates;
+			ShaderParams->InRadius = HeightmapRadius;
+			ShaderParams->InFalloff = HeightmapFalloff;
+			ShaderParams->InSourceTexture = InputCopySRV;
+			ShaderParams->InSourceTextureOffset = DestinationBounds.Min;
+			ShaderParams->RenderTargets[0] = FRenderTargetBinding(DestinationTexture, ERenderTargetLoadAction::ENoAction,
+				/*InMipIndex = */0, LandscapeTextureSliceIndex);
+
+			if (bEditVisibility)
+			{
+				FLandscapeCircleVisibilityPatchPS::AddToRenderGraph(GraphBuilder, ShaderParams, DestinationBounds);
+			}
+			else
+			{
+				FLandscapeCircleHeightPatchPS::AddToRenderGraph(GraphBuilder, ShaderParams, DestinationBounds);
+			}
+		};
+
+	// We need to specify the final state of the external texture to prevent the graph builder from transitioning it to SRVMask :
+	RDGBuilderRecorder.EnqueueRDGCommand(RDGCommand, { { InMergedLandscapeTextureResource, (OutputAccess == ERHIAccess::None) ? ERHIAccess::RTV : OutputAccess} });
 }
 
 #if WITH_EDITOR
diff --git a/Engine/Plugins/Experimental/LandscapePatch/Source/LandscapePatch/Private/LandscapeTexturePatch.cpp b/Engine/Plugins/Experimental/LandscapePatch/Source/LandscapePatch/Private/LandscapeTexturePatch.cpp
--- a/Engine/Plugins/Experimental/LandscapePatch/Source/LandscapePatch/Private/LandscapeTexturePatch.cpp
+++ b/Engine/Plugins/Experimental/LandscapePatch/Source/LandscapePatch/Private/LandscapeTexturePatch.cpp
@@ -124,8 +124,28 @@ namespace LandscapeTexturePatchLocals
 }
 
 #if WITH_EDITOR
-void ULandscapeTexturePatch::RenderLayer(FRenderParams& InRenderParams)
+UE::Landscape::EditLayers::ERenderFlags ULandscapeTexturePatch::GetRenderFlags(const UE::Landscape::EditLayers::FMergeContext* InMergeContext) const
 {
+	using namespace UE::Landscape::EditLayers;
+
+	ERenderFlags RenderFlags = ERenderFlags::None;
+	// COMMENT [jonathan.bard] : this is not something we want to keep (we will use partial edit layer renders for this eventually) but we can still render in immediate mode
+	//  in the "reinitialize on next render case" because we perform a synchronous read then so we need to run on the game thread to perform the rendering commands flush 
+	if ((InMergeContext->IsHeightmapMerge() && bReinitializeHeightOnNextRender)
+		|| Algo::AnyOf(WeightPatches, [](const TObjectPtr<ULandscapeWeightPatchTextureInfo>& Info) { return Info->bReinitializeOnNextRender; }))
+	{
+		RenderFlags |= ERenderFlags::RenderMode_Immediate;
+	}
+	else
+	{
+		RenderFlags |= ERenderFlags::RenderMode_Recorded;
+	}
+	return RenderFlags;
+}
+
+void ULandscapeTexturePatch::RenderLayer(FRenderParams& InRenderParams, UE::Landscape::FRDGBuilderRecorder& RDGBuilderRecorder)
+{
+	using namespace UE::Landscape;
 	using namespace UE::Landscape::PatchUtil;
 	using namespace UE::Landscape::EditLayers;
 
@@ -142,18 +162,17 @@ void ULandscapeTexturePatch::RenderLayer(FRenderParams& InRenderParams)
 			return;
 		}
 
-		// graph builder expects external textures to start as SRV
-		LandscapeScratchRT->TransitionTo(ERHIAccess::SRVMask);
-		
 		if (bReinitializeHeightOnNextRender)
 		{
 			bReinitializeHeightOnNextRender = false;
+			checkf(!RDGBuilderRecorder.IsRecording(), TEXT("We should be using ERenderFlags::RenderMode_Immediate when reinitializing height"));
 			ReinitializeHeight(CurrentData, LandscapeHeightmapToWorld);
 			return;
 		}
 		else
 		{
-			ApplyToHeightmap(CurrentData, GetHeightmapToWorld(InRenderParams.RenderAreaWorldTransform));
+			checkf(RDGBuilderRecorder.IsRecording(), TEXT("We should be using ERenderFlags::RenderMode_Recorded in the typical case"));
+			ApplyToHeightmap(RDGBuilderRecorder, CurrentData, GetHeightmapToWorld(InRenderParams.RenderAreaWorldTransform));
 			return;
 		}
 	}
@@ -163,9 +182,6 @@ void ULandscapeTexturePatch::RenderLayer(FRenderParams& InRenderParams)
 	UTextureRenderTarget2DArray* TextureArray = LandscapeScratchRT->TryGetRenderTarget2DArray();
 	check(TextureArray != nullptr);
 	
-	// Only need to transition if we get a matching weight patch
-	bool bTransitionedToSRV = false;
-
 	int32 NumTargetLayersInRenderGroup = InRenderParams.RenderGroupTargetLayerNames.Num();
 	check(TextureArray->Slices == NumTargetLayersInRenderGroup);
 
@@ -192,24 +208,19 @@ void ULandscapeTexturePatch::RenderLayer(FRenderParams& InRenderParams)
 			continue;
 		}
 
-		// graph builder expects external textures to start as SRV
-		if (!bTransitionedToSRV)
-		{
-			LandscapeScratchRT->TransitionTo(ERHIAccess::SRVMask);
-			bTransitionedToSRV = true;
-		}
-
 		if (WeightPatchInfo->bReinitializeOnNextRender)
 		{
 			WeightPatchInfo->bReinitializeOnNextRender = false;
+			checkf(!RDGBuilderRecorder.IsRecording(), TEXT("We should be using ERenderFlags::RenderMode_Immediate when reinitializing weight"));
 			ReinitializeWeightPatch(WeightPatchInfo, TextureArray->GetResource(),
 				FIntPoint(TextureArray->SizeX, TextureArray->SizeY), TargetLayerIndexInRenderGroup, LandscapeHeightmapToWorld);
 		}
 		else
 		{
-			ApplyToWeightmap(WeightPatchInfo, 
-				[Resource = TextureArray->GetResource()]() { return Resource->GetTexture2DArrayRHI(); },
-				TargetLayerIndexInRenderGroup,
+			checkf(RDGBuilderRecorder.IsRecording(), TEXT("We should be using ERenderFlags::RenderMode_Recorded in the typical case"));
+			ApplyToWeightmap(RDGBuilderRecorder, WeightPatchInfo,
+				TextureArray->GetResource(),
+				TargetLayerIndexInRenderGroup, 
 				InRenderParams.RenderAreaSectionRect.Size(), 
 				GetHeightmapToWorld(InRenderParams.RenderAreaWorldTransform));
 		}
@@ -241,7 +252,9 @@ UTextureRenderTarget2D* ULandscapeTexturePatch::RenderLayer_Native(const FLandsc
 		}
 		else
 		{
-			return ApplyToHeightmap(InParameters.CombinedResult, LandscapeHeightmapToWorld);
+			// We render in immediate mode for the legacy path so we use an immediate recorder :  
+			FRDGBuilderRecorder RDGBuilderRecorderImmediate;
+			return ApplyToHeightmap(RDGBuilderRecorderImmediate, InParameters.CombinedResult, LandscapeHeightmapToWorld, /*OutputAccess = */ERHIAccess::SRVMask); // in the legacy path the RT leaves as a SRV
 		}
 	}
 	else
@@ -283,18 +296,21 @@ UTextureRenderTarget2D* ULandscapeTexturePatch::RenderLayer_Native(const FLandsc
 		}
 		else
 		{
-			ApplyToWeightmap(WeightPatchInfo, 
-				[Resource = InParameters.CombinedResult->GetResource()]() { return Resource->GetTexture2DRHI(); },
+			// We render in immediate mode for the legacy path so we use an immediate recorder :  
+			FRDGBuilderRecorder RDGBuilderRecorderImmediate;
+			ApplyToWeightmap(RDGBuilderRecorderImmediate, WeightPatchInfo,
+				InParameters.CombinedResult->GetResource(),
 				0, // Slice index
 				FIntPoint(InParameters.CombinedResult->SizeX, InParameters.CombinedResult->SizeY),
-				LandscapeHeightmapToWorld);
+				LandscapeHeightmapToWorld, 
+				/*OutputAccess = */ERHIAccess::SRVMask); // in the legacy path the RT leaves as a SRV
 			return InParameters.CombinedResult;
 		}
 	}
 }
 
-UTextureRenderTarget2D* ULandscapeTexturePatch::ApplyToHeightmap(UTextureRenderTarget2D* InCombinedResult, 
-	const FTransform& LandscapeHeightmapToWorld)
+UTextureRenderTarget2D* ULandscapeTexturePatch::ApplyToHeightmap(UE::Landscape::FRDGBuilderRecorder& RDGBuilderRecorder, UTextureRenderTarget2D* InCombinedResult,
+	const FTransform& LandscapeHeightmapToWorld, ERHIAccess OutputAccess)
 {
 	using namespace UE::Landscape;
 
@@ -346,12 +362,11 @@ UTextureRenderTarget2D* ULandscapeTexturePatch::ApplyToHeightmap(UTextureRenderT
 		return InCombinedResult;
 	}
 
-	ENQUEUE_RENDER_COMMAND(LandscapeTextureHeightPatch)([InCombinedResult, ShaderParamsToCopy, Patch, DestinationBounds](FRHICommandListImmediate& RHICmdList)
+	auto RDGCommand =
+		[InCombinedResult, ShaderParamsToCopy, Patch, DestinationBounds](FRDGBuilder& GraphBuilder)
 		{
 			TRACE_CPUPROFILER_EVENT_SCOPE(LandscapeTextureHeightPatch_Render);
 
-			FRDGBuilder GraphBuilder(RHICmdList, RDG_EVENT_NAME("ApplyTextureHeightPatch"));
-
 			TRefCountPtr<IPooledRenderTarget> DestinationRenderTarget = CreateRenderTarget(InCombinedResult->GetResource()->GetTexture2DRHI(), TEXT("LandscapeTextureHeightPatchOutput"));
 			FRDGTextureRef DestinationTexture = GraphBuilder.RegisterExternalTexture(DestinationRenderTarget);
 
@@ -387,16 +402,16 @@ UTextureRenderTarget2D* ULandscapeTexturePatch::ApplyToHeightmap(UTextureRenderT
 			ShaderParams->RenderTargets[0] = FRenderTargetBinding(DestinationTexture, ERenderTargetLoadAction::ENoAction, /*InMipIndex = */0);
 
 			FApplyLandscapeTextureHeightPatchPS::AddToRenderGraph(GraphBuilder, ShaderParams, DestinationBounds);
-
-			GraphBuilder.Execute();
-		});
+		};
+	// We need to specify the final state of the external texture to prevent the graph builder from transitioning it to SRVMask :
+	RDGBuilderRecorder.EnqueueRDGCommand(RDGCommand, { { InCombinedResult->GetResource(), (OutputAccess == ERHIAccess::None) ? ERHIAccess::RTV : OutputAccess} });
 
 	return InCombinedResult;
 }
 
-void ULandscapeTexturePatch::ApplyToWeightmap(ULandscapeWeightPatchTextureInfo* PatchInfo, 
-	TFunction<FRHITexture*()> RenderThreadLandscapeTextureGetter, int32 LandscapeTextureSliceIndex, 
-	const FIntPoint& LandscapeTextureResolution, const FTransform& LandscapeHeightmapToWorld)
+void ULandscapeTexturePatch::ApplyToWeightmap(UE::Landscape::FRDGBuilderRecorder& RDGBuilderRecorder, ULandscapeWeightPatchTextureInfo* PatchInfo,
+	FTextureResource* InMergedLandscapeTextureResource, int32 LandscapeTextureSliceIndex,
+	const FIntPoint& LandscapeTextureResolution, const FTransform& LandscapeHeightmapToWorld, ERHIAccess OutputAccess)
 {
 	using namespace UE::Landscape;
 
@@ -454,57 +469,54 @@ void ULandscapeTexturePatch::ApplyToWeightmap(ULandscapeWeightPatchTextureInfo*
 		return;
 	}
 
-	ENQUEUE_RENDER_COMMAND(LandscapeTextureWeightPatch)(
-		[RenderThreadLandscapeTextureGetter, LandscapeTextureSliceIndex, ShaderParamsToCopy, Patch, DestinationBounds](FRHICommandListImmediate& RHICmdList)
-	{
-		TRACE_CPUPROFILER_EVENT_SCOPE(LandscapeTextureWeightPatch_Render);
-
-		FRDGBuilder GraphBuilder(RHICmdList, RDG_EVENT_NAME("ApplyTextureWeightPatch"));
-
-		TRefCountPtr<IPooledRenderTarget> DestinationRenderTarget = CreateRenderTarget(
-			RenderThreadLandscapeTextureGetter(), TEXT("LandscapeTextureWeightPatchOutput"));
-		FRDGTextureRef DestinationTexture = GraphBuilder.RegisterExternalTexture(DestinationRenderTarget);
+	auto RDGCommand =
+		[InMergedLandscapeTextureResource, LandscapeTextureSliceIndex, ShaderParamsToCopy, Patch, DestinationBounds](FRDGBuilder& GraphBuilder)
+		{
+			TRACE_CPUPROFILER_EVENT_SCOPE(LandscapeTextureWeightPatch_Render);
 
-		// Make a copy of the portion of our weightmap input that we're writing to so that we can 
-		// read and write at the same time (needed for blending)
-		FRDGTextureDesc InputCopyDescription = DestinationTexture->Desc;
-		InputCopyDescription.Dimension = ETextureDimension::Texture2D;
-		InputCopyDescription.ArraySize = 1;
-		InputCopyDescription.NumMips = 1;
-		InputCopyDescription.Extent = DestinationBounds.Size();
-		FRDGTextureRef InputCopy = GraphBuilder.CreateTexture(InputCopyDescription, TEXT("LandscapeTextureWeightPatchInputCopy"));
+			TRefCountPtr<IPooledRenderTarget> DestinationRenderTarget = CreateRenderTarget(InMergedLandscapeTextureResource->GetTextureRHI(), TEXT("LandscapeTextureWeightPatchOutput"));
+			FRDGTextureRef DestinationTexture = GraphBuilder.RegisterExternalTexture(DestinationRenderTarget);
 
-		FRHICopyTextureInfo CopyTextureInfo;
-		CopyTextureInfo.SourceMipIndex = 0;
-		CopyTextureInfo.NumMips = 1;
-		CopyTextureInfo.SourceSliceIndex = LandscapeTextureSliceIndex;
-		CopyTextureInfo.NumSlices = 1;
-		CopyTextureInfo.SourcePosition = FIntVector(DestinationBounds.Min.X, DestinationBounds.Min.Y, 0);
-		CopyTextureInfo.Size = FIntVector(InputCopyDescription.Extent.X, InputCopyDescription.Extent.Y, 0);
-		AddCopyTexturePass(GraphBuilder, DestinationTexture, InputCopy, CopyTextureInfo);
+			// Make a copy of the portion of our weightmap input that we're writing to so that we can 
+			// read and write at the same time (needed for blending)
+			FRDGTextureDesc InputCopyDescription = DestinationTexture->Desc;
+			InputCopyDescription.Dimension = ETextureDimension::Texture2D;
+			InputCopyDescription.ArraySize = 1;
+			InputCopyDescription.NumMips = 1;
+			InputCopyDescription.Extent = DestinationBounds.Size();
+			FRDGTextureRef InputCopy = GraphBuilder.CreateTexture(InputCopyDescription, TEXT("LandscapeTextureWeightPatchInputCopy"));
 
-		FApplyLandscapeTextureWeightPatchPS::FParameters* ShaderParams =
-			GraphBuilder.AllocParameters<FApplyLandscapeTextureWeightPatchPS::FParameters>();
-		*ShaderParams = ShaderParamsToCopy;
+			FRHICopyTextureInfo CopyTextureInfo;
+			CopyTextureInfo.SourceMipIndex = 0;
+			CopyTextureInfo.NumMips = 1;
+			CopyTextureInfo.SourceSliceIndex = LandscapeTextureSliceIndex;
+			CopyTextureInfo.NumSlices = 1;
+			CopyTextureInfo.SourcePosition = FIntVector(DestinationBounds.Min.X, DestinationBounds.Min.Y, 0);
+			CopyTextureInfo.Size = FIntVector(InputCopyDescription.Extent.X, InputCopyDescription.Extent.Y, 0);
+			AddCopyTexturePass(GraphBuilder, DestinationTexture, InputCopy, CopyTextureInfo);
 
-		TRefCountPtr<IPooledRenderTarget> PatchRenderTarget = CreateRenderTarget(Patch->GetTexture2DRHI(), TEXT("LandscapeTextureWeightPatch"));
-		FRDGTextureRef PatchTexture = GraphBuilder.RegisterExternalTexture(PatchRenderTarget);
-		FRDGTextureSRVRef PatchSRV = GraphBuilder.CreateSRV(FRDGTextureSRVDesc::CreateForMipLevel(PatchTexture, 0));
-		ShaderParams->InWeightPatch = PatchSRV;
-		ShaderParams->InWeightPatchSampler = TStaticSamplerState<SF_Bilinear, AM_Clamp, AM_Clamp>::GetRHI();
+			FApplyLandscapeTextureWeightPatchPS::FParameters* ShaderParams =
+				GraphBuilder.AllocParameters<FApplyLandscapeTextureWeightPatchPS::FParameters>();
+			*ShaderParams = ShaderParamsToCopy;
 
-		FRDGTextureSRVRef InputCopySRV = GraphBuilder.CreateSRV(FRDGTextureSRVDesc::CreateForMipLevel(InputCopy, 0));
-		ShaderParams->InSourceWeightmap = InputCopySRV;
-		ShaderParams->InSourceWeightmapCoordOffset = DestinationBounds.Min;
+			TRefCountPtr<IPooledRenderTarget> PatchRenderTarget = CreateRenderTarget(Patch->GetTexture2DRHI(), TEXT("LandscapeTextureWeightPatch"));
+			FRDGTextureRef PatchTexture = GraphBuilder.RegisterExternalTexture(PatchRenderTarget);
+			FRDGTextureSRVRef PatchSRV = GraphBuilder.CreateSRV(FRDGTextureSRVDesc::CreateForMipLevel(PatchTexture, 0));
+			ShaderParams->InWeightPatch = PatchSRV;
+			ShaderParams->InWeightPatchSampler = TStaticSamplerState<SF_Bilinear, AM_Clamp, AM_Clamp>::GetRHI();
 
-		ShaderParams->RenderTargets[0] = FRenderTargetBinding(DestinationTexture, ERenderTargetLoadAction::ENoAction,
-			/*InMipIndex = */0, LandscapeTextureSliceIndex);
+			FRDGTextureSRVRef InputCopySRV = GraphBuilder.CreateSRV(FRDGTextureSRVDesc::CreateForMipLevel(InputCopy, 0));
+			ShaderParams->InSourceWeightmap = InputCopySRV;
+			ShaderParams->InSourceWeightmapCoordOffset = DestinationBounds.Min;
 
-		FApplyLandscapeTextureWeightPatchPS::AddToRenderGraph(GraphBuilder, ShaderParams, DestinationBounds);
+			ShaderParams->RenderTargets[0] = FRenderTargetBinding(DestinationTexture, ERenderTargetLoadAction::ENoAction,
+				/*InMipIndex = */0, LandscapeTextureSliceIndex);
 
-		GraphBuilder.Execute();
-	});
+			FApplyLandscapeTextureWeightPatchPS::AddToRenderGraph(GraphBuilder, ShaderParams, DestinationBounds);
+		};
 
+	// We need to specify the final state of the external texture to prevent the graph builder from transitioning it to SRVMask :
+	RDGBuilderRecorder.EnqueueRDGCommand(RDGCommand, { { InMergedLandscapeTextureResource, (OutputAccess == ERHIAccess::None) ? ERHIAccess::RTV : OutputAccess} });
 }
 
 void ULandscapeTexturePatch::GetCommonShaderParams(const FTransform& LandscapeHeightmapToWorldIn,
@@ -1022,15 +1034,12 @@ void ULandscapeTexturePatch::ReinitializeWeightPatch(ULandscapeWeightPatchTextur
 
 		if (SliceIndex < 0)
 		{
-			FRDGTextureRef SourceTexture = GraphBuilder.RegisterExternalTexture(CreateRenderTarget(
-				InputResource->GetTexture2DRHI(), TEXT("ReinitializationSource")));
+			FRDGTextureRef SourceTexture = GraphBuilder.RegisterExternalTexture(CreateRenderTarget(InputResource->GetTexture2DRHI(), TEXT("ReinitializationSource")));
 			ShaderParams->InSource = GraphBuilder.CreateSRV(FRDGTextureSRVDesc::CreateForMipLevel(SourceTexture, 0));
-		
 		}
 		else
 		{
-			FRDGTextureRef SourceTexture = GraphBuilder.RegisterExternalTexture(CreateRenderTarget(
-				InputResource->GetTexture2DArrayRHI(), TEXT("ReinitializationSource")));
+			FRDGTextureRef SourceTexture = GraphBuilder.RegisterExternalTexture(CreateRenderTarget(InputResource->GetTexture2DArrayRHI(), TEXT("ReinitializationSource")));
 			FRDGTextureSRVDesc Desc = FRDGTextureSRVDesc::CreateForSlice(SourceTexture, SliceIndex);
 			Desc.MipLevel = 0;
 			Desc.NumMipLevels = 1;
@@ -2031,26 +2040,26 @@ void ULandscapeTexturePatch::GetRendererStateInfo(const UE::Landscape::EditLayer
 	TArray<TBitArray<>>& OutRenderGroups) const
 {
 	if (InMergeContext->IsHeightmapMerge())
+{
+	if (CanAffectHeightmap())
 	{
-		if (CanAffectHeightmap())
-		{
-			OutSupported.AddTargetType(ELandscapeToolTargetType::Heightmap);
-		}
+		OutSupported.AddTargetType(ELandscapeToolTargetType::Heightmap);
+	}
 	}
 	else
 	{
-		for (const TObjectPtr<ULandscapeWeightPatchTextureInfo>& WeightPatch : WeightPatches)
+	for (const TObjectPtr<ULandscapeWeightPatchTextureInfo>& WeightPatch : WeightPatches)
+	{
+		if (IsValid(WeightPatch) && WeightPatchCanRender(*WeightPatch))
 		{
-			if (IsValid(WeightPatch) && WeightPatchCanRender(*WeightPatch))
+			if (WeightPatch->bEditVisibilityLayer)
 			{
-				if (WeightPatch->bEditVisibilityLayer)
-				{
-					OutSupported.AddTargetType(ELandscapeToolTargetType::Visibility);
-				}
+				OutSupported.AddTargetType(ELandscapeToolTargetType::Visibility);
+			}
 				else if (InMergeContext->IsValidTargetLayerName(WeightPatch->WeightmapLayerName))
-				{
+			{
 					const int32 TargetLayerIndex = InMergeContext->GetTargetLayerIndexForNameChecked(WeightPatch->WeightmapLayerName);
-					OutSupported.AddTargetType(ELandscapeToolTargetType::Weightmap);
+				OutSupported.AddTargetType(ELandscapeToolTargetType::Weightmap);
 					OutSupported.AddWeightmap(TargetLayerIndex);
 				}
 			}
@@ -2081,19 +2090,19 @@ TArray<UE::Landscape::EditLayers::FEditLayerRenderItem> ULandscapeTexturePatch::
 
 	if (InMergeContext->IsHeightmapMerge())
 	{
-		if (CanAffectHeightmap())
-		{
+	if (CanAffectHeightmap())
+	{
 			FEditLayerTargetTypeState TargetInfo(InMergeContext, ELandscapeToolTargetTypeFlags::Heightmap);
-			FEditLayerRenderItem Item(TargetInfo, InputWorldArea, OutputWorldArea, false);
-			AffectedAreas.Add(Item);
-		}
+		FEditLayerRenderItem Item(TargetInfo, InputWorldArea, OutputWorldArea, false);
+		AffectedAreas.Add(Item);
+	}
 	}
 	else
 	{
-		for (const TObjectPtr<ULandscapeWeightPatchTextureInfo>& WeightPatch : WeightPatches)
+	for (const TObjectPtr<ULandscapeWeightPatchTextureInfo>& WeightPatch : WeightPatches)
+	{
+		if (IsValid(WeightPatch) && WeightPatchCanRender(*WeightPatch))
 		{
-			if (IsValid(WeightPatch) && WeightPatchCanRender(*WeightPatch))
-			{
 				FEditLayerTargetTypeState TargetInfo(InMergeContext);
 				if (WeightPatch->bEditVisibilityLayer)
 				{
@@ -2106,10 +2115,10 @@ TArray<UE::Landscape::EditLayers::FEditLayerRenderItem> ULandscapeTexturePatch::
 					TargetInfo.AddWeightmap(TargetLayerIndex);
 				}
 				FEditLayerRenderItem Item(TargetInfo, InputWorldArea, OutputWorldArea, /*bInModifyExistingWeightmapsOnly = */false);
-				AffectedAreas.Add(Item);
-			}
+			AffectedAreas.Add(Item);
 		}
 	}
+	}
 
 	return AffectedAreas;
 }
diff --git a/Engine/Plugins/Experimental/LandscapePatch/Source/LandscapePatch/Public/LandscapeCircleHeightPatch.h b/Engine/Plugins/Experimental/LandscapePatch/Source/LandscapePatch/Public/LandscapeCircleHeightPatch.h
--- a/Engine/Plugins/Experimental/LandscapePatch/Source/LandscapePatch/Public/LandscapeCircleHeightPatch.h
+++ b/Engine/Plugins/Experimental/LandscapePatch/Source/LandscapePatch/Public/LandscapeCircleHeightPatch.h
@@ -5,9 +5,17 @@
 
 #include "LandscapePatchComponent.h"
 #include "LandscapeEditTypes.h"
+#include "RHIAccess.h"
 
 #include "LandscapeCircleHeightPatch.generated.h"
 
+class FTextureResource;
+
+namespace UE::Landscape
+{
+	class FRDGBuilderRecorder;
+} // namespace namespace UE::Landscape
+
 /**
  * The simplest height patch: a circle of flat ground with a falloff past the initial radius across which the
  * alpha decreases linearly. When added to an actor, initializes itself to the bottom of the bounding box.
@@ -32,7 +40,8 @@ public:
 		TArray<TBitArray<>>& OutRenderGroups) const override;
 	virtual FString GetEditLayerRendererDebugName() const override;
 	virtual TArray<FEditLayerRenderItem> GetRenderItems(const UE::Landscape::EditLayers::FMergeContext* InMergeContext) const override;
-	virtual void RenderLayer(FRenderParams& InRenderParams) override;
+	virtual UE::Landscape::EditLayers::ERenderFlags GetRenderFlags(const UE::Landscape::EditLayers::FMergeContext* InMergeContext) const override { return UE::Landscape::EditLayers::ERenderFlags::RenderMode_Recorded; }
+	virtual void RenderLayer(FRenderParams& InRenderParams, UE::Landscape::FRDGBuilderRecorder& RDGBuilderRecorder) override;
 #endif
 
 	virtual bool CanAffectHeightmap() const override { return !bEditVisibility; }
@@ -60,7 +69,7 @@ protected:
 	bool bExclusiveRadius = false;
 
 private:
-	void ApplyCirclePatch(bool bIsVisibilityLayer,
-		const TFunction<FRHITexture* ()>& RenderThreadLandscapeTextureGetter, int32 LandscapeTextureSliceIndex,
-		const FIntPoint& DestinationResolution, const FTransform& HeightmapCoordsToWorld);
+	void ApplyCirclePatch(UE::Landscape::FRDGBuilderRecorder& RDGBuilderRecorder, bool bIsVisibilityLayer,
+		FTextureResource* InMergedLandscapeTextureResource, int32 LandscapeTextureSliceIndex,
+		const FIntPoint& DestinationResolution, const FTransform& HeightmapCoordsToWorld, ERHIAccess OutputAccess = ERHIAccess::None);
 };
diff --git a/Engine/Plugins/Experimental/LandscapePatch/Source/LandscapePatch/Public/LandscapePatchComponent.h b/Engine/Plugins/Experimental/LandscapePatch/Source/LandscapePatch/Public/LandscapePatchComponent.h
--- a/Engine/Plugins/Experimental/LandscapePatch/Source/LandscapePatch/Public/LandscapePatchComponent.h
+++ b/Engine/Plugins/Experimental/LandscapePatch/Source/LandscapePatch/Public/LandscapePatchComponent.h
@@ -59,7 +59,6 @@ public:
 
 #if WITH_EDITOR
 	// ILandscapeEditLayerRenderer
-	virtual bool CanRender() const override { return CanAffectLandscape(); }
 	// Subclasses are expected to implement GetEditLayerRendererDebugName, GetRendererStateInfo, GetRenderItems, and RenderLayer
 #endif
 
diff --git a/Engine/Plugins/Experimental/LandscapePatch/Source/LandscapePatch/Public/LandscapePatchManager.h b/Engine/Plugins/Experimental/LandscapePatch/Source/LandscapePatch/Public/LandscapePatchManager.h
--- a/Engine/Plugins/Experimental/LandscapePatch/Source/LandscapePatch/Public/LandscapePatchManager.h
+++ b/Engine/Plugins/Experimental/LandscapePatch/Source/LandscapePatch/Public/LandscapePatchManager.h
@@ -47,11 +47,8 @@ public:
 		UE::Landscape::EditLayers::FEditLayerTargetTypeState& OutSupportedTargetTypeState,
 		UE::Landscape::EditLayers::FEditLayerTargetTypeState& OutEnabledTargetTypeState,
 		TArray<TBitArray<>>& OutRenderGroups) const override {}
-	virtual TArray<UE::Landscape::EditLayers::FEditLayerRenderItem> GetRenderItems(const UE::Landscape::EditLayers::FMergeContext* InMergeContext) const override
-	{
-		return {};
-	}
-	virtual void RenderLayer(ILandscapeEditLayerRenderer::FRenderParams& InRenderParams) override { };
+	virtual TArray<UE::Landscape::EditLayers::FEditLayerRenderItem> GetRenderItems(const UE::Landscape::EditLayers::FMergeContext* InMergeContext) const override { return {}; }
+	virtual void RenderLayer(ILandscapeEditLayerRenderer::FRenderParams& InRenderParams, UE::Landscape::FRDGBuilderRecorder& RDGBuilderRecorder) override { };
 #endif // WITH_EDITOR
 
 	// Adds the brush to the given landscape, removing it from any previous one. This differs from SetOwningLandscape
diff --git a/Engine/Plugins/Experimental/LandscapePatch/Source/LandscapePatch/Public/LandscapeTexturePatch.h b/Engine/Plugins/Experimental/LandscapePatch/Source/LandscapePatch/Public/LandscapeTexturePatch.h
--- a/Engine/Plugins/Experimental/LandscapePatch/Source/LandscapePatch/Public/LandscapeTexturePatch.h
+++ b/Engine/Plugins/Experimental/LandscapePatch/Source/LandscapePatch/Public/LandscapeTexturePatch.h
@@ -5,11 +5,12 @@
 
 #include "LandscapePatchComponent.h"
 #include "LandscapeTextureBackedRenderTarget.h"
+#include "RHIAccess.h"
 
 #include "LandscapeTexturePatch.generated.h"
 
+class FTextureResource;
 class ULandscapeTexturePatch;
-
 class ULandscapeHeightTextureBackedRenderTarget;
 class UTexture;
 class UTexture2D;
@@ -18,6 +19,7 @@ namespace UE::Landscape
 {
 	class FApplyLandscapeTextureHeightPatchPSParameters;
 	class FApplyLandscapeTextureWeightPatchPSParameters;
+	class FRDGBuilderRecorder;
 } // namespace UE::Landscape
 
 /**
@@ -236,7 +238,9 @@ public:
 		TArray<TBitArray<>>& OutRenderGroups) const override;
 	virtual FString GetEditLayerRendererDebugName() const override;
 	virtual TArray<FEditLayerRenderItem> GetRenderItems(const UE::Landscape::EditLayers::FMergeContext* InMergeContext) const override;
-	virtual void RenderLayer(FRenderParams& InRenderParams) override;
+
+	virtual UE::Landscape::EditLayers::ERenderFlags GetRenderFlags(const UE::Landscape::EditLayers::FMergeContext* InMergeContext) const override;
+	virtual void RenderLayer(FRenderParams& RenderParams, UE::Landscape::FRDGBuilderRecorder& RDGBuilderRecorder) override;
 
 	// ULandscapePatchComponent
 	UTextureRenderTarget2D* RenderLayer_Native(const FLandscapeBrushParameters& InParameters, const FTransform& HeightmapToWorld) override;
@@ -636,18 +640,15 @@ protected:
 	 * resolution than the landscape (perhaps useful for slightly more accurate results while not aligned to landscape), and
 	 * a value less that 1.0 will use lower.
 	 */
-	UPROPERTY(EditAnywhere, Category = Initialization, meta = (
-		EditCondition = "bBaseResolutionOffLandscape"))
+	UPROPERTY(EditAnywhere, Category = Initialization, meta = (EditCondition = "bBaseResolutionOffLandscape"))
 	float ResolutionMultiplier = 1;
 
 	/** Texture width to use when reinitializing using Reinitialize Weights or ReinitializeHeight, if not basing resolution off landscape. */
-	UPROPERTY(EditAnywhere, Category = Initialization, meta = (EditCondition = "!bBaseResolutionOffLandscape", 
-		ClampMin = "1"))
+	UPROPERTY(EditAnywhere, Category = Initialization, meta = (EditCondition = "!bBaseResolutionOffLandscape", ClampMin = "1"))
 	int32 InitTextureSizeX = 33;
 
 	/** Texture height to use when reinitializing using Reinitialize Weights or ReinitializeHeight, if not basing resolution off landscape. */
-	UPROPERTY(EditAnywhere, Category = Initialization, meta = (EditCondition = "!bBaseResolutionOffLandscape", 
-		ClampMin = "1"))
+	UPROPERTY(EditAnywhere, Category = Initialization, meta = (EditCondition = "!bBaseResolutionOffLandscape", ClampMin = "1"))
 	int32 InitTextureSizeY = 33;
 
 private:
@@ -659,10 +660,11 @@ private:
 #if WITH_EDITOR
 	void TransitionHeightSourceModeInternal(ELandscapeTexturePatchSourceMode OldMode, ELandscapeTexturePatchSourceMode NewMode);
 	FLandscapeHeightPatchConvertToNativeParams GetHeightConvertToNativeParams() const;
-	UTextureRenderTarget2D* ApplyToHeightmap(UTextureRenderTarget2D* InCombinedResult, const FTransform& LandscapeHeightmapToWorld);
-	void ApplyToWeightmap(ULandscapeWeightPatchTextureInfo* PatchInfo,
-		TFunction<FRHITexture* ()> RenderThreadLandscapeTextureGetter, int32 LandscapeTextureSliceIndex,
-		const FIntPoint& LandscapeTextureResolution, const FTransform& LandscapeHeightmapToWorld);
+	UTextureRenderTarget2D* ApplyToHeightmap(UE::Landscape::FRDGBuilderRecorder& RDGBuilderRecorder, UTextureRenderTarget2D* InCombinedResult, const FTransform& LandscapeHeightmapToWorld, 
+		ERHIAccess OutputAccess = ERHIAccess::None);
+	void ApplyToWeightmap(UE::Landscape::FRDGBuilderRecorder& RDGBuilderRecorder, ULandscapeWeightPatchTextureInfo* PatchInfo,
+		FTextureResource* InMergedLandscapeTextureResource, int32 LandscapeTextureSliceIndex,
+		const FIntPoint& LandscapeTextureResolution, const FTransform& LandscapeHeightmapToWorld, ERHIAccess OutputAccess = ERHIAccess::None);
 
 	void GetCommonShaderParams(const FTransform& LandscapeHeightmapToWorldIn, 
 		const FIntPoint& SourceResolutionIn, const FIntPoint& DestinationResolutionIn,
diff --git a/Engine/Source/Runtime/Landscape/Classes/LandscapeEditLayer.h b/Engine/Source/Runtime/Landscape/Classes/LandscapeEditLayer.h
--- a/Engine/Source/Runtime/Landscape/Classes/LandscapeEditLayer.h
+++ b/Engine/Source/Runtime/Landscape/Classes/LandscapeEditLayer.h
@@ -197,9 +197,11 @@ public:
 	LANDSCAPE_API virtual void GetRendererStateInfo(const UE::Landscape::EditLayers::FMergeContext* InMergeContext,
 		UE::Landscape::EditLayers::FEditLayerTargetTypeState& OutSupportedTargetTypeState, UE::Landscape::EditLayers::FEditLayerTargetTypeState& OutEnabledTargetTypeState, TArray<TBitArray<>>& OutRenderGroups) const override;
 	LANDSCAPE_EDIT_LAYERS_BATCHED_MERGE_EXPERIMENTAL
+	virtual UE::Landscape::EditLayers::ERenderFlags GetRenderFlags(const UE::Landscape::EditLayers::FMergeContext* InMergeContext) const override { return UE::Landscape::EditLayers::ERenderFlags::RenderMode_Recorded; }
+	LANDSCAPE_EDIT_LAYERS_BATCHED_MERGE_EXPERIMENTAL
 	LANDSCAPE_API virtual TArray<UE::Landscape::EditLayers::FEditLayerRenderItem> GetRenderItems(const UE::Landscape::EditLayers::FMergeContext* InMergeContext) const override;
 	LANDSCAPE_EDIT_LAYERS_BATCHED_MERGE_EXPERIMENTAL
-	LANDSCAPE_API virtual void RenderLayer(ILandscapeEditLayerRenderer::FRenderParams& InRenderParams) override;
+	LANDSCAPE_API virtual void RenderLayer(FRenderParams& RenderParams, UE::Landscape::FRDGBuilderRecorder& RDGBuilderRecorder) override;
 	LANDSCAPE_EDIT_LAYERS_BATCHED_MERGE_EXPERIMENTAL
 	LANDSCAPE_API virtual FString GetEditLayerRendererDebugName() const override;
 	//~ End ILandscapeEditLayerRenderer implementation
diff --git a/Engine/Source/Runtime/Landscape/Private/LandscapeBlueprintBrushBase.cpp b/Engine/Source/Runtime/Landscape/Private/LandscapeBlueprintBrushBase.cpp
--- a/Engine/Source/Runtime/Landscape/Private/LandscapeBlueprintBrushBase.cpp
+++ b/Engine/Source/Runtime/Landscape/Private/LandscapeBlueprintBrushBase.cpp
@@ -254,7 +254,7 @@ void ALandscapeBlueprintBrushBase::GetRendererStateInfo(const UE::Landscape::Edi
 
 	// Mark which weightmap is supported/enabled : 
 	if (!InMergeContext->IsHeightmapMerge() && CanAffectWeightmap())
-	{
+		{
 		InMergeContext->ForEachValidTargetLayer(
 			[InMergeContext, this, &OutSupportedTargetTypeState, &OutEnabledTargetTypeState](int32 InTargetLayerIndex, const FName& InTargetLayerName, ULandscapeLayerInfoObject* InWeightmapLayerInfo)
 			{
@@ -264,8 +264,8 @@ void ALandscapeBlueprintBrushBase::GetRendererStateInfo(const UE::Landscape::Edi
 					if (AffectsWeightmapLayer(InTargetLayerName))
 					{
 						OutEnabledTargetTypeState.AddWeightmap(InTargetLayerIndex);
-					}
-				}
+			}
+		}
 				return true;
 			});
 	}
@@ -291,21 +291,23 @@ TArray<UE::Landscape::EditLayers::FEditLayerRenderItem> ALandscapeBlueprintBrush
 	return { FEditLayerRenderItem(EnabledTargetTypeState, InputWorldArea, OutputWorldArea, /*bModifyExistingWeightmapsOnly = */false)};
 }
 
-void ALandscapeBlueprintBrushBase::RenderLayer(ILandscapeEditLayerRenderer::FRenderParams& InRenderParams)
+void ALandscapeBlueprintBrushBase::RenderLayer(ILandscapeEditLayerRenderer::FRenderParams& RenderParams, UE::Landscape::FRDGBuilderRecorder& RDGBuilderRecorder)
 {
 	using namespace UE::Landscape;
 
+	checkf(!RDGBuilderRecorder.IsRecording(), TEXT("ERenderFlags::RenderMode_Immediate means the command recorder should be recording at this point"));
+
 	// By default, use the old way of rendering BP brushes : 
 
 	// Swap the render targets so that the layer's input RT is the latest combined result :
 	// The write render target will be accessed as ERHIAccess::CopyDest all along : 
-	InRenderParams.MergeRenderContext->CycleBlendRenderTargets(/*InDesiredWriteAccess = */ERHIAccess::CopyDest);
-	ULandscapeScratchRenderTarget* WriteRT = InRenderParams.MergeRenderContext->GetBlendRenderTargetWrite();
-	ULandscapeScratchRenderTarget* CurrentLayerReadRT = InRenderParams.MergeRenderContext->GetBlendRenderTargetRead();
+	RenderParams.MergeRenderContext->CycleBlendRenderTargets(/*InDesiredWriteAccess = */ERHIAccess::CopyDest, RDGBuilderRecorder);
+	ULandscapeScratchRenderTarget* WriteRT = RenderParams.MergeRenderContext->GetBlendRenderTargetWrite();
+	ULandscapeScratchRenderTarget* CurrentLayerReadRT = RenderParams.MergeRenderContext->GetBlendRenderTargetRead();
+
+	const bool bIsHeightmapMerge = RenderParams.MergeRenderContext->IsHeightmapMerge();
+	const TArray<FName> EnabledWeightmaps = RenderParams.RendererState.GetActiveTargetWeightmaps();
 
-	const bool bIsHeightmapMerge = InRenderParams.MergeRenderContext->IsHeightmapMerge();
-	const TArray<FName> EnabledWeightmaps = InRenderParams.RendererState.GetActiveTargetWeightmaps();
-	
 	ULandscapeEditResourcesSubsystem* LandscapeEditResourcesSubsystem = GEngine->GetEngineSubsystem<ULandscapeEditResourcesSubsystem>();
 	check(LandscapeEditResourcesSubsystem != nullptr);
 	ULandscapeScratchRenderTarget* CurrentLayerReadRT2D = CurrentLayerReadRT;
@@ -320,7 +322,7 @@ void ALandscapeBlueprintBrushBase::RenderLayer(ILandscapeEditLayerRenderer::FRen
 	//  render target, so if, from one call to another, this size is changing because there happens to be a scratch render target big enough, we will have to call Initialize again 
 	//  on the brush, which might have performance implications. So using a size that matches what we've already used in the past here allows such brushes not to uselessly go 
 	//  through the initialization process again 
-	const FIntPoint RenderTargetIdealSize = InRenderParams.RenderAreaSectionRect.Size().ComponentMax(CurrentRenderAreaSize);
+	const FIntPoint RenderTargetIdealSize = RenderParams.RenderAreaSectionRect.Size().ComponentMax(CurrentRenderAreaSize);
 	// Use exact dimensions for the scratch texture here for the reason explained above (see RenderTargetIdealSize) :
 	FScratchRenderTargetParams ScratchRenderTargetParams(TEXT("BPBrushScratchRT"), /*bInExactDimensions = */true, /*bInUseUAV = */false, /*bInTargetArraySlicesIndependently = */false,
 		RenderTargetIdealSize, /*InNumSlices = */0, CurrentLayerReadRT->GetFormat(), CurrentLayerReadRT->GetClearColor(), ERHIAccess::CopyDest);
@@ -328,12 +330,12 @@ void ALandscapeBlueprintBrushBase::RenderLayer(ILandscapeEditLayerRenderer::FRen
 	CurrentLayerReadRT2D = ScratchTexture.RenderTarget;
 
 	// The original texture array will be accessed as ERHIAccess::CopySrc all along : 
-	CurrentLayerReadRT->TransitionTo(ERHIAccess::CopySrc);
+	CurrentLayerReadRT->TransitionTo(ERHIAccess::CopySrc, RDGBuilderRecorder);
 
-	const int32 NumTargetLayersInRenderGroup = InRenderParams.RenderGroupTargetLayerNames.Num();
+	const int32 NumTargetLayersInRenderGroup = RenderParams.RenderGroupTargetLayerNames.Num();
 	for (int32 TargetLayerIndexInRenderGroup = 0; TargetLayerIndexInRenderGroup < NumTargetLayersInRenderGroup; ++TargetLayerIndexInRenderGroup)
 	{
-		const FName TargetLayerName = InRenderParams.RenderGroupTargetLayerNames[TargetLayerIndexInRenderGroup];
+		const FName TargetLayerName = RenderParams.RenderGroupTargetLayerNames[TargetLayerIndexInRenderGroup];
 		RHI_BREADCRUMB_EVENT_GAMETHREAD("Render %s", TargetLayerName);
 
 		// If necessary, copy from the source render target (slice) to the scratch render target 2D : 
@@ -341,13 +343,13 @@ void ALandscapeBlueprintBrushBase::RenderLayer(ILandscapeEditLayerRenderer::FRen
 
 		{
 			ULandscapeScratchRenderTarget::FCopyFromScratchRenderTargetParams CopyParams(CurrentLayerReadRT);
-			// It's important to copy only the needed size, because CurrentLayerReadRT2D comes from ScratchTexture, whose size may be larger than InRenderParams.RenderAreaSectionRect)
-			CopyParams.CopySize = InRenderParams.RenderAreaSectionRect.Size();
+			// It's important to copy only the needed size, because CurrentLayerReadRT2D comes from ScratchTexture, whose size may be larger than RenderParams.RenderAreaSectionRect)
+			CopyParams.CopySize = RenderParams.RenderAreaSectionRect.Size();
 			// Copy from the proper slice in the texture array if any : 
 			CopyParams.SourceSliceIndex = TargetLayerIndexInRenderGroup;
-			CurrentLayerReadRT2D->CopyFrom(CopyParams);
+			CurrentLayerReadRT2D->CopyFrom(CopyParams, RDGBuilderRecorder);
 		}
-		CurrentLayerReadRT2D->TransitionTo(ERHIAccess::SRVMask);
+		CurrentLayerReadRT2D->TransitionTo(ERHIAccess::SRVMask, RDGBuilderRecorder);
 
 		check(CurrentLayerReadRT2D->GetCurrentState() == ERHIAccess::SRVMask);
 		check(WriteRT->GetCurrentState() == ERHIAccess::CopyDest);
@@ -361,7 +363,7 @@ void ALandscapeBlueprintBrushBase::RenderLayer(ILandscapeEditLayerRenderer::FRen
 		if (bIsHeightmapMerge || EnabledWeightmaps.Contains(TargetLayerName))
 		{
 			// Execute (i.e. (Initialize/)Render the BP brush) : 
-			FLandscapeBrushParameters BrushParameters(bIsHeightmapMerge, InRenderParams.RenderAreaWorldTransform, InRenderParams.RenderAreaSectionRect.Size(), ReadRT2D, TargetLayerName);
+			FLandscapeBrushParameters BrushParameters(bIsHeightmapMerge, RenderParams.RenderAreaWorldTransform, RenderParams.RenderAreaSectionRect.Size(), ReadRT2D, TargetLayerName);
 			if (UTextureRenderTarget2D* BrushOutputRT2D = Execute(BrushParameters))
 			{
 				// Only consider the BP brush's result if it's valid :
@@ -384,10 +386,10 @@ void ALandscapeBlueprintBrushBase::RenderLayer(ILandscapeEditLayerRenderer::FRen
 
 			ULandscapeScratchRenderTarget::FCopyFromTextureParams CopyParams(OutputRT2D);
 			// It's important to copy only the needed size, because OutputRT2D is currently sized after ReadRT2D, which can be a render target larger than it needs to be (because 
-			//  it can come from ScratchTexture, whose size may be larger than InRenderParams.RenderAreaSectionRect)
-			CopyParams.CopySize = InRenderParams.RenderAreaSectionRect.Size();
+			//  it can come from ScratchTexture, whose size may be larger than RenderParams.RenderAreaSectionRect)
+			CopyParams.CopySize = RenderParams.RenderAreaSectionRect.Size();
 			CopyParams.DestSliceIndex = TargetLayerIndexInRenderGroup;
-			WriteRT->CopyFrom(CopyParams);
+			WriteRT->CopyFrom(CopyParams, RDGBuilderRecorder);
 
 			// It's also expected we leave the RT return by the BP as SRV :
 			ENQUEUE_RENDER_COMMAND(TransitionToSRV)([Resource = OutputRT2D->GetResource()](FRHICommandListImmediate& InRHICmdList) mutable
@@ -398,8 +400,8 @@ void ALandscapeBlueprintBrushBase::RenderLayer(ILandscapeEditLayerRenderer::FRen
 	}
 
 	// Leave the render targets in the state they're expected to be in: 
-	WriteRT->TransitionTo(ERHIAccess::RTV);
-	CurrentLayerReadRT->TransitionTo(ERHIAccess::SRVMask);
+	WriteRT->TransitionTo(ERHIAccess::RTV, RDGBuilderRecorder);
+	CurrentLayerReadRT->TransitionTo(ERHIAccess::SRVMask, RDGBuilderRecorder);
 }
 
 FString ALandscapeBlueprintBrushBase::GetEditLayerRendererDebugName() const
diff --git a/Engine/Source/Runtime/Landscape/Private/LandscapeEditLayerRenderer.cpp b/Engine/Source/Runtime/Landscape/Private/LandscapeEditLayerRenderer.cpp
--- a/Engine/Source/Runtime/Landscape/Private/LandscapeEditLayerRenderer.cpp
+++ b/Engine/Source/Runtime/Landscape/Private/LandscapeEditLayerRenderer.cpp
@@ -425,12 +425,12 @@ void FEditLayerTargetTypeState::SetTargetTypeMask(ELandscapeToolTargetTypeFlags
 		const int32 VisibilityTargetLayerIndex = MergeContext->GetVisibilityTargetLayerIndex();
 		if (VisibilityTargetLayerIndex != INDEX_NONE)
 		{
-			if (EnumHasAnyFlags(TargetTypeMask, ELandscapeToolTargetTypeFlags::Visibility))
-			{
+		if (EnumHasAnyFlags(TargetTypeMask, ELandscapeToolTargetTypeFlags::Visibility))
+		{
 				AddWeightmap(VisibilityTargetLayerIndex);
-			}
+		}
 			else if (WeightmapTargetLayerBitIndices[VisibilityTargetLayerIndex])
-			{
+		{
 				RemoveWeightmap(VisibilityTargetLayerIndex);
 			}
 		}
@@ -555,16 +555,16 @@ FEditLayerRendererState::FEditLayerRendererState(const FMergeContext* InMergeCon
 	//  can render them without requesting the presence of other target layers (e.g. no weight-blending)
 	InMergeContext->ForEachTargetLayer(SupportedTargetTypeState.GetActiveWeightmapBitIndices(), 
 		[InMergeContext, InRenderer, this](int32 InTargetLayerIndex, const FName& InTargetLayerName, ULandscapeLayerInfoObject* InWeightmapLayerInfo)
-		{
+	{
 			const int32 RenderGroupCount = Algo::CountIf(RenderGroups, [InTargetLayerIndex](const TBitArray<>& InRenderGroup) { return InRenderGroup[InTargetLayerIndex]; });
-			checkf(RenderGroupCount < 2, TEXT("Target layer %s belongs to more than 1 render group in edit layer renderer %s. This is forbidden: in the end, it must belong to 1 and 1 only."),
+		checkf(RenderGroupCount < 2, TEXT("Target layer %s belongs to more than 1 render group in edit layer renderer %s. This is forbidden: in the end, it must belong to 1 and 1 only."),
 				*InTargetLayerName.ToString(), *InRenderer->GetEditLayerRendererDebugName());
-			if (RenderGroupCount == 0)
-			{
+		if (RenderGroupCount == 0)
+		{
 				TBitArray<> RenderGroup = InMergeContext->BuildTargetLayerBitIndices(false);
 				RenderGroup[InTargetLayerIndex] = true;
 				RenderGroups.Add(RenderGroup);
-			}
+}
 			return true;
 		});
 
@@ -608,7 +608,7 @@ void FEditLayerRendererState::EnableWeightmap(ELandscapeToolTargetType InTargetT
 		UpdateActiveTargetTypeState();
 	}
 }
-
+	
 void FEditLayerRendererState::EnableWeightmapChecked(ELandscapeToolTargetType InTargetType, const FName& InWeightmapLayerName)
 {
 	EnableWeightmap(InTargetType, (InWeightmapLayerName != NAME_None) ? MergeContext->GetTargetLayerIndexForNameChecked(InWeightmapLayerName) : INDEX_NONE);
@@ -913,7 +913,7 @@ void FMergeRenderContext::FreeBatchResources(const FMergeRenderBatch& InRenderBa
 	PerTargetLayerValidityRenderTargets.Empty();
 }
 
-void FMergeRenderContext::CycleBlendRenderTargets(ERHIAccess InDesiredWriteAccess)
+void FMergeRenderContext::CycleBlendRenderTargets(ERHIAccess InDesiredWriteAccess, UE::Landscape::FRDGBuilderRecorder& RDGBuilderRecorder)
 {
 	const bool bFirstWrite = (CurrentBlendRenderTargetWriteIndex < 0);
 	CurrentBlendRenderTargetWriteIndex = (CurrentBlendRenderTargetWriteIndex + 1) % BlendRenderTargets.Num();
@@ -923,11 +923,11 @@ void FMergeRenderContext::CycleBlendRenderTargets(ERHIAccess InDesiredWriteAcces
 		// Optionally clear the write render target for debug purposes : 
 		if (CVarLandscapeEditLayersClearBeforeEachWriteToScratch.GetValueOnGameThread())
 		{
-			GetBlendRenderTargetWrite()->Clear();
+			GetBlendRenderTargetWrite()->Clear(RDGBuilderRecorder);
 		}
 
 		// Change the state of the new Read and Write (ReadPrevious is already SRV since it was Read before)
-		ULandscapeScratchRenderTarget::FTransitionBatcherScope TransitionsScope;
+		ULandscapeScratchRenderTarget::FTransitionBatcherScope TransitionsScope(RDGBuilderRecorder);
 		TransitionsScope.TransitionTo(GetBlendRenderTargetRead(), ERHIAccess::SRVMask);
 		if (InDesiredWriteAccess != ERHIAccess::None)
 		{
@@ -977,8 +977,10 @@ void FMergeRenderContext::ResetVisualLogOffset()
 	CurrentVisualLogOffset = FVector(ForceInitToZero);
 }
 
-void FMergeRenderContext::RenderValidityRenderTargets(const FMergeRenderBatch& InRenderBatch)
+void FMergeRenderContext::RenderValidityRenderTargets(const FMergeRenderBatch& InRenderBatch, UE::Landscape::FRDGBuilderRecorder& RDGBuilderRecorder)
 {
+	using namespace UE::Landscape;
+
 	struct FTextureAndRects
 	{
 		FTextureAndRects(const FName& InTargetLayerName, const FString& InTextureDebugName, FTextureResource* InTextureResource)
@@ -996,13 +998,13 @@ void FMergeRenderContext::RenderValidityRenderTargets(const FMergeRenderBatch& I
 	TArray<FTextureAndRects> TexturesAndRects;
 	TexturesAndRects.Reserve(InRenderBatch.TargetLayerBitIndices.CountSetBits());
 	ForEachTargetLayer(InRenderBatch.TargetLayerBitIndices,
-		[this, &TexturesAndRects, &InRenderBatch](int32 InTargetLayerIndex, const FName& InTargetLayerName, ULandscapeLayerInfoObject* InWeightmapLayerInfo)
+		[this, &RDGBuilderRecorder, &TexturesAndRects, &InRenderBatch](int32 InTargetLayerIndex, const FName& InTargetLayerName, ULandscapeLayerInfoObject* InWeightmapLayerInfo)
 		{
 			ULandscapeScratchRenderTarget* ScratchRenderTarget = PerTargetLayerValidityRenderTargets.FindChecked(InTargetLayerName);
 			check(ScratchRenderTarget != nullptr);
 
 			// Make sure the validity mask is entirely cleared first:
-			ScratchRenderTarget->Clear();
+			ScratchRenderTarget->Clear(RDGBuilderRecorder);
 
 			FTextureAndRects& TextureAndRects = TexturesAndRects.Emplace_GetRef(InTargetLayerName, ScratchRenderTarget->GetDebugName(), ScratchRenderTarget->GetRenderTarget2D()->GetResource());
 
@@ -1015,49 +1017,53 @@ void FMergeRenderContext::RenderValidityRenderTargets(const FMergeRenderBatch& I
 				TextureAndRects.Rects.Add(FUintVector4(ComponentRect.Min.X, ComponentRect.Min.Y, ComponentRect.Max.X + 1, ComponentRect.Max.Y + 1));
 			}
 
-			ScratchRenderTarget->TransitionTo(ERHIAccess::RTV);
+			ScratchRenderTarget->TransitionTo(ERHIAccess::RTV, RDGBuilderRecorder);
 			return true;
 		});
 
-	ENQUEUE_RENDER_COMMAND(MarkTargetLayersValidity)([TexturesAndRects = MoveTemp(TexturesAndRects)](FRHICommandListImmediate& InRHICmdList)
-	{
-		FRDGBuilder GraphBuilder(InRHICmdList, RDG_EVENT_NAME("MarkTargetLayersValidity"));
-
-		for (const FTextureAndRects& TextureAndRects : TexturesAndRects)
+	auto RDGCommand =
+		[TexturesAndRects = MoveTemp(TexturesAndRects)](FRDGBuilder& GraphBuilder)
 		{
-			FRDGBufferRef RectBuffer = CreateUploadBuffer(GraphBuilder, TEXT("MarkValidityRects"), MakeArrayView(TextureAndRects.Rects));
-			FRDGBufferSRVRef RectBufferSRV = GraphBuilder.CreateSRV(FRDGBufferSRVDesc(RectBuffer, PF_R32G32B32A32_UINT));
-			FRDGTextureRef OutputTexture = GraphBuilder.RegisterExternalTexture(CreateRenderTarget(TextureAndRects.TextureResource->GetTexture2DRHI(), TEXT("ValidityMask")));
-
-			FMarkValidityPSParameters* PassParameters = GraphBuilder.AllocParameters<FMarkValidityPSParameters>();
-			PassParameters->RenderTargets[0] = FRenderTargetBinding(OutputTexture, ERenderTargetLoadAction::ELoad);
-			FGlobalShaderMap* ShaderMap = GetGlobalShaderMap(GMaxRHIFeatureLevel);
-			TShaderRef<FMarkValidityPS> PixelShader = ShaderMap->GetShader<FMarkValidityPS>();
-
-			FPixelShaderUtils::AddRasterizeToRectsPass<FMarkValidityPS>(GraphBuilder,
-				ShaderMap,
-				RDG_EVENT_NAME("MarkValidity(%s) -> %s", *TextureAndRects.TargetLayerName.ToString(), *TextureAndRects.TextureDebugName),
-				PixelShader,
-				PassParameters,
-				/*ViewportSize = */OutputTexture->Desc.Extent,
-				RectBufferSRV,
-				TextureAndRects.Rects.Num(),
-				/*BlendState = */ nullptr,
-				/*RasterizerState = */ nullptr,
-				/*DepthStencilState = */ nullptr,
-				/*StencilRef = */ 0,
-				/*TextureSize = */ OutputTexture->Desc.Extent);
-
-			// We need to specify the final state of the external texture to prevent the graph builder from transitioning it to SRVMask :
-			GraphBuilder.SetTextureAccessFinal(OutputTexture, ERHIAccess::RTV);
-		}
+			for (const FTextureAndRects& TextureAndRects : TexturesAndRects)
+			{
+				FRDGBufferRef RectBuffer = CreateUploadBuffer(GraphBuilder, TEXT("MarkValidityRects"), MakeArrayView(TextureAndRects.Rects));
+				FRDGBufferSRVRef RectBufferSRV = GraphBuilder.CreateSRV(FRDGBufferSRVDesc(RectBuffer, PF_R32G32B32A32_UINT));
+				FRDGTextureRef OutputTexture = GraphBuilder.RegisterExternalTexture(CreateRenderTarget(TextureAndRects.TextureResource->GetTexture2DRHI(), TEXT("ValidityMask")));
+
+				FMarkValidityPSParameters* PassParameters = GraphBuilder.AllocParameters<FMarkValidityPSParameters>();
+				PassParameters->RenderTargets[0] = FRenderTargetBinding(OutputTexture, ERenderTargetLoadAction::ELoad);
+				FGlobalShaderMap* ShaderMap = GetGlobalShaderMap(GMaxRHIFeatureLevel);
+				TShaderRef<FMarkValidityPS> PixelShader = ShaderMap->GetShader<FMarkValidityPS>();
+
+				FPixelShaderUtils::AddRasterizeToRectsPass<FMarkValidityPS>(GraphBuilder,
+					ShaderMap,
+					RDG_EVENT_NAME("MarkValidity(%s) -> %s", *TextureAndRects.TargetLayerName.ToString(), *TextureAndRects.TextureDebugName),
+					PixelShader,
+					PassParameters,
+					/*ViewportSize = */OutputTexture->Desc.Extent,
+					RectBufferSRV,
+					TextureAndRects.Rects.Num(),
+					/*BlendState = */ nullptr,
+					/*RasterizerState = */ nullptr,
+					/*DepthStencilState = */ nullptr,
+					/*StencilRef = */ 0,
+					/*TextureSize = */ OutputTexture->Desc.Extent);
+			}
+		};
 
-		GraphBuilder.Execute();
-	});
+	// We need to specify the final state of the external textures to prevent the graph builder from transitioning them to SRVMask :
+	TArray<FRDGBuilderRecorder::FRDGExternalTextureAccessFinal> RDGExternalTextureAccessFinalList;
+	Algo::Transform(TexturesAndRects, RDGExternalTextureAccessFinalList, [](const FTextureAndRects& InTextureAndRects) -> FRDGBuilderRecorder::FRDGExternalTextureAccessFinal
+		{ 
+			return { InTextureAndRects.TextureResource, ERHIAccess::RTV };
+		});
+	RDGBuilderRecorder.EnqueueRDGCommand(RDGCommand, RDGExternalTextureAccessFinalList);
 }
 
-void FMergeRenderContext::RenderExpandedRenderTarget(const FMergeRenderBatch& InRenderBatch)
+void FMergeRenderContext::RenderExpandedRenderTarget(const FMergeRenderBatch& InRenderBatch, UE::Landscape::FRDGBuilderRecorder& RDGBuilderRecorder)
 {
+	using namespace UE::Landscape;
+
 	TArray<FUintVector4> SourceRects, DestinationRects;
 	{
 		TArray<FIntRect> SourceInclusiveRects, DestinationInclusiveRects;
@@ -1069,71 +1075,76 @@ void FMergeRenderContext::RenderExpandedRenderTarget(const FMergeRenderBatch& In
 
 	ULandscapeScratchRenderTarget* WriteRT = GetBlendRenderTargetWrite();
 	ULandscapeScratchRenderTarget* ReadRT = GetBlendRenderTargetRead();
-	WriteRT->TransitionTo(ERHIAccess::RTV);
-	ReadRT->TransitionTo(ERHIAccess::SRVMask);
+	WriteRT->TransitionTo(ERHIAccess::RTV, RDGBuilderRecorder);
+	ReadRT->TransitionTo(ERHIAccess::SRVMask, RDGBuilderRecorder);
 
 	FSceneInterface* SceneInterface = GetLandscape()->GetWorld()->Scene;
 
-	ENQUEUE_RENDER_COMMAND(Expand)(
+	auto RDGCommand =
 		[ SourceRects = MoveTemp(SourceRects)
 		, DestinationRects = MoveTemp(DestinationRects)
 		, OutputResource = WriteRT->GetRenderTarget2D()->GetResource()
 		, SourceResource = ReadRT->GetRenderTarget2D()->GetResource()
-		, SceneInterface] (FRHICommandListImmediate& InRHICmdList)
+		, SceneInterface] (FRDGBuilder& GraphBuilder)
 	{
-		FRDGBuilder GraphBuilder(InRHICmdList, RDG_EVENT_NAME("Expand"));
-
-		FRDGBufferRef RectBuffer = CreateUploadBuffer(GraphBuilder, TEXT("ExpandRects"), MakeArrayView(DestinationRects));
-		FRDGBufferSRVRef RectBufferSRV = GraphBuilder.CreateSRV(FRDGBufferSRVDesc(RectBuffer, PF_R32G32B32A32_UINT));
-
-		FRDGBufferRef RectUVBuffer = CreateUploadBuffer(GraphBuilder, TEXT("ExpandRectsUVs"), MakeArrayView(SourceRects));
-		FRDGBufferSRVRef RectUVBufferSRV = GraphBuilder.CreateSRV(FRDGBufferSRVDesc(RectUVBuffer, PF_R32G32B32A32_UINT));
-
-		FRDGTextureRef OutputTexture = GraphBuilder.RegisterExternalTexture(CreateRenderTarget(OutputResource->GetTexture2DRHI(), TEXT("OutputTexture")));
-		FRDGTextureRef SourceTexture = GraphBuilder.RegisterExternalTexture(CreateRenderTarget(SourceResource->GetTexture2DRHI(), TEXT("SourceTexture")));
-
-		// TODO [jonathan.bard] this is just a rhi validation error for unoptimized shaders... once validation is made to not issue those errors, we can remove this
-		// Create a SceneView to please the shader bindings, but it's unused in practice 
-		FSceneViewFamilyContext ViewFamily(FSceneViewFamily::ConstructionValues(nullptr, SceneInterface, FEngineShowFlags(ESFIM_Game)).SetTime(FGameTime::GetTimeSinceAppStart()));
-		FSceneViewInitOptions ViewInitOptions;
-		ViewInitOptions.ViewFamily = &ViewFamily;
-		ViewInitOptions.SetViewRectangle(FIntRect(0, 0, 1, 1)); // Use a dummy rect to avoid a check(slow)
-		GetRendererModule().CreateAndInitSingleView(InRHICmdList, &ViewFamily, &ViewInitOptions);
-		const FSceneView* View = ViewFamily.Views[0];
-
-		FCopyQuadsPSParameters* PassParameters = GraphBuilder.AllocParameters<FCopyQuadsPSParameters>();
-		PassParameters->RenderTargets[0] = FRenderTargetBinding(OutputTexture, ERenderTargetLoadAction::ELoad);
-		PassParameters->PS.View = View->ViewUniformBuffer;
-		PassParameters->PS.InSourceTexture = SourceTexture;
-
-		FGlobalShaderMap* ShaderMap = GetGlobalShaderMap(GMaxRHIFeatureLevel);
-		TShaderRef<FCopyQuadsPS> PixelShader = ShaderMap->GetShader<FCopyQuadsPS>();
-
-		FPixelShaderUtils::AddRasterizeToRectsPass<FCopyQuadsPS>(GraphBuilder,
-			ShaderMap,
-			RDG_EVENT_NAME("CopyQuadsPS"),
-			PixelShader,
-			PassParameters,
-			/*ViewportSize = */OutputTexture->Desc.Extent,
-			RectBufferSRV,
-			DestinationRects.Num(),
-			/*BlendState = */nullptr,
-			/*RasterizerState = */nullptr,
-			/*DepthStencilState = */nullptr,
-			/*StencilRef = */0,
-			/*TextureSize = */SourceTexture->Desc.Extent,
-			RectUVBufferSRV);
-
-		// We need to specify the final state of the external textures to prevent the graph build from transitioning them to SRVMask :
-		GraphBuilder.SetTextureAccessFinal(OutputTexture, ERHIAccess::RTV);
-		GraphBuilder.SetTextureAccessFinal(SourceTexture, ERHIAccess::SRVMask);
-
-		GraphBuilder.Execute();
-	});
-}
-
-void FMergeRenderContext::Render(TFunction<void(const FOnRenderBatchTargetGroupDoneParams&)> OnBatchTargetGroupDone)
+			FRDGBufferRef RectBuffer = CreateUploadBuffer(GraphBuilder, TEXT("ExpandRects"), MakeArrayView(DestinationRects));
+			FRDGBufferSRVRef RectBufferSRV = GraphBuilder.CreateSRV(FRDGBufferSRVDesc(RectBuffer, PF_R32G32B32A32_UINT));
+
+			FRDGBufferRef RectUVBuffer = CreateUploadBuffer(GraphBuilder, TEXT("ExpandRectsUVs"), MakeArrayView(SourceRects));
+			FRDGBufferSRVRef RectUVBufferSRV = GraphBuilder.CreateSRV(FRDGBufferSRVDesc(RectUVBuffer, PF_R32G32B32A32_UINT));
+
+			FRDGTextureRef OutputTexture = GraphBuilder.RegisterExternalTexture(CreateRenderTarget(OutputResource->GetTexture2DRHI(), TEXT("OutputTexture")));
+			FRDGTextureRef SourceTexture = GraphBuilder.RegisterExternalTexture(CreateRenderTarget(SourceResource->GetTexture2DRHI(), TEXT("SourceTexture")));
+
+			// TODO [jonathan.bard] this is just a rhi validation error for unoptimized shaders... once validation is made to not issue those errors, we can remove this
+			// Create a SceneView to please the shader bindings, but it's unused in practice 
+			FSceneViewFamilyContext ViewFamily(FSceneViewFamily::ConstructionValues(nullptr, SceneInterface, FEngineShowFlags(ESFIM_Game)).SetTime(FGameTime::GetTimeSinceAppStart()));
+			FSceneViewInitOptions ViewInitOptions;
+			ViewInitOptions.ViewFamily = &ViewFamily;
+			ViewInitOptions.SetViewRectangle(FIntRect(0, 0, 1, 1)); // Use a dummy rect to avoid a check(slow)
+			GetRendererModule().CreateAndInitSingleView(GraphBuilder.RHICmdList, &ViewFamily, &ViewInitOptions);
+			const FSceneView* View = ViewFamily.Views[0];
+
+			FCopyQuadsPSParameters* PassParameters = GraphBuilder.AllocParameters<FCopyQuadsPSParameters>();
+			PassParameters->RenderTargets[0] = FRenderTargetBinding(OutputTexture, ERenderTargetLoadAction::ELoad);
+			PassParameters->PS.View = View->ViewUniformBuffer;
+			PassParameters->PS.InSourceTexture = SourceTexture;
+
+			FGlobalShaderMap* ShaderMap = GetGlobalShaderMap(GMaxRHIFeatureLevel);
+			TShaderRef<FCopyQuadsPS> PixelShader = ShaderMap->GetShader<FCopyQuadsPS>();
+
+			FPixelShaderUtils::AddRasterizeToRectsPass<FCopyQuadsPS>(GraphBuilder,
+				ShaderMap,
+				RDG_EVENT_NAME("CopyQuadsPS"),
+				PixelShader,
+				PassParameters,
+				/*ViewportSize = */OutputTexture->Desc.Extent,
+				RectBufferSRV,
+				DestinationRects.Num(),
+				/*BlendState = */nullptr,
+				/*RasterizerState = */nullptr,
+				/*DepthStencilState = */nullptr,
+				/*StencilRef = */0,
+				/*TextureSize = */SourceTexture->Desc.Extent,
+				RectUVBufferSRV);
+		};
+
+	// We need to specify the final state of the external textures to prevent the graph builder from transitioning them to SRVMask (even those that end up as SRVMask at the end of this command,
+	//  because they will likely be part of another RDGCommand down the line so we need to maintain an accurate picture of every external texture ever involved in the recorded command so that 
+	//  we can set a proper access when the recorder is flushed (and the FRDGBuilder, executed) :
+	TArray<FRDGBuilderRecorder::FRDGExternalTextureAccessFinal> RDGExternalTextureAccessFinalList =
+	{
+		{ WriteRT->GetRenderTarget()->GetResource(), ERHIAccess::RTV },
+		{ ReadRT->GetRenderTarget()->GetResource(), ERHIAccess::SRVMask }
+	};
+	RDGBuilderRecorder.EnqueueRDGCommand(RDGCommand, RDGExternalTextureAccessFinalList);
+}
+
+void FMergeRenderContext::Render(TFunction<void(const FOnRenderBatchTargetGroupDoneParams& /*InParams*/, UE::Landscape::FRDGBuilderRecorder& /*RDGBuilderRecorder*/)> OnBatchTargetGroupDone)
 {
+	using namespace UE::Landscape;
+	using namespace UE::Landscape::EditLayers;
+
 	TRACE_CPUPROFILER_EVENT_SCOPE(FMergeRenderContext::Render);
 
 	check(CurrentRenderBatchIndex == INDEX_NONE);
@@ -1147,8 +1158,11 @@ void FMergeRenderContext::Render(TFunction<void(const FOnRenderBatchTargetGroupD
 
 	AllocateResources();
 
+	// Command recorder when accumulating render commands for several consecutive RenderLayer calls when ERenderFlags::Recorded is used 
+	FRDGBuilderRecorder RDGBuilderRecorder;
+
 	// Kick start the blend render targets : 
-	CycleBlendRenderTargets();
+	CycleBlendRenderTargets(/*InDesiredWriteAccess = */ERHIAccess::None, RDGBuilderRecorder);
 
 	const int32 NumBatches = RenderBatches.Num();
 	for (CurrentRenderBatchIndex = 0; CurrentRenderBatchIndex < NumBatches; ++CurrentRenderBatchIndex)
@@ -1157,8 +1171,8 @@ void FMergeRenderContext::Render(TFunction<void(const FOnRenderBatchTargetGroupD
 		FString RenderBatchDebugName = FString::Format(TEXT("Render Batch [{0}] : ({1},{2})->({3},{4})"), { CurrentRenderBatchIndex, RenderBatch.SectionRect.Min.X, RenderBatch.SectionRect.Min.Y, RenderBatch.SectionRect.Max.X, RenderBatch.SectionRect.Max.Y });
 		RHI_BREADCRUMB_EVENT_GAMETHREAD("%s", RenderBatchDebugName);
 
-		checkf((RenderBatch.RenderSteps.Num() >= 1) && (RenderBatch.RenderSteps.Last().Type == FMergeRenderStep::EType::SignalBatchMergeGroupDone), TEXT("Any batch should end with a SignalBatchMergeGroupDone step and there \
-			should be at least another step prior to that, otherwise, the batch is just useless."));
+		checkf((RenderBatch.RenderSteps.Num() >= 1) && (RenderBatch.RenderSteps.Last().Type == FMergeRenderStep::EType::SignalBatchMergeGroupDone),
+			TEXT("Any batch should end with a SignalBatchMergeGroupDone step and there should be at least another step prior to that, otherwise, the batch is just useless."));
 
 		AllocateBatchResources(RenderBatch);
 
@@ -1179,89 +1193,149 @@ void FMergeRenderContext::Render(TFunction<void(const FOnRenderBatchTargetGroupD
 			});
 
 		const int32 NumRenderSteps = RenderBatch.RenderSteps.Num();
+		// Current index of RenderLayer (for debugging purposes)
+		int32 RenderLayerStepIndex = 0;
+		// Index of RenderLayer at which we started recording the current render command sequence
+		int32 RenderCommandStartLayerStepIndex = INDEX_NONE;
 		for (int32 RenderStepIndex = 0; RenderStepIndex < NumRenderSteps; ++RenderStepIndex)
 		{
 			const FMergeRenderStep& RenderStep = RenderBatch.RenderSteps[RenderStepIndex];
-			TArray<FName> RenderGroupTargetLayerNames = ConvertTargetLayerBitIndicesToNames(RenderStep.RenderGroupBitIndices);
-			TArray<ULandscapeLayerInfoObject*> RenderGroupTargetLayerInfos = bIsHeightmapMerge ? TArray<ULandscapeLayerInfoObject*> { nullptr } : ConvertTargetLayerBitIndicesToLayerInfos(RenderStep.RenderGroupBitIndices);
 
-			// Compute all necessary info about the components affected by this renderer at this step
+			TArray<FName> RenderGroupTargetLayerNames;
+			TArray<ULandscapeLayerInfoObject*> RenderGroupTargetLayerInfos;
 			TArray<FComponentMergeRenderInfo> SortedComponentMergeRenderInfos;
-			SortedComponentMergeRenderInfos.Reserve(RenderStep.ComponentsToRender.Num());
-			Algo::Transform(RenderStep.ComponentsToRender, SortedComponentMergeRenderInfos, [MinComponentKey = RenderBatch.MinComponentKey](ULandscapeComponent* InComponent)
-			{
-				FComponentMergeRenderInfo ComponentMergeRenderInfo;
-				ComponentMergeRenderInfo.Component = InComponent;
 
-				const FIntPoint ComponentKey = InComponent->GetComponentKey();
-				const FIntPoint LocalComponentKey = ComponentKey - MinComponentKey;
-				check((LocalComponentKey.X >= 0) && (LocalComponentKey.Y >= 0));
-				ComponentMergeRenderInfo.ComponentKeyInRenderArea = LocalComponentKey;
-				// Area in the render target for this component : 
-				ComponentMergeRenderInfo.ComponentRegionInRenderArea = FIntRect(LocalComponentKey * InComponent->ComponentSizeQuads, (LocalComponentKey + 1) * InComponent->ComponentSizeQuads);
+			// Compute some data for the actual render steps : 
+			if ((RenderStep.Type == FMergeRenderStep::EType::RenderLayer_Immediate)
+				|| (RenderStep.Type == FMergeRenderStep::EType::RenderLayer_Recorded)
+				|| (RenderStep.Type == FMergeRenderStep::EType::SignalBatchMergeGroupDone))
+			{
+				RenderGroupTargetLayerNames = ConvertTargetLayerBitIndicesToNames(RenderStep.RenderGroupBitIndices);
+				RenderGroupTargetLayerInfos = bIsHeightmapMerge ? TArray<ULandscapeLayerInfoObject*> { nullptr } : ConvertTargetLayerBitIndicesToLayerInfos(RenderStep.RenderGroupBitIndices);
 
-				return ComponentMergeRenderInfo;
-			});
-			SortedComponentMergeRenderInfos.Sort();
+				// Compute all necessary info about the components affected by this renderer at this step
+				SortedComponentMergeRenderInfos.Reserve(RenderStep.ComponentsToRender.Num());
+				Algo::Transform(RenderStep.ComponentsToRender, SortedComponentMergeRenderInfos, [MinComponentKey = RenderBatch.MinComponentKey](ULandscapeComponent* InComponent)
+				{
+					FComponentMergeRenderInfo ComponentMergeRenderInfo;
+					ComponentMergeRenderInfo.Component = InComponent;
+
+					const FIntPoint ComponentKey = InComponent->GetComponentKey();
+					const FIntPoint LocalComponentKey = ComponentKey - MinComponentKey;
+					check((LocalComponentKey.X >= 0) && (LocalComponentKey.Y >= 0));
+					ComponentMergeRenderInfo.ComponentKeyInRenderArea = LocalComponentKey;
+					// Area in the render target for this component : 
+					ComponentMergeRenderInfo.ComponentRegionInRenderArea = FIntRect(LocalComponentKey * InComponent->ComponentSizeQuads, (LocalComponentKey + 1) * InComponent->ComponentSizeQuads);
+
+					return ComponentMergeRenderInfo;
+				});
+				SortedComponentMergeRenderInfos.Sort();
+			}
 
-			// Is it a step involving a renderer ? 
-			if (ILandscapeEditLayerRenderer* Renderer = RenderStep.RendererState.GetRenderer())
+			switch (RenderStep.Type)
 			{
-				if (RenderStep.Type == FMergeRenderStep::EType::RenderLayer)
-				{
-					TRACE_CPUPROFILER_EVENT_SCOPE(RenderLayer);
+			case FMergeRenderStep::EType::BeginRenderCommandRecorder:
+			{
+				// Start recording a new sequence of RDG commands in order to let the upcoming FMergeRenderStep::EType::RenderLayer_Recorded steps push their render thread-based "operations" (lambdas)
+				checkf(RDGBuilderRecorder.IsEmpty() && !RDGBuilderRecorder.IsRecording(), TEXT("There shouldn't be any pending command being recorded when starting a new render command"));
+				checkf(RenderCommandStartLayerStepIndex == INDEX_NONE, TEXT("RenderCommandStartLayerStepIndex should be invalid as no render command recording should be active"));
+
+				// Remember the render layer step index at which we started recording the RDG render command : 
+				RenderCommandStartLayerStepIndex = RenderLayerStepIndex;
+				check(RenderCommandStartLayerStepIndex >= 0);
+				RDGBuilderRecorder.StartRecording();
+				break;
+			}
+			case FMergeRenderStep::EType::EndRenderCommandRecorder:
+			{
+				checkf(RDGBuilderRecorder.IsRecording(), TEXT("There should be a pending command being recorded when starting a ending a render command"));
+				checkf(RenderCommandStartLayerStepIndex != INDEX_NONE, TEXT("We should have initiated the render command recording with a begin operation, which should set a valid RenderCommandStartLayerStepIndex"));
+				check(RenderCommandStartLayerStepIndex <= RenderLayerStepIndex);
+
+				// This is where we actually push the current render command with all the render thread operations that have been accumulated on the render command context : 
+				FRDGEventName RDGEventName = (RenderCommandStartLayerStepIndex == RenderLayerStepIndex - 1) 
+					? RDG_EVENT_NAME("Step [%i]", RenderCommandStartLayerStepIndex)
+					: RDG_EVENT_NAME("Steps [%i-%i]", RenderCommandStartLayerStepIndex, RenderLayerStepIndex - 1);
+				RDGBuilderRecorder.StopRecordingAndFlush(MoveTemp(RDGEventName));
+
+				// We've flushed the render command, we can reset the starting render layer step index : 
+				RenderCommandStartLayerStepIndex = INDEX_NONE;
+				break;
+			}
+			case FMergeRenderStep::EType::RenderLayer_Immediate:
+			case FMergeRenderStep::EType::RenderLayer_Recorded:
+			{
+				TRACE_CPUPROFILER_EVENT_SCOPE(RenderLayer);
+
+				ILandscapeEditLayerRenderer* Renderer = RenderStep.RendererState.GetRenderer();
+				check((Renderer != nullptr) && EnumHasAnyFlags(Renderer->GetRenderFlags(this), ERenderFlags::RenderMode_Mask));
+
+				// TODO[jonathan.bard] offset the world transform to account for the half-pixel offset?
+				//RenderParams.RenderAreaWorldTransform = FTransform(LandscapeTransform.GetRotation(), LandscapeTransform.GetTranslation() + FVector(RenderBatch.SectionRect.Min) /** (double)ComponentSizeQuads - FVector(0.5, 0.5, 0)*/, LandscapeTransform.GetScale3D());
+
+				FString RenderStepProfilingEventName = FString::Format(TEXT("Step [{0}] ({1}): Render {2}"), { RenderLayerStepIndex, *ConvertTargetLayerNamesToString(RenderGroupTargetLayerNames), *Renderer->GetEditLayerRendererDebugName() });
+
+				// TODO [jonathan.bard] : this is more of a Batch world transform/section rect at the moment. Shall we have a RenderAreaWorldTransform/RenderAreaSectionRect in FRenderParams and a BatchRenderAreaWorldTransform in FMergeRenderBatch?
+				//  because currently the old BP brushes work with FMergeRenderBatch data (i.e. 1 transform for the batch and a section rect for the entire batch) but eventually, renderers might be interested in just their Render step context, 
+				//  that is : 1 matrix corresponding to the bottom left corner of their list of components to render?
+				FTransform RenderAreaWorldTransform = FTransform(FVector(RenderBatch.SectionRect.Min)) * LandscapeTransform;
+				FIntRect RenderAreaSectionRect = RenderBatch.SectionRect;
+
+				// Drop some visual cues to help understand how each renderer is applied :
+				UE_IFVLOG(
+					if (IsVisualLogEnabled() && bVisualLogMergeProcess)
+					{
+						FTransform RenderAreaWorldTransformForVisLog = FTransform(FVector(RenderBatch.SectionRect.Min)) * LandscapeWorldTransformForVisLog;
+						IncrementVisualLogOffset();
+						UE_VLOG_LOCATION(Landscape, LogLandscape, Log, ComputeVisualLogTransform(RenderAreaWorldTransformForVisLog).GetTranslation(), 10.0f, FColor::Red, TEXT("%s"), *RenderStepProfilingEventName);
+						UE_VLOG_WIREOBOX(Landscape, LogLandscape, Log, FBox(FVector(RenderBatch.SectionRect.Min) - FVector(0.5, 0.5, 0.0), FVector(RenderBatch.SectionRect.Max) - FVector(0.5, 0.5, 0.0)),
+							ComputeVisualLogTransform(LandscapeWorldTransformForVisLog).ToMatrixWithScale(), FColor::White, TEXT(""));
+
+						// Draw each component's bounds rendered by this renderer : 
+						for (const FComponentMergeRenderInfo& ComponentMergeRenderInfo : SortedComponentMergeRenderInfos)
+						{
+							UE_VLOG_WIREOBOX(Landscape, LogLandscape, Log, FBox(FVector(ComponentMergeRenderInfo.ComponentRegionInRenderArea.Min), FVector(ComponentMergeRenderInfo.ComponentRegionInRenderArea.Max)),
+								ComputeVisualLogTransform(RenderAreaWorldTransformForVisLog).ToMatrixWithScale(), FColor::White, TEXT(""));
+						}
+					});
 
-					check(Renderer->CanRender());
+				ILandscapeEditLayerRenderer::FRenderParams RenderParams(this, RenderGroupTargetLayerNames, RenderGroupTargetLayerInfos, RenderStep.RendererState, SortedComponentMergeRenderInfos, RenderAreaWorldTransform, RenderAreaSectionRect);
+				if (RenderStep.Type == FMergeRenderStep::EType::RenderLayer_Recorded)
+				{
+					checkf(RDGBuilderRecorder.IsRecording(), TEXT("RenderLayer_Recorded must be preceded by a BeginRenderCommandRecorder which should create a command recorder"));
 
-					// TODO[jonathan.bard] offset the world transform to account for the half-pixel offset?
-					//RenderParams.RenderAreaWorldTransform = FTransform(LandscapeTransform.GetRotation(), LandscapeTransform.GetTranslation() + FVector(RenderBatch.SectionRect.Min) /** (double)ComponentSizeQuads - FVector(0.5, 0.5, 0)*/, LandscapeTransform.GetScale3D());
+					RDG_RENDER_COMMAND_RECORDER_BREADCRUMB_EVENT(RDGBuilderRecorder, "%s", *RenderStepProfilingEventName);
+					Renderer->RenderLayer(RenderParams, RDGBuilderRecorder);
+				}
+				else
+				{
+					checkf(RDGBuilderRecorder.IsEmpty() && !RDGBuilderRecorder.IsRecording(), TEXT("RenderLayer_Immediate should be preceded by a EndRenderCommandRecorder which should finalize a command recorder and destroy it"));
 
-					FString RenderStepProfilingEventName = FString::Format(TEXT("Step [{0}] ({1}): Render {2}"), { RenderStepIndex, *ConvertTargetLayerNamesToString(RenderGroupTargetLayerNames), *Renderer->GetEditLayerRendererDebugName() });
 					RHI_BREADCRUMB_EVENT_GAMETHREAD("%s", RenderStepProfilingEventName);
-
-					// TODO [jonathan.bard] : this is more of a Batch world transform/section rect at the moment. Shall we have a RenderAreaWorldTransform/RenderAreaSectionRect in FRenderParams and a BatchRenderAreaWorldTransform in FMergeRenderBatch?
-					//  because currently the old BP brushes work with FMergeRenderBatch data (i.e. 1 transform for the batch and a section rect for the entire batch) but eventually, renderers might be interested in just their Render step context, 
-					//  that is : 1 matrix corresponding to the bottom left corner of their list of components to render?
-					FTransform RenderAreaWorldTransform = FTransform(FVector(RenderBatch.SectionRect.Min)) * LandscapeTransform;
-					FIntRect RenderAreaSectionRect = RenderBatch.SectionRect;
-				
-					// Drop some visual cues to help understand how each renderer is applied :
-					UE_IFVLOG(
-						if (IsVisualLogEnabled() && bVisualLogMergeProcess)
-						{
-							FTransform RenderAreaWorldTransformForVisLog = FTransform(FVector(RenderBatch.SectionRect.Min)) * LandscapeWorldTransformForVisLog;
-							IncrementVisualLogOffset();
-							UE_VLOG_LOCATION(Landscape, LogLandscape, Log, ComputeVisualLogTransform(RenderAreaWorldTransformForVisLog).GetTranslation(), 10.0f, FColor::Red, TEXT("%s"), *RenderStepProfilingEventName);
-							UE_VLOG_WIREOBOX(Landscape, LogLandscape, Log, FBox(FVector(RenderBatch.SectionRect.Min) - FVector(0.5, 0.5, 0.0), FVector(RenderBatch.SectionRect.Max) - FVector(0.5, 0.5, 0.0)),
-								ComputeVisualLogTransform(LandscapeWorldTransformForVisLog).ToMatrixWithScale(), FColor::White, TEXT(""));
-
-							// Draw each component's bounds rendered by this renderer : 
-							for (const FComponentMergeRenderInfo& ComponentMergeRenderInfo : SortedComponentMergeRenderInfos)
-							{
-								UE_VLOG_WIREOBOX(Landscape, LogLandscape, Log, FBox(FVector(ComponentMergeRenderInfo.ComponentRegionInRenderArea.Min), FVector(ComponentMergeRenderInfo.ComponentRegionInRenderArea.Max)),
-									ComputeVisualLogTransform(RenderAreaWorldTransformForVisLog).ToMatrixWithScale(), FColor::White, TEXT(""));
-							}
-						});
-			
-					ILandscapeEditLayerRenderer::FRenderParams RenderParams(this, RenderGroupTargetLayerNames, RenderGroupTargetLayerInfos, RenderStep.RendererState, SortedComponentMergeRenderInfos, RenderAreaWorldTransform, RenderAreaSectionRect);
-					Renderer->RenderLayer(RenderParams);
+					Renderer->RenderLayer(RenderParams, RDGBuilderRecorder);
 				}
+
+				++RenderLayerStepIndex;
+				break;
 			}
-			else if ((RenderStep.Type == FMergeRenderStep::EType::SignalBatchMergeGroupDone))
+			case FMergeRenderStep::EType::SignalBatchMergeGroupDone:
 			{
 				TRACE_CPUPROFILER_EVENT_SCOPE(MergeGroupDone);
-				RHI_BREADCRUMB_EVENT_GAMETHREAD("Step [%i] (%s) : Render Group Done", RenderStepIndex, ConvertTargetLayerNamesToString(RenderGroupTargetLayerNames));
+				RHI_BREADCRUMB_EVENT_GAMETHREAD("Step [%i] (%s) : Render Group Done", RenderLayerStepIndex, ConvertTargetLayerNamesToString(RenderGroupTargetLayerNames));
 
+				checkf(RDGBuilderRecorder.IsEmpty() && !RDGBuilderRecorder.IsRecording(), TEXT("SignalBatchMergeGroupDone should be preceded by a EndRenderCommandRecorder which should finalize a command recorder and destroy it"));
 				// The last render target we wrote to is the one containing the batch group's merge result : 
 				FOnRenderBatchTargetGroupDoneParams Params(this, RenderBatch, RenderGroupTargetLayerNames, RenderGroupTargetLayerInfos, SortedComponentMergeRenderInfos);
-				OnBatchTargetGroupDone(Params);
+				OnBatchTargetGroupDone(Params, RDGBuilderRecorder);
+				break;
 			}
-			else
-			{
-				check(false); // unknown render step type
+			default:
+				checkf(false, TEXT("Invalid render step type"));
 			}
 		}
 		
+		checkf(RDGBuilderRecorder.IsEmpty() && !RDGBuilderRecorder.IsRecording(), TEXT("We should not have any pending render command recorder at the end of render"));
+
 		FreeBatchResources(RenderBatch);
 	}
 
diff --git a/Engine/Source/Runtime/Landscape/Private/LandscapeEditLayerRendererPrivate.cpp b/Engine/Source/Runtime/Landscape/Private/LandscapeEditLayerRendererPrivate.cpp
--- a/Engine/Source/Runtime/Landscape/Private/LandscapeEditLayerRendererPrivate.cpp
+++ b/Engine/Source/Runtime/Landscape/Private/LandscapeEditLayerRendererPrivate.cpp
@@ -43,21 +43,23 @@ FString ULandscapeDefaultEditLayerRenderer::GetEditLayerRendererDebugName() cons
 	return TEXT("Default");
 }
 
-void ULandscapeDefaultEditLayerRenderer::RenderLayer(FRenderParams& InRenderParams) 
+void ULandscapeDefaultEditLayerRenderer::RenderLayer(FRenderParams& RenderParams, UE::Landscape::FRDGBuilderRecorder& RDGBuilderRecorder)
 {
 	using namespace UE::Landscape::EditLayers;
 
-	FMergeRenderContext* RenderContext = InRenderParams.MergeRenderContext;
+	checkf(RDGBuilderRecorder.IsRecording(), TEXT("ERenderFlags::RenderMode_Recorded means the command recorder should be recording at this point"));
+
+	FMergeRenderContext* RenderContext = RenderParams.MergeRenderContext;
 	const FMergeRenderBatch* RenderBatch = RenderContext->GetCurrentRenderBatch();
 
-	RenderContext->CycleBlendRenderTargets(/*InDesiredWriteAccess = */ERHIAccess::RTV);
+	RenderContext->CycleBlendRenderTargets(/*InDesiredWriteAccess = */ERHIAccess::RTV, RDGBuilderRecorder);
 	ULandscapeScratchRenderTarget* WriteRT = RenderContext->GetBlendRenderTargetWrite();
 
 	// Start from a blank canvas so that the first layer is blended with nothing underneath :
-	WriteRT->Clear();
+	WriteRT->Clear(RDGBuilderRecorder);
 
 	// Render the components of the batch for each target layer into the "pseudo-stencil" buffer, so that it can be sampled by users as a UTextures in UMaterials and such : 
-	RenderContext->RenderValidityRenderTargets(*RenderBatch);
+	RenderContext->RenderValidityRenderTargets(*RenderBatch, RDGBuilderRecorder);
 }
 
 
diff --git a/Engine/Source/Runtime/Landscape/Private/LandscapeEditLayerRendererPrivate.h b/Engine/Source/Runtime/Landscape/Private/LandscapeEditLayerRendererPrivate.h
--- a/Engine/Source/Runtime/Landscape/Private/LandscapeEditLayerRendererPrivate.h
+++ b/Engine/Source/Runtime/Landscape/Private/LandscapeEditLayerRendererPrivate.h
@@ -33,7 +33,8 @@ public:
 		UE::Landscape::EditLayers::FEditLayerTargetTypeState& OutSupportedTargetTypeState, UE::Landscape::EditLayers::FEditLayerTargetTypeState& OutEnabledTargetTypeState, TArray<TBitArray<>>& OutRenderGroups) const override;
 	virtual TArray<UE::Landscape::EditLayers::FEditLayerRenderItem> GetRenderItems(const UE::Landscape::EditLayers::FMergeContext* InMergeContext) const override;
 	virtual FString GetEditLayerRendererDebugName() const override;
-	virtual void RenderLayer(FRenderParams& InRenderParams) override;
+	virtual UE::Landscape::EditLayers::ERenderFlags GetRenderFlags(const UE::Landscape::EditLayers::FMergeContext* InMergeContext) const override { return UE::Landscape::EditLayers::ERenderFlags::RenderMode_Recorded; }
+	virtual void RenderLayer(FRenderParams& RenderParams, UE::Landscape::FRDGBuilderRecorder& RDGBuilderRecorder) override;
 	//~ End ILandscapeEditLayerRenderer implementation
 #endif // WITH_EDITOR
 };
@@ -58,7 +59,8 @@ public:
 		UE::Landscape::EditLayers::FEditLayerTargetTypeState& OutSupportedTargetTypeState, UE::Landscape::EditLayers::FEditLayerTargetTypeState& OutEnabledTargetTypeState, TArray<TBitArray<>>& OutRenderGroups) const override;
 	virtual TArray<UE::Landscape::EditLayers::FEditLayerRenderItem> GetRenderItems(const UE::Landscape::EditLayers::FMergeContext* InMergeContext) const override;
 	virtual FString GetEditLayerRendererDebugName() const override;
-	virtual void RenderLayer(FRenderParams& InRenderParams) override;
+	virtual UE::Landscape::EditLayers::ERenderFlags GetRenderFlags(const UE::Landscape::EditLayers::FMergeContext* InMergeContext) const override { return UE::Landscape::EditLayers::ERenderFlags::RenderMode_Recorded; }
+	virtual void RenderLayer(FRenderParams& RenderParams, UE::Landscape::FRDGBuilderRecorder& RDGBuilderRecorder) override;
 	//~ End ILandscapeEditLayerRenderer implementation
 #endif // WITH_EDITOR
 };
@@ -81,7 +83,8 @@ public:
 		UE::Landscape::EditLayers::FEditLayerTargetTypeState& OutSupportedTargetTypeState, UE::Landscape::EditLayers::FEditLayerTargetTypeState& OutEnabledTargetTypeState, TArray<TBitArray<>>& OutRenderGroups) const override;
 	virtual TArray<UE::Landscape::EditLayers::FEditLayerRenderItem> GetRenderItems(const UE::Landscape::EditLayers::FMergeContext* InMergeContext) const override;
 	virtual FString GetEditLayerRendererDebugName() const override;
-	virtual void RenderLayer(FRenderParams& InRenderParams) override;
+	virtual UE::Landscape::EditLayers::ERenderFlags GetRenderFlags(const UE::Landscape::EditLayers::FMergeContext* InMergeContext) const override { return UE::Landscape::EditLayers::ERenderFlags::RenderMode_Recorded; }
+	virtual void RenderLayer(FRenderParams& RenderParams, UE::Landscape::FRDGBuilderRecorder& RDGBuilderRecorder) override;
 	//~ End ILandscapeEditLayerRenderer implementation
 
 	TBitArray<> GatherWeightBlendedWeightmapLayerBitIndices(const UE::Landscape::EditLayers::FMergeContext* InMergeContext) const;
diff --git a/Engine/Source/Runtime/Landscape/Private/LandscapeEditLayers.cpp b/Engine/Source/Runtime/Landscape/Private/LandscapeEditLayers.cpp
--- a/Engine/Source/Runtime/Landscape/Private/LandscapeEditLayers.cpp
+++ b/Engine/Source/Runtime/Landscape/Private/LandscapeEditLayers.cpp
@@ -280,8 +280,8 @@ int32 LandscapeEditLayersLocalMerge = 0;
 static FAutoConsoleVariableRef CVarLandscapeEditLayersLocalMerge(
 	TEXT("landscape.EditLayersLocalMerge.Enable"),
 	LandscapeEditLayersLocalMerge,
-	TEXT("Setting this to 1 will allow the local merge algorithm (that merges layers at the landscape component level) to be used on landscapes that support it. This is a temporary measure while waiting for non-compatible landscapes to be deprecated.\n\
-		Setting this to 2 will allow the batched merge algorithm (that merges layers in batches). Supports all landscape types. This is also a temporary measure but will be the default once it's stable."),
+	TEXT("Setting this to 1 will allow the local merge algorithm (that merges layers at the landscape component level) to be used on landscapes that support it. This is a temporary measure while waiting for non-compatible landscapes to be deprecated.\n"
+		"Setting this to 2 will allow the batched merge algorithm (that merges layers in batches). Supports all landscape types. This is also a temporary measure but will be the default once it's stable."),
 	FConsoleVariableDelegate::CreateStatic(OnLandscapeEditLayersLocalMergeChanged));
 
 TAutoConsoleVariable<int32> CVarLandscapeEditLayersMaxComponentsPerHeightmapResolveBatch(
@@ -1673,7 +1673,7 @@ public:
 		OutEnvironment.SetDefine(TEXT("MERGE_EDIT_LAYER"), 1);
 	}
 
-	static void MergeEditLayerPS(FRDGBuilder& GraphBuilder, FParameters* InParameters, const FName& InTargetLayerName, const FIntPoint& InTextureSize)
+	static void MergeEditLayerPS(FRDGEventName&& InRDGEventName, FRDGBuilder& GraphBuilder, FParameters* InParameters, const FIntPoint& InTextureSize)
 	{
 		FGlobalShaderMap* ShaderMap = GetGlobalShaderMap(GMaxRHIFeatureLevel);
 		TShaderMapRef<FLandscapeEditLayersHeightmapsMergeEditLayerPS> PixelShader(ShaderMap);
@@ -1681,7 +1681,7 @@ public:
 		FPixelShaderUtils::AddFullscreenPass(
 			GraphBuilder,
 			ShaderMap,
-			RDG_EVENT_NAME("LandscapeLayers_MergeEditLayerPS (%s)", *InTargetLayerName.ToString()),
+			MoveTemp(InRDGEventName),
 			PixelShader,
 			InParameters,
 			FIntRect(0, 0, InTextureSize.X, InTextureSize.Y));
@@ -1718,7 +1718,7 @@ public:
 		OutEnvironment.CompilerFlags.Remove(CFLAG_HLSL2021 | CFLAG_PrecompileWithDXC);
 	}
 
-	static void GenerateNormalsPS(FRDGBuilder& GraphBuilder, FParameters* InParameters, const FIntPoint& InTextureSize)
+	static void GenerateNormalsPS(FRDGEventName&& InRDGEventName, FRDGBuilder& GraphBuilder, FParameters* InParameters, const FIntPoint& InTextureSize)
 	{
 		FGlobalShaderMap* ShaderMap = GetGlobalShaderMap(GMaxRHIFeatureLevel);
 		TShaderMapRef<FLandscapeEditLayersHeightmapsGenerateNormalsPS> PixelShader(ShaderMap);
@@ -1726,7 +1726,7 @@ public:
 		FPixelShaderUtils::AddFullscreenPass(
 			GraphBuilder,
 			ShaderMap,
-			RDG_EVENT_NAME("LandscapeLayers_GenerateNormalsPS"),
+			MoveTemp(InRDGEventName),
 			PixelShader,
 			InParameters,
 			FIntRect(0, 0, InTextureSize.X, InTextureSize.Y));
@@ -1803,7 +1803,7 @@ public:
 		OutEnvironment.SetDefine(TEXT("MERGE_EDIT_LAYER"), 1);
 	}
 
-	static void MergeEditLayerPS(FRDGBuilder& GraphBuilder, FParameters* InParameters, const FName& InTargetLayerName, const FIntPoint& InTextureSize)
+	static void MergeEditLayerPS(FRDGEventName&& InRDGEventName, FRDGBuilder& GraphBuilder, FParameters* InParameters, const FIntPoint& InTextureSize)
 	{
 		FGlobalShaderMap* ShaderMap = GetGlobalShaderMap(GMaxRHIFeatureLevel);
 		TShaderMapRef<FLandscapeEditLayersWeightmapsMergeEditLayerPS> PixelShader(ShaderMap);
@@ -1811,7 +1811,7 @@ public:
 		FPixelShaderUtils::AddFullscreenPass(
 			GraphBuilder,
 			ShaderMap,
-			RDG_EVENT_NAME("LandscapeLayers_MergeEditLayerPS (%s)", *InTargetLayerName.ToString()),
+			MoveTemp(InRDGEventName),
 			PixelShader,
 			InParameters,
 			FIntRect(0, 0, InTextureSize.X, InTextureSize.Y));
@@ -1844,7 +1844,7 @@ public:
 		OutEnvironment.SetDefine(TEXT("PERFORM_LEGACY_WEIGHT_BLENDING"), 1);
 	}
 
-	static void PerformLegacyWeightBlendingPS(FRDGBuilder& GraphBuilder, FParameters* InParameters, const FName& InTargetLayerName, const FIntPoint& InTextureSize)
+	static void PerformLegacyWeightBlendingPS(FRDGEventName&& InRDGEventName, FRDGBuilder& GraphBuilder, FParameters* InParameters, const FIntPoint& InTextureSize)
 	{
 		FGlobalShaderMap* ShaderMap = GetGlobalShaderMap(GMaxRHIFeatureLevel);
 		TShaderMapRef<FLandscapeEditLayersWeightmapsPerformLegacyWeightBlendingPS> PixelShader(ShaderMap);
@@ -1852,7 +1852,7 @@ public:
 		FPixelShaderUtils::AddFullscreenPass(
 			GraphBuilder,
 			ShaderMap,
-			RDG_EVENT_NAME("LandscapeLayers_PerformLegacyWeightBlendingPS (%s)", *InTargetLayerName.ToString()),
+			MoveTemp(InRDGEventName),
 			PixelShader,
 			InParameters,
 			FIntRect(0, 0, InTextureSize.X, InTextureSize.Y));
@@ -5536,42 +5536,28 @@ TArray<UE::Landscape::EditLayers::FEditLayerRenderItem> ULandscapeEditLayerPersi
 	return RenderItems;
 }
 
-void ULandscapeEditLayerPersistent::RenderLayer(ILandscapeEditLayerRenderer::FRenderParams& InRenderParams)
+void ULandscapeEditLayerPersistent::RenderLayer(FRenderParams& RenderParams, UE::Landscape::FRDGBuilderRecorder& RDGBuilderRecorder)
 {
 	using namespace UE::Landscape;
 	using namespace UE::Landscape::Private;
 	using namespace UE::Landscape::EditLayers;
 
-	check(!InRenderParams.SortedComponentMergeRenderInfos.IsEmpty());
+	checkf(RDGBuilderRecorder.IsRecording(), TEXT("ERenderFlags::RenderMode_Recorded means the command recorder should be recording at this point"));
+	check(!RenderParams.SortedComponentMergeRenderInfos.IsEmpty());
 
-	ULandscapeInfo* Info = InRenderParams.MergeRenderContext->GetLandscape()->GetLandscapeInfo();
+	ULandscapeInfo* Info = RenderParams.MergeRenderContext->GetLandscape()->GetLandscapeInfo();
 	check(Info != nullptr);
-	ALandscape* Landscape = InRenderParams.MergeRenderContext->GetLandscape();
+	ALandscape* Landscape = RenderParams.MergeRenderContext->GetLandscape();
 	FSceneInterface* SceneInterface = Landscape->GetWorld()->Scene;
-	ULandscapeScratchRenderTarget* WriteRT = nullptr;
 
 	const FLandscapeLayer* LayerStruct = GetOwningLayer();
 	check(LayerStruct != nullptr);
 
-	// Don't do anything with the blend render targets just yet : if this layer has nothing to render at all (e.g. no weightmap on the rendered area), that will save some processing :
-	auto LazyInitializeBlendRenderTargets = [MergeRenderContext = InRenderParams.MergeRenderContext, &WriteRT]
-	{
-		if (WriteRT == nullptr)
-		{
-			MergeRenderContext->CycleBlendRenderTargets();
-			WriteRT = MergeRenderContext->GetBlendRenderTargetWrite();
-
-			// We will write to the RT using a (bunch of) PS : 
-			WriteRT->Clear();
-			WriteRT->TransitionTo(ERHIAccess::RTV);
-		}
-	};
-
-	const int32 NumTargetLayersInRenderGroup = InRenderParams.RenderGroupTargetLayerNames.Num();
+	const int32 NumTargetLayersInRenderGroup = RenderParams.RenderGroupTargetLayerNames.Num();
 
 	// The first step is to copy all of the necessary components' textures to WriteRT and remove the duplicate borders when doing so (subsection by subsection). This is done with a 
 	//  "copy from multiple sources" shader instead of several texture copies, in order to reduce the amount of copy texture commands, which can have a big impact on render-thread performance
-	//  for large landscapes. 
+	//  for large landscapes. Let's first build a list of quads to render and we'll render them all as efficiently as possible thereafter : 
 
 	// For each subsection, we'll add a quad to render, using the following struct :
 	bool bUseSwizzleMask = false;
@@ -5589,25 +5575,25 @@ void ULandscapeEditLayerPersistent::RenderLayer(ILandscapeEditLayerRenderer::FRe
 		uint8 ChannelSwizzleMask = RGBAToRGBASwizzleMask;
 	};
 	TArray<FCopyQuadParams> CopyQuadsParams;
-	CopyQuadsParams.Reserve(InRenderParams.SortedComponentMergeRenderInfos.Num() * NumTargetLayersInRenderGroup* Landscape->NumSubsections * Landscape->NumSubsections);
+	CopyQuadsParams.Reserve(RenderParams.SortedComponentMergeRenderInfos.Num() * NumTargetLayersInRenderGroup* Landscape->NumSubsections * Landscape->NumSubsections);
 
 	// This is kinda hacky, but we since we reuse AddRasterizeToRectsPass and it passes the source texture's sample coordinates from VS to PS via the UVs, dividing the source rect by a single TextureSize, 
 	//  we correct each quad's source rect by a common scale factor so that it samples at the right location in each source texture: 
-	const FVector2D CommonSourceTextureSize(InRenderParams.MergeRenderContext->GetMaxNeededResolution());
+	const FVector2D CommonSourceTextureSize(RenderParams.MergeRenderContext->GetMaxNeededResolution());
 	const int32 SubsectionSizeVerts = (Landscape->SubsectionSizeQuads + 1);
 
 	// Build the full list of quads to copy :
-	for (const FComponentMergeRenderInfo& ComponentMergeRenderInfo : InRenderParams.SortedComponentMergeRenderInfos)
+	for (const FComponentMergeRenderInfo& ComponentMergeRenderInfo : RenderParams.SortedComponentMergeRenderInfos)
 	{
 		for (int32 TargetLayerIndexInRenderGroup = 0; TargetLayerIndexInRenderGroup < NumTargetLayersInRenderGroup; ++TargetLayerIndexInRenderGroup)
 		{
-			const FName TargetLayerName = InRenderParams.RenderGroupTargetLayerNames[TargetLayerIndexInRenderGroup];
+			const FName TargetLayerName = RenderParams.RenderGroupTargetLayerNames[TargetLayerIndexInRenderGroup];
 			UTexture2D* SourceTexture = nullptr;
 			FVector2D SourceTextureBias(ForceInit);
 			// Copy all channels by default :
 			uint8 CopyTextureSwizzleMask = RGBAToRGBASwizzleMask;
 
-			if (InRenderParams.MergeRenderContext->IsHeightmapMerge())
+			if (RenderParams.MergeRenderContext->IsHeightmapMerge())
 			{
 				SourceTexture = ComponentMergeRenderInfo.Component->GetHeightmap(LayerStruct->Guid);
 				SourceTextureBias = FVector2D(ComponentMergeRenderInfo.Component->HeightmapScaleBias.Z, ComponentMergeRenderInfo.Component->HeightmapScaleBias.W);
@@ -5632,10 +5618,6 @@ void ULandscapeEditLayerPersistent::RenderLayer(ILandscapeEditLayerRenderer::FRe
 
 			if (SourceTexture != nullptr)
 			{
-				// We know we'll render something now so it's time we initialized our render target : 
-				LazyInitializeBlendRenderTargets();
-				check(WriteRT != nullptr);
-
 				FTextureResource* SourceTextureResource = SourceTexture->GetResource();
 				checkf(!SourceTexture->IsCompiling(), TEXT("All mips must have been loaded prior to using this function (%s)"), *SourceTexture->GetName());
 				// We get the overall texture size via the resource instead of direct GetSizeX/Y calls because the latter is unreliable while the texture is being built.
@@ -5669,221 +5651,234 @@ void ULandscapeEditLayerPersistent::RenderLayer(ILandscapeEditLayerRenderer::FRe
 		}
 	}
 
-	// Sort the list of quads by output texture first, then by source texture, because we'll coalesce several texture copies into the same pass (one output, multiple inputs)
-	CopyQuadsParams.Sort([](const FCopyQuadParams& InLHS, const FCopyQuadParams& InRHS)
+	if (CopyQuadsParams.IsEmpty())
+	{
+		// No need to do anything if there's nothing to be rendered (e.g. no weightmap on the rendered area), this layer will just be ineffective on this batch
+		//  and we don't need to cycle the blend render targets, that will save some processing:
+		return;
+	}
+
+	RenderParams.MergeRenderContext->CycleBlendRenderTargets(/*InDesiredWriteAccess = */ERHIAccess::None, RDGBuilderRecorder);
+	ULandscapeScratchRenderTarget* WriteRT = RenderParams.MergeRenderContext->GetBlendRenderTargetWrite();
+	WriteRT->Clear(RDGBuilderRecorder);
+	// We will write to the RT using a (bunch of) PS : 
+	WriteRT->TransitionTo(ERHIAccess::RTV, RDGBuilderRecorder);
+
+	// First step : render the edit layer's component quads
 	{
-		if (InLHS.DestinationArrayIndex == InRHS.DestinationArrayIndex)
+		// Sort the list of quads by output texture first, then by source texture, because we'll coalesce several texture copies into the same pass (one output, multiple inputs)
+		CopyQuadsParams.Sort([](const FCopyQuadParams& InLHS, const FCopyQuadParams& InRHS)
 		{
-			if (InLHS.SourceTextureResource == InRHS.SourceTextureResource)
+			if (InLHS.DestinationArrayIndex == InRHS.DestinationArrayIndex)
 			{
-				return (InLHS.SourceRect.Min.Y == InRHS.SourceRect.Min.Y) 
-					? (InLHS.SourceRect.Min.X < InRHS.SourceRect.Min.X)
-					: (InLHS.SourceRect.Min.Y < InRHS.SourceRect.Min.Y);
+				if (InLHS.SourceTextureResource == InRHS.SourceTextureResource)
+				{
+					return (InLHS.SourceRect.Min.Y == InRHS.SourceRect.Min.Y) 
+						? (InLHS.SourceRect.Min.X < InRHS.SourceRect.Min.X)
+						: (InLHS.SourceRect.Min.Y < InRHS.SourceRect.Min.Y);
+				}
+				return (InLHS.SourceTextureResource < InRHS.SourceTextureResource);
 			}
-			return (InLHS.SourceTextureResource < InRHS.SourceTextureResource);
-		}
 
-		return (InLHS.DestinationArrayIndex < InRHS.DestinationArrayIndex);
-	});
+			return (InLHS.DestinationArrayIndex < InRHS.DestinationArrayIndex);
+		});
 
-	// Now process this list of quads and prepare as many passes as necessary for performing all the copies :
-	if (!CopyQuadsParams.IsEmpty())
-	{
-		ENQUEUE_RENDER_COMMAND(LandscapeLayers_Cmd_CopySourcesToTargetLayer)(
-			[ CopyQuadsParams = MoveTemp(CopyQuadsParams)
+		// Now process this list of quads and prepare as many passes as necessary for performing all the copies :
+		auto RDGCommand =
+			[CopyQuadsParams = MoveTemp(CopyQuadsParams)
 			, OutputResource = WriteRT->GetRenderTarget()->GetResource()
 			, OutputResourceName = WriteRT->GetDebugName()
 			, SceneInterface
 			, CommonSourceTextureSize
-			, bUseSwizzleMask] (FRHICommandListImmediate& InRHICmdList)
-		{
-			SCOPED_DRAW_EVENTF(InRHICmdList, LandscapeLayers, TEXT("CopyEditLayer -> %s"), OutputResourceName);
-			FRDGBuilder GraphBuilder(InRHICmdList, RDG_EVENT_NAME("CopySourcesToTargetLayer"));
-
-			FRDGTextureSRVRef BlackDummySRVRef = GraphBuilder.CreateSRV(FRDGTextureSRVDesc::Create(GSystemTextures.GetBlackDummy(GraphBuilder)));
-
-			FCopyQuadsMultiSourcePS::FPermutationDomain PermutationVector;
-			PermutationVector.Set< FCopyQuadsMultiSourcePS::FUseSwizleMask >(bUseSwizzleMask);
-			FGlobalShaderMap* ShaderMap = GetGlobalShaderMap(GMaxRHIFeatureLevel);
-			TShaderRef<FCopyQuadsMultiSourcePS> PixelShader = ShaderMap->GetShader<FCopyQuadsMultiSourcePS>(PermutationVector);
+			, bUseSwizzleMask] (FRDGBuilder& GraphBuilder)
+			{
+				RDG_EVENT_SCOPE(GraphBuilder, "CopyEditLayer -> %s", *OutputResourceName);
 
-			// TODO [jonathan.bard] this is just to avoid a RHI validation error for unoptimized shaders... once validation is made to not issue those errors, we can remove this
-			// Create a SceneView to please the shader bindings, but it's unused in practice 
-			FSceneViewFamilyContext ViewFamily(FSceneViewFamily::ConstructionValues(nullptr, SceneInterface, FEngineShowFlags(ESFIM_Game)).SetTime(FGameTime::GetTimeSinceAppStart()));
-			FSceneViewInitOptions ViewInitOptions;
-			ViewInitOptions.ViewFamily = &ViewFamily;
-			ViewInitOptions.SetViewRectangle(FIntRect(0, 0, 1, 1)); // Use a dummy rect to avoid a check(slow)
-			GetRendererModule().CreateAndInitSingleView(InRHICmdList, &ViewFamily, &ViewInitOptions);
-			const FSceneView* View = ViewFamily.Views[0];
+				FRDGTextureSRVRef BlackDummySRVRef = GraphBuilder.CreateSRV(FRDGTextureSRVDesc::Create(GSystemTextures.GetBlackDummy(GraphBuilder)));
 
-			// The following variables allow to accumulate quads for rendering in a single pass. Then, we flush them all out when the pass is full and we start again : 
-			int32 CurrentDestinationArrayIndex = INDEX_NONE;
-			TArray<FTextureResource*> CurrentSourceTextureResources;
-			TArray<FUintVector2> CurrentQuadInfos;
-			TArray<FUintVector4> CurrentSourceRects, CurrentDestinationRects;
+				FCopyQuadsMultiSourcePS::FPermutationDomain PermutationVector;
+				PermutationVector.Set< FCopyQuadsMultiSourcePS::FUseSwizleMask >(bUseSwizzleMask);
+				FGlobalShaderMap* ShaderMap = GetGlobalShaderMap(GMaxRHIFeatureLevel);
+				TShaderRef<FCopyQuadsMultiSourcePS> PixelShader = ShaderMap->GetShader<FCopyQuadsMultiSourcePS>(PermutationVector);
+
+				// TODO [jonathan.bard] this is just to avoid a RHI validation error for unoptimized shaders... once validation is made to not issue those errors, we can remove this
+				// Create a SceneView to please the shader bindings, but it's unused in practice 
+				FSceneViewFamilyContext ViewFamily(FSceneViewFamily::ConstructionValues(nullptr, SceneInterface, FEngineShowFlags(ESFIM_Game)).SetTime(FGameTime::GetTimeSinceAppStart()));
+				FSceneViewInitOptions ViewInitOptions;
+				ViewInitOptions.ViewFamily = &ViewFamily;
+				ViewInitOptions.SetViewRectangle(FIntRect(0, 0, 1, 1)); // Use a dummy rect to avoid a check(slow)
+				GetRendererModule().CreateAndInitSingleView(GraphBuilder.RHICmdList, &ViewFamily, &ViewInitOptions);
+				const FSceneView* View = ViewFamily.Views[0];
+
+				// The following variables allow to accumulate quads for rendering in a single pass. Then, we flush them all out when the pass is full and we start again : 
+				int32 CurrentDestinationArrayIndex = INDEX_NONE;
+				TArray<FTextureResource*> CurrentSourceTextureResources;
+				TArray<FUintVector2> CurrentQuadInfos;
+				TArray<FUintVector4> CurrentSourceRects, CurrentDestinationRects;
+
+				FRDGTextureRef OutputTexture = GraphBuilder.RegisterExternalTexture(CreateRenderTarget(OutputResource->GetTextureRHI(), TEXT("OutputTexture")));
+
+				// When we've reached the limit of what we can render in one pass, flush all and add a render pass : 
+				auto FlushCopyQuads = [&GraphBuilder, View, OutputTexture, &CurrentDestinationArrayIndex, &CurrentSourceTextureResources, &CurrentSourceRects, &CurrentDestinationRects,
+					&CurrentQuadInfos, &CommonSourceTextureSize, BlackDummySRVRef, ShaderMap, PixelShader]()
+					{
+						if (CurrentSourceTextureResources.IsEmpty())
+						{
+							// Nothing to flush :
+							check(CurrentSourceRects.IsEmpty() && CurrentDestinationRects.IsEmpty() && CurrentQuadInfos.IsEmpty());
+							return;
+						}
 
-			FRDGTextureRef OutputTexture = GraphBuilder.RegisterExternalTexture(CreateRenderTarget(OutputResource->GetTextureRHI(), TEXT("OutputTexture")));
+						check(!CurrentSourceRects.IsEmpty() && (CurrentSourceRects.Num() == CurrentDestinationRects.Num()) && (CurrentSourceRects.Num() == CurrentQuadInfos.Num()));
 
-			// When we've reached the limit of what we can render in one pass, flush all and add a render pass : 
-			auto FlushCopyQuads = [&GraphBuilder, View, OutputTexture, &CurrentDestinationArrayIndex, &CurrentSourceTextureResources, &CurrentSourceRects, &CurrentDestinationRects,
-				&CurrentQuadInfos, &CommonSourceTextureSize, BlackDummySRVRef, ShaderMap, PixelShader]()
-			{
-				if (CurrentSourceTextureResources.IsEmpty())
-				{
-					// Nothing to flush :
-					check(CurrentSourceRects.IsEmpty() && CurrentDestinationRects.IsEmpty() && CurrentQuadInfos.IsEmpty());
-					return;
-				}
+						FRDGBufferRef RectBuffer = CreateUploadBuffer(GraphBuilder, TEXT("DestinationRects"), TConstArrayView<FUintVector4>(MakeArrayView(CurrentDestinationRects)));
+						FRDGBufferSRVRef RectBufferSRV = GraphBuilder.CreateSRV(FRDGBufferSRVDesc(RectBuffer, PF_R32G32B32A32_UINT));
 
-				check(!CurrentSourceRects.IsEmpty() && (CurrentSourceRects.Num() == CurrentDestinationRects.Num()) && (CurrentSourceRects.Num() == CurrentQuadInfos.Num()));
+						FRDGBufferRef RectUVBuffer = CreateUploadBuffer(GraphBuilder, TEXT("RectUVs"), TConstArrayView<FUintVector4>(MakeArrayView(CurrentSourceRects)));
+						FRDGBufferSRVRef RectUVBufferSRV = GraphBuilder.CreateSRV(FRDGBufferSRVDesc(RectUVBuffer, PF_R32G32B32A32_UINT));
 
-				FRDGBufferRef RectBuffer = CreateUploadBuffer(GraphBuilder, TEXT("DestinationRects"), TConstArrayView<FUintVector4>(MakeArrayView(CurrentDestinationRects)));
-				FRDGBufferSRVRef RectBufferSRV = GraphBuilder.CreateSRV(FRDGBufferSRVDesc(RectBuffer, PF_R32G32B32A32_UINT));
+						FRDGBufferRef QuadInfosBuffer = CreateUploadBuffer(GraphBuilder, TEXT("QuadInfos"), TConstArrayView<FUintVector2>(MakeArrayView(CurrentQuadInfos)));
+						FRDGBufferSRVRef QuadInfosBufferSRV = GraphBuilder.CreateSRV(FRDGBufferSRVDesc(QuadInfosBuffer, PF_R32G32_UINT));
 
-				FRDGBufferRef RectUVBuffer = CreateUploadBuffer(GraphBuilder, TEXT("RectUVs"), TConstArrayView<FUintVector4>(MakeArrayView(CurrentSourceRects)));
-				FRDGBufferSRVRef RectUVBufferSRV = GraphBuilder.CreateSRV(FRDGBufferSRVDesc(RectUVBuffer, PF_R32G32B32A32_UINT));
+						// We may target an array slice directly : 
+						int32 ArrayIndex = -1;
+						if (OutputTexture->Desc.IsTextureArray())
+						{
+							check(CurrentDestinationArrayIndex < OutputTexture->Desc.ArraySize);
+							ArrayIndex = CurrentDestinationArrayIndex;
+						}
 
-				FRDGBufferRef QuadInfosBuffer = CreateUploadBuffer(GraphBuilder, TEXT("QuadInfos"), TConstArrayView<FUintVector2>(MakeArrayView(CurrentQuadInfos)));
-				FRDGBufferSRVRef QuadInfosBufferSRV = GraphBuilder.CreateSRV(FRDGBufferSRVDesc(QuadInfosBuffer, PF_R32G32_UINT));
-				
-				// We may target an array slice directly : 
-				int32 ArrayIndex = -1;
-				if (OutputTexture->Desc.IsTextureArray())
-				{
-					check(CurrentDestinationArrayIndex < OutputTexture->Desc.ArraySize);
-					ArrayIndex = CurrentDestinationArrayIndex;
-				}
+						FCopyQuadsMultiSourcePSParameters* PassParameters = GraphBuilder.AllocParameters<FCopyQuadsMultiSourcePSParameters>();
+						PassParameters->RenderTargets[0] = FRenderTargetBinding(OutputTexture, ERenderTargetLoadAction::ELoad, /*InMipIndex = */0, ArrayIndex);
+						PassParameters->PS.View = View->ViewUniformBuffer;
+						PassParameters->PS.InQuadInfos = QuadInfosBufferSRV;
+						check(CurrentSourceTextureResources.Num() <= FCopyQuadsMultiSourcePS::NumMultiSources);
+						for (int32 TextureIndex = 0; TextureIndex < FCopyQuadsMultiSourcePS::NumMultiSources; ++TextureIndex)
+						{
+							FRDGTextureRef SourceTexture = BlackDummySRVRef->GetParent();
+							if (CurrentSourceTextureResources.IsValidIndex(TextureIndex))
+							{
+								SourceTexture = GraphBuilder.RegisterExternalTexture(CreateRenderTarget(CurrentSourceTextureResources[TextureIndex]->GetTexture2DRHI(), TEXT("SourceTexture")));
+							}
+							PassParameters->PS.InSourceTexture[TextureIndex] = SourceTexture;
+						}
 
-				FCopyQuadsMultiSourcePSParameters* PassParameters = GraphBuilder.AllocParameters<FCopyQuadsMultiSourcePSParameters>();
-				PassParameters->RenderTargets[0] = FRenderTargetBinding(OutputTexture, ERenderTargetLoadAction::ELoad, /*InMipIndex = */0, ArrayIndex);
-				PassParameters->PS.View = View->ViewUniformBuffer;
-				PassParameters->PS.InQuadInfos = QuadInfosBufferSRV;
-				check(CurrentSourceTextureResources.Num() <= FCopyQuadsMultiSourcePS::NumMultiSources);
-				for (int32 TextureIndex = 0; TextureIndex < FCopyQuadsMultiSourcePS::NumMultiSources; ++TextureIndex)
-				{
-					FRDGTextureRef SourceTexture = BlackDummySRVRef->GetParent();
-					if (CurrentSourceTextureResources.IsValidIndex(TextureIndex))
+						FPixelShaderUtils::AddRasterizeToRectsPass<FCopyQuadsMultiSourcePS>(GraphBuilder,
+							ShaderMap,
+							RDG_EVENT_NAME("CopyQuadsMultiSourcePS"),
+							PixelShader,
+							PassParameters,
+							/*ViewportSize = */OutputTexture->Desc.Extent,
+							RectBufferSRV,
+							CurrentDestinationRects.Num(),
+							/*BlendState = */nullptr,
+							/*RasterizerState = */nullptr,
+							/*DepthStencilState = */nullptr,
+							/*StencilRef = */0,
+							/*TextureSize = */FIntPoint(FMath::FloorToInt32(CommonSourceTextureSize.X), FMath::FloorToInt32(CommonSourceTextureSize.Y)),
+							RectUVBufferSRV);
+
+						// Reset all for the next pass that comes along : 
+						CurrentDestinationArrayIndex = INDEX_NONE;
+						CurrentSourceTextureResources.Empty();
+						CurrentSourceRects.Empty();
+						CurrentDestinationRects.Empty();
+						CurrentQuadInfos.Empty();
+					};
+
+				for (const FCopyQuadParams& InSingleQuadParams : CopyQuadsParams)
+				{
+					check(InSingleQuadParams.DestinationArrayIndex != INDEX_NONE);
+					// If the output texture/slice has changed since the last iteration, it's time to flush and start a new pass :
+					if (CurrentDestinationArrayIndex != InSingleQuadParams.DestinationArrayIndex)
 					{
-						SourceTexture = GraphBuilder.RegisterExternalTexture(CreateRenderTarget(CurrentSourceTextureResources[TextureIndex]->GetTexture2DRHI(), TEXT("SourceTexture")));
+						if (CurrentDestinationArrayIndex != INDEX_NONE)
+						{
+							FlushCopyQuads();
+						}
+						CurrentDestinationArrayIndex = InSingleQuadParams.DestinationArrayIndex;
 					}
-					PassParameters->PS.InSourceTexture[TextureIndex] = SourceTexture;
-				}
-
-				FPixelShaderUtils::AddRasterizeToRectsPass<FCopyQuadsMultiSourcePS>(GraphBuilder,
-					ShaderMap,
-					RDG_EVENT_NAME("CopyQuadsMultiSourcePS"),
-					PixelShader,
-					PassParameters,
-					/*ViewportSize = */OutputTexture->Desc.Extent,
-					RectBufferSRV,
-					CurrentDestinationRects.Num(),
-					/*BlendState = */nullptr,
-					/*RasterizerState = */nullptr,
-					/*DepthStencilState = */nullptr,
-					/*StencilRef = */0,
-					/*TextureSize = */FIntPoint(FMath::FloorToInt32(CommonSourceTextureSize.X), FMath::FloorToInt32(CommonSourceTextureSize.Y)),
-					RectUVBufferSRV);
-
-				// Reset all for the next pass that comes along : 
-				CurrentDestinationArrayIndex = INDEX_NONE;
-				CurrentSourceTextureResources.Empty();
-				CurrentSourceRects.Empty();
-				CurrentDestinationRects.Empty();
-				CurrentQuadInfos.Empty();
-			};
 
-			for (const FCopyQuadParams& InSingleQuadParams : CopyQuadsParams)
-			{
-				check(InSingleQuadParams.DestinationArrayIndex != INDEX_NONE);
-				// If the output texture/slice has changed since the last iteration, it's time to flush and start a new pass :
-				if (CurrentDestinationArrayIndex != InSingleQuadParams.DestinationArrayIndex)
-				{
-					if (CurrentDestinationArrayIndex != INDEX_NONE)
+					if (CurrentSourceTextureResources.IsEmpty())
 					{
-						FlushCopyQuads();
+						CurrentSourceTextureResources.Add(InSingleQuadParams.SourceTextureResource);
 					}
-					CurrentDestinationArrayIndex = InSingleQuadParams.DestinationArrayIndex;
-				}
-
-				if (CurrentSourceTextureResources.IsEmpty())
-				{
-					CurrentSourceTextureResources.Add(InSingleQuadParams.SourceTextureResource);
-				}
-				else if (InSingleQuadParams.SourceTextureResource != CurrentSourceTextureResources.Last())
-				{
-					// If we've reached the amount of textures we can render in a single pass, we flush the pass and initiate a new one : 
-					if (CurrentSourceTextureResources.Num() == FCopyQuadsMultiSourcePS::NumMultiSources)
+					else if (InSingleQuadParams.SourceTextureResource != CurrentSourceTextureResources.Last())
 					{
-						FlushCopyQuads();
-						check(CurrentSourceTextureResources.IsEmpty());
+						// If we've reached the amount of textures we can render in a single pass, we flush the pass and initiate a new one : 
+						if (CurrentSourceTextureResources.Num() == FCopyQuadsMultiSourcePS::NumMultiSources)
+						{
+							FlushCopyQuads();
+							check(CurrentSourceTextureResources.IsEmpty());
 
-						CurrentDestinationArrayIndex = InSingleQuadParams.DestinationArrayIndex;
+							CurrentDestinationArrayIndex = InSingleQuadParams.DestinationArrayIndex;
+						}
+						CurrentSourceTextureResources.Add(InSingleQuadParams.SourceTextureResource);
 					}
-					CurrentSourceTextureResources.Add(InSingleQuadParams.SourceTextureResource);
-				}
 
-				// If we are using the same texture as the last one, we can render it in the same pass, just append our quad : 
-				if (InSingleQuadParams.SourceTextureResource == CurrentSourceTextureResources.Last())
-				{
-					const int32 SourceTextureIndex = CurrentSourceTextureResources.Num() - 1;
-					CurrentQuadInfos.Add(FUintVector2(static_cast<uint32>(SourceTextureIndex), InSingleQuadParams.ChannelSwizzleMask));
-					CurrentSourceRects.Add(FUintVector4(InSingleQuadParams.SourceRect.Min.X, InSingleQuadParams.SourceRect.Min.Y, InSingleQuadParams.SourceRect.Max.X, InSingleQuadParams.SourceRect.Max.Y));
-					CurrentDestinationRects.Add(FUintVector4(InSingleQuadParams.DestinationRect.Min.X, InSingleQuadParams.DestinationRect.Min.Y, InSingleQuadParams.DestinationRect.Max.X, InSingleQuadParams.DestinationRect.Max.Y));
+					// If we are using the same texture as the last one, we can render it in the same pass, just append our quad : 
+					if (InSingleQuadParams.SourceTextureResource == CurrentSourceTextureResources.Last())
+					{
+						const int32 SourceTextureIndex = CurrentSourceTextureResources.Num() - 1;
+						CurrentQuadInfos.Add(FUintVector2(static_cast<uint32>(SourceTextureIndex), InSingleQuadParams.ChannelSwizzleMask));
+						CurrentSourceRects.Add(FUintVector4(InSingleQuadParams.SourceRect.Min.X, InSingleQuadParams.SourceRect.Min.Y, InSingleQuadParams.SourceRect.Max.X, InSingleQuadParams.SourceRect.Max.Y));
+						CurrentDestinationRects.Add(FUintVector4(InSingleQuadParams.DestinationRect.Min.X, InSingleQuadParams.DestinationRect.Min.Y, InSingleQuadParams.DestinationRect.Max.X, InSingleQuadParams.DestinationRect.Max.Y));
+					}
 				}
-			}
 
-			// Flush the remaining quads if any : 
-			FlushCopyQuads();
-
-			// We need to specify the final state of the external texture to prevent the graph builder from transitioning it to SRVMask :
-			GraphBuilder.SetTextureAccessFinal(OutputTexture, ERHIAccess::RTV);
+				// Flush the remaining quads if any : 
+				FlushCopyQuads();
+			};
 
-			GraphBuilder.Execute();
-		});
+		// We need to specify the final state of the external texture to prevent the graph builder from transitioning it to SRVMask :
+		RDGBuilderRecorder.EnqueueRDGCommand(RDGCommand, { { WriteRT->GetRenderTarget()->GetResource(), ERHIAccess::RTV } });
 	}
 
-	// No need to do anything if there's nothing to be rendered, this layer will just be ineffective on this batch :
-	if (WriteRT != nullptr)
+	// The second step is to perform the blend : merge the edit layer with the previous layers in the stack
 	{
-		// The second step is to perform the blend : merge the edit layer with the previous layers in the stack
-		InRenderParams.MergeRenderContext->CycleBlendRenderTargets(/*InDesiredWriteAccess = */ERHIAccess::RTV);
-		WriteRT = InRenderParams.MergeRenderContext->GetBlendRenderTargetWrite();
-		ULandscapeScratchRenderTarget* CurrentLayerReadRT = InRenderParams.MergeRenderContext->GetBlendRenderTargetRead();
-		ULandscapeScratchRenderTarget* PreviousLayersReadRT = InRenderParams.MergeRenderContext->GetBlendRenderTargetReadPrevious();
+		RenderParams.MergeRenderContext->CycleBlendRenderTargets(/*InDesiredWriteAccess = */ERHIAccess::RTV, RDGBuilderRecorder);
+		WriteRT = RenderParams.MergeRenderContext->GetBlendRenderTargetWrite();
+		ULandscapeScratchRenderTarget* CurrentLayerReadRT = RenderParams.MergeRenderContext->GetBlendRenderTargetRead();
+		ULandscapeScratchRenderTarget* PreviousLayersReadRT = RenderParams.MergeRenderContext->GetBlendRenderTargetReadPrevious();
 	 
-		if (InRenderParams.MergeRenderContext->IsHeightmapMerge())
+		if (RenderParams.MergeRenderContext->IsHeightmapMerge())
 		{
-			ENQUEUE_RENDER_COMMAND(LandscapeLayers_Cmd_HeightmapsMergeEditLayer)(
-				[ TargetLayerName = InRenderParams.RenderGroupTargetLayerNames[0]
+			auto RDGCommand =
+				[TargetLayerName = RenderParams.RenderGroupTargetLayerNames[0]
 				, OutputResource = WriteRT->GetRenderTarget2D()->GetResource()
 				, OutputResourceName = WriteRT->GetDebugName()
 				, CurrentEditLayerResource = CurrentLayerReadRT->GetRenderTarget2D()->GetResource()
 				, PreviousEditLayersResource = PreviousLayersReadRT->GetRenderTarget2D()->GetResource()
-				, EffectiveTextureSize = InRenderParams.MergeRenderContext->GetCurrentRenderBatch()->GetRenderTargetResolution(/*bInWithDuplicateBorders = */false)
-				, BlendMode = LayerStruct->BlendMode, HeightmapAlpha = LayerStruct->HeightmapAlpha ]
-				(FRHICommandListImmediate& InRHICmdList)
-			{
-				SCOPED_DRAW_EVENTF(InRHICmdList, LandscapeLayers, TEXT("MergeEditLayer -> %s"), OutputResourceName);
-				FRDGBuilder GraphBuilder(InRHICmdList, RDG_EVENT_NAME("HeightmapsMergeEditLayer"));
-
-				FRDGTextureRef OutputTextureRef = GraphBuilder.RegisterExternalTexture(CreateRenderTarget(OutputResource->TextureRHI, TEXT("OutputTexture")));
-				FRDGTextureRef CurrentEditLayerTextureRef = GraphBuilder.RegisterExternalTexture(CreateRenderTarget(CurrentEditLayerResource->TextureRHI, TEXT("CurrentEditLayerTexture")));
-				FRDGTextureRef PreviousEditLayersTextureRef = GraphBuilder.RegisterExternalTexture(CreateRenderTarget(PreviousEditLayersResource->TextureRHI, TEXT("PreviousEditLayersTexture")));
-
-				FLandscapeEditLayersHeightmapsMergeEditLayerPS::FParameters* PSParams = GraphBuilder.AllocParameters<FLandscapeEditLayersHeightmapsMergeEditLayerPS::FParameters>();
-				PSParams->RenderTargets[0] = FRenderTargetBinding(OutputTextureRef, ERenderTargetLoadAction::ENoAction);
-				PSParams->InEditLayerBlendMode = static_cast<uint32>(LandscapeBlendModeToEditLayerBlendMode(BlendMode));
-				PSParams->InEditLayerAlpha = HeightmapAlpha;
-				PSParams->InCurrentEditLayerHeightmap = GraphBuilder.CreateSRV(FRDGTextureSRVDesc::Create(CurrentEditLayerTextureRef));
-				PSParams->InPreviousEditLayersHeightmap = GraphBuilder.CreateSRV(FRDGTextureSRVDesc::Create(PreviousEditLayersTextureRef));
-
-				FLandscapeEditLayersHeightmapsMergeEditLayerPS::MergeEditLayerPS(GraphBuilder, PSParams, TargetLayerName, EffectiveTextureSize);
+				, EffectiveTextureSize = RenderParams.MergeRenderContext->GetCurrentRenderBatch()->GetRenderTargetResolution(/*bInWithDuplicateBorders = */false)
+				, BlendMode = LayerStruct->BlendMode, HeightmapAlpha = LayerStruct->HeightmapAlpha]
+				(FRDGBuilder& GraphBuilder)
+				{
+					FRDGTextureRef OutputTextureRef = GraphBuilder.RegisterExternalTexture(CreateRenderTarget(OutputResource->TextureRHI, TEXT("OutputTexture")));
+					FRDGTextureRef CurrentEditLayerTextureRef = GraphBuilder.RegisterExternalTexture(CreateRenderTarget(CurrentEditLayerResource->TextureRHI, TEXT("CurrentEditLayerTexture")));
+					FRDGTextureRef PreviousEditLayersTextureRef = GraphBuilder.RegisterExternalTexture(CreateRenderTarget(PreviousEditLayersResource->TextureRHI, TEXT("PreviousEditLayersTexture")));
+
+					FLandscapeEditLayersHeightmapsMergeEditLayerPS::FParameters* PSParams = GraphBuilder.AllocParameters<FLandscapeEditLayersHeightmapsMergeEditLayerPS::FParameters>();
+					PSParams->RenderTargets[0] = FRenderTargetBinding(OutputTextureRef, ERenderTargetLoadAction::ENoAction);
+					PSParams->InEditLayerBlendMode = static_cast<uint32>(LandscapeBlendModeToEditLayerBlendMode(BlendMode));
+					PSParams->InEditLayerAlpha = HeightmapAlpha;
+					PSParams->InCurrentEditLayerHeightmap = GraphBuilder.CreateSRV(FRDGTextureSRVDesc::Create(CurrentEditLayerTextureRef));
+					PSParams->InPreviousEditLayersHeightmap = GraphBuilder.CreateSRV(FRDGTextureSRVDesc::Create(PreviousEditLayersTextureRef));
+
+					FLandscapeEditLayersHeightmapsMergeEditLayerPS::MergeEditLayerPS(
+						RDG_EVENT_NAME("MergeEditLayer(%s) -> %s", *TargetLayerName.ToString(), *OutputResourceName), GraphBuilder, PSParams, EffectiveTextureSize);
+				};
 
-				// We need to specify the final state of the external texture to prevent the graph builder from transitioning it to SRVMask :
-				GraphBuilder.SetTextureAccessFinal(OutputTextureRef, ERHIAccess::RTV);
 
-				GraphBuilder.Execute();
-			});
+			// We need to specify the final state of the external textures to prevent the graph builder from transitioning them to SRVMask (even those that end up as SRVMask at the end of this command,
+			//  because they will likely be part of another RDGCommand down the line so we need to maintain an accurate picture of every external texture ever involved in the recorded command so that 
+			//  we can set a proper access when the recorder is flushed (and the FRDGBuilder, executed) :
+			TArray<FRDGBuilderRecorder::FRDGExternalTextureAccessFinal> RDGExternalTextureAccessFinalList =
+			{
+				{ WriteRT->GetRenderTarget()->GetResource(), ERHIAccess::RTV },
+				{ CurrentLayerReadRT->GetRenderTarget()->GetResource(), ERHIAccess::SRVMask },
+				{ PreviousLayersReadRT->GetRenderTarget()->GetResource(), ERHIAccess::SRVMask },
+			};
+			RDGBuilderRecorder.EnqueueRDGCommand(RDGCommand, RDGExternalTextureAccessFinalList);
 		}
 		else
 		{
@@ -5901,7 +5896,7 @@ void ULandscapeEditLayerPersistent::RenderLayer(ILandscapeEditLayerRenderer::FRe
 			// Prepare a list of all paint layers to merge and do it all in 1 command :
 			TArray<FMergePaintLayerParams> MergePaintLayersParams;
 			MergePaintLayersParams.Reserve(NumTargetLayersInRenderGroup);
-			for (const FName& TargetLayerName : InRenderParams.RenderGroupTargetLayerNames)
+			for (const FName& TargetLayerName : RenderParams.RenderGroupTargetLayerNames)
 			{
 				FMergePaintLayerParams& MergePaintLayerParams = MergePaintLayersParams.Emplace_GetRef(TargetLayerName);
 
@@ -5921,43 +5916,47 @@ void ULandscapeEditLayerPersistent::RenderLayer(ILandscapeEditLayerRenderer::FRe
 				}
 			}
 
-			ENQUEUE_RENDER_COMMAND(LandscapeLayers_Cmd_WeightmapsMergeEditLayer)(
-				[ MergePaintLayersParams
+			auto RDGCommand =
+				[MergePaintLayersParams
 				, OutputResource = WriteRT->GetRenderTarget2DArray()->GetResource()
 				, OutputResourceName = WriteRT->GetDebugName()
 				, CurrentEditLayerResource = CurrentLayerReadRT->GetRenderTarget2DArray()->GetResource()
 				, PreviousEditLayersResource = PreviousLayersReadRT->GetRenderTarget2DArray()->GetResource()
-				, EffectiveTextureSize = InRenderParams.MergeRenderContext->GetCurrentRenderBatch()->GetRenderTargetResolution(/*bInWithDuplicateBorders = */false) ]
-				(FRHICommandListImmediate& InRHICmdList)
-			{
-				SCOPED_DRAW_EVENTF(InRHICmdList, LandscapeLayers, TEXT("MergeEditLayer -> %s"), OutputResourceName);
-				FRDGBuilder GraphBuilder(InRHICmdList, RDG_EVENT_NAME("WeightmapsMergeEditLayer"));
-
-				FRDGTextureRef OutputTextureRef = GraphBuilder.RegisterExternalTexture(CreateRenderTarget(OutputResource->TextureRHI, TEXT("OutputTexture")));
-				FRDGTextureRef CurrentEditLayerTextureRef = GraphBuilder.RegisterExternalTexture(CreateRenderTarget(CurrentEditLayerResource->TextureRHI, TEXT("CurrentEditLayerTexture")));
-				FRDGTextureSRVRef CurrentEditLayerTextureSRVRef = GraphBuilder.CreateSRV(FRDGTextureSRVDesc::Create(CurrentEditLayerTextureRef));
-				FRDGTextureRef PreviousEditLayersTextureRef = GraphBuilder.RegisterExternalTexture(CreateRenderTarget(PreviousEditLayersResource->TextureRHI, TEXT("PreviousEditLayersTexture")));
-				FRDGTextureSRVRef PreviousEditLayersTextureSRVRef = GraphBuilder.CreateSRV(FRDGTextureSRVDesc::Create(PreviousEditLayersTextureRef));
-
-				const int32 NumTargetLayers = MergePaintLayersParams.Num();
-				for (int32 TargetLayerIndex = 0; TargetLayerIndex < NumTargetLayers; ++TargetLayerIndex)
+				, EffectiveTextureSize = RenderParams.MergeRenderContext->GetCurrentRenderBatch()->GetRenderTargetResolution(/*bInWithDuplicateBorders = */false)]
+				(FRDGBuilder& GraphBuilder)
 				{
-					FLandscapeEditLayersWeightmapsMergeEditLayerPS::FParameters* PSParams = GraphBuilder.AllocParameters<FLandscapeEditLayersWeightmapsMergeEditLayerPS::FParameters>();
-					PSParams->RenderTargets[0] = FRenderTargetBinding(OutputTextureRef, ERenderTargetLoadAction::ENoAction, /*InMipIndex = */0, /*InArraySlice = */TargetLayerIndex);
-					PSParams->InPaintLayerIndex = TargetLayerIndex;
-					PSParams->InEditLayerPaintLayerBlendMode = static_cast<uint32>(MergePaintLayersParams[TargetLayerIndex].EditLayerPaintLayerBlendMode);
-					PSParams->InEditLayerAlpha = MergePaintLayersParams[TargetLayerIndex].EditLayerAlpha;
-					PSParams->InCurrentEditLayerWeightmaps = CurrentEditLayerTextureSRVRef;
-					PSParams->InPreviousEditLayersWeightmaps = PreviousEditLayersTextureSRVRef;
-
-					FLandscapeEditLayersWeightmapsMergeEditLayerPS::MergeEditLayerPS(GraphBuilder, PSParams, MergePaintLayersParams[TargetLayerIndex].PaintLayerName, EffectiveTextureSize);
-				}
-
-				// We need to specify the final state of the external texture to prevent the graph builder from transitioning it to SRVMask :
-				GraphBuilder.SetTextureAccessFinal(OutputTextureRef, ERHIAccess::RTV);
+					FRDGTextureRef OutputTextureRef = GraphBuilder.RegisterExternalTexture(CreateRenderTarget(OutputResource->TextureRHI, TEXT("OutputTexture")));
+					FRDGTextureRef CurrentEditLayerTextureRef = GraphBuilder.RegisterExternalTexture(CreateRenderTarget(CurrentEditLayerResource->TextureRHI, TEXT("CurrentEditLayerTexture")));
+					FRDGTextureSRVRef CurrentEditLayerTextureSRVRef = GraphBuilder.CreateSRV(FRDGTextureSRVDesc::Create(CurrentEditLayerTextureRef));
+					FRDGTextureRef PreviousEditLayersTextureRef = GraphBuilder.RegisterExternalTexture(CreateRenderTarget(PreviousEditLayersResource->TextureRHI, TEXT("PreviousEditLayersTexture")));
+					FRDGTextureSRVRef PreviousEditLayersTextureSRVRef = GraphBuilder.CreateSRV(FRDGTextureSRVDesc::Create(PreviousEditLayersTextureRef));
 
-				GraphBuilder.Execute();
-			});
+					const int32 NumTargetLayers = MergePaintLayersParams.Num();
+					for (int32 TargetLayerIndex = 0; TargetLayerIndex < NumTargetLayers; ++TargetLayerIndex)
+					{
+						FLandscapeEditLayersWeightmapsMergeEditLayerPS::FParameters* PSParams = GraphBuilder.AllocParameters<FLandscapeEditLayersWeightmapsMergeEditLayerPS::FParameters>();
+						PSParams->RenderTargets[0] = FRenderTargetBinding(OutputTextureRef, ERenderTargetLoadAction::ENoAction, /*InMipIndex = */0, /*InArraySlice = */TargetLayerIndex);
+						PSParams->InPaintLayerIndex = TargetLayerIndex;
+						PSParams->InEditLayerPaintLayerBlendMode = static_cast<uint32>(MergePaintLayersParams[TargetLayerIndex].EditLayerPaintLayerBlendMode);
+						PSParams->InEditLayerAlpha = MergePaintLayersParams[TargetLayerIndex].EditLayerAlpha;
+						PSParams->InCurrentEditLayerWeightmaps = CurrentEditLayerTextureSRVRef;
+						PSParams->InPreviousEditLayersWeightmaps = PreviousEditLayersTextureSRVRef;
+
+						FLandscapeEditLayersWeightmapsMergeEditLayerPS::MergeEditLayerPS(
+							RDG_EVENT_NAME("MergeEditLayer(%s) -> %s", *MergePaintLayersParams[TargetLayerIndex].PaintLayerName.ToString(), *OutputResourceName), GraphBuilder, PSParams, EffectiveTextureSize);
+					}
+				};
+				
+			// We need to specify the final state of the external textures to prevent the graph builder from transitioning them to SRVMask (even those that end up as SRVMask at the end of this command,
+			//  because they will likely be part of another RDGCommand down the line so we need to maintain an accurate picture of every external texture ever involved in the recorded command so that 
+			//  we can set a proper access when the recorder is flushed (and the FRDGBuilder, executed) :
+			TArray<FRDGBuilderRecorder::FRDGExternalTextureAccessFinal> RDGExternalTextureAccessFinalList =
+			{
+				{ WriteRT->GetRenderTarget()->GetResource(), ERHIAccess::RTV },
+				{ CurrentLayerReadRT->GetRenderTarget()->GetResource(), ERHIAccess::SRVMask },
+				{ PreviousLayersReadRT->GetRenderTarget()->GetResource(), ERHIAccess::SRVMask },
+			};
+			RDGBuilderRecorder.EnqueueRDGCommand(RDGCommand, RDGExternalTextureAccessFinalList);
 		}
 	}
 }
@@ -6422,78 +6421,85 @@ namespace UE::Landscape::EditLayers::Private
 #endif // ENABLE_VISUAL_LOG
 } // UE::Landscape::EditLayers::Private
 
-void ULandscapeHeightmapNormalsEditLayerRenderer::RenderLayer(FRenderParams& InRenderParams)
+void ULandscapeHeightmapNormalsEditLayerRenderer::RenderLayer(FRenderParams& RenderParams, UE::Landscape::FRDGBuilderRecorder& RDGBuilderRecorder)
 {
+	using namespace UE::Landscape;
 	using namespace UE::Landscape::EditLayers;
 	using namespace UE::Landscape::EditLayers::Private;
 
-	FMergeRenderContext* RenderContext = InRenderParams.MergeRenderContext;
+	FMergeRenderContext* RenderContext = RenderParams.MergeRenderContext;
 	const FMergeRenderBatch* RenderBatch = RenderContext->GetCurrentRenderBatch();
 	ALandscape* Landscape = RenderContext->GetLandscape();
 
-	checkf(InRenderParams.RenderGroupTargetLayerNames.Num() == 1, TEXT("Normals should only be generated on heightmap merge, which should have 1 and only target layer"));
+	checkf(RDGBuilderRecorder.IsRecording(), TEXT("ERenderFlags::RenderMode_Recorded means the command recorder should be recording at this point"));
+	checkf(RenderParams.RenderGroupTargetLayerNames.Num() == 1, TEXT("Normals should only be generated on heightmap merge, which should have 1 and only target layer"));
 
-	RenderContext->CycleBlendRenderTargets(/*InDesiredWriteAccess = */ERHIAccess::RTV);
+	RenderContext->CycleBlendRenderTargets(/*InDesiredWriteAccess = */ERHIAccess::RTV, RDGBuilderRecorder);
 	ULandscapeScratchRenderTarget* WriteRT = RenderContext->GetBlendRenderTargetWrite();
 	ULandscapeScratchRenderTarget* ReadRT = RenderContext->GetBlendRenderTargetRead();
-	ULandscapeScratchRenderTarget* ValidityRT = RenderContext->GetValidityRenderTarget(InRenderParams.RenderGroupTargetLayerNames[0]);
+	ULandscapeScratchRenderTarget* ValidityRT = RenderContext->GetValidityRenderTarget(RenderParams.RenderGroupTargetLayerNames[0]);
 	check((WriteRT != nullptr) && (ReadRT != nullptr) && (ValidityRT != nullptr));
 
-	ValidityRT->TransitionTo(ERHIAccess::SRVMask);
+	ValidityRT->TransitionTo(ERHIAccess::SRVMask, RDGBuilderRecorder);
 
 	FIntPoint NumComponentsInRenderArea = RenderBatch->MaxComponentKey - RenderBatch->MinComponentKey + 1;
 	check((NumComponentsInRenderArea.X > 0) && (NumComponentsInRenderArea.Y > 0));
 
-	ENQUEUE_RENDER_COMMAND(LandscapeLayers_Cmd_HeightmapsGenerateNormals)(
-		[OutputResource = WriteRT->GetRenderTarget2D()->GetResource()
+	auto RDGCommand =
+		[ OutputResource = WriteRT->GetRenderTarget2D()->GetResource()
 		, OutputResourceName = WriteRT->GetDebugName()
 		, SourceResource = ReadRT->GetRenderTarget2D()->GetResource()
 		, ValidityResource = ValidityRT->GetRenderTarget2D()->GetResource()
 		, EffectiveTextureSize = RenderBatch->GetRenderTargetResolution(/*bInWithDuplicateBorders = */false)
 		, LandscapeGridScale = Landscape->GetRootComponent()->GetRelativeScale3D()
 		, ComponentSizeQuads = Landscape->ComponentSizeQuads
-		, NumComponentsInRenderArea]
-		(FRHICommandListImmediate& InRHICmdList)
-	{
-		SCOPED_DRAW_EVENTF(InRHICmdList, LandscapeLayers, TEXT("GenerateNormals -> %s"), OutputResourceName);
-		FRDGBuilder GraphBuilder(InRHICmdList, RDG_EVENT_NAME("HeightmapsGenerateNormals"));
-
-		FRDGTextureRef OutputTextureRef = GraphBuilder.RegisterExternalTexture(CreateRenderTarget(OutputResource->TextureRHI, TEXT("OutputTexture")));
-		FRDGTextureRef SourceTextureRef = GraphBuilder.RegisterExternalTexture(CreateRenderTarget(SourceResource->TextureRHI, TEXT("SourceTexture")));
-		FRDGTextureRef ValidityTextureRef = GraphBuilder.RegisterExternalTexture(CreateRenderTarget(ValidityResource->TextureRHI, TEXT("ValidityTexture")));
-
-		FLandscapeEditLayersHeightmapsGenerateNormalsPS::FParameters* PSParams = GraphBuilder.AllocParameters<FLandscapeEditLayersHeightmapsGenerateNormalsPS::FParameters>();
-		PSParams->RenderTargets[0] = FRenderTargetBinding(OutputTextureRef, ERenderTargetLoadAction::ENoAction);
-		PSParams->InTextureSize = FUintVector4(EffectiveTextureSize.X, EffectiveTextureSize.Y, SourceResource->GetSizeX(), SourceResource->GetSizeY());
-		PSParams->InLandscapeGridScale = static_cast<FVector3f>(LandscapeGridScale);
-		PSParams->InComponentSizeQuads = ComponentSizeQuads;
-		PSParams->InNumComponents = FUintVector2(NumComponentsInRenderArea.X, NumComponentsInRenderArea.Y);
-		PSParams->InSourceHeightmapSampler = TStaticSamplerState<SF_Point, AM_Clamp, AM_Clamp>::GetRHI();
-		PSParams->InSourceHeightmap = GraphBuilder.CreateSRV(FRDGTextureSRVDesc::Create(SourceTextureRef));
-		PSParams->InValidityTexture = GraphBuilder.CreateSRV(FRDGTextureSRVDesc::Create(ValidityTextureRef));
-		FLandscapeEditLayersHeightmapsGenerateNormalsPS::GenerateNormalsPS(GraphBuilder, PSParams, EffectiveTextureSize);
-
-		// We need to specify the final state of the external texture to prevent the graph builder from transitioning it to SRVMask :
-		GraphBuilder.SetTextureAccessFinal(OutputTextureRef, ERHIAccess::RTV);
+		, NumComponentsInRenderArea](FRDGBuilder& GraphBuilder)
+		{
+			FRDGTextureRef OutputTextureRef = GraphBuilder.RegisterExternalTexture(CreateRenderTarget(OutputResource->TextureRHI, TEXT("OutputTexture")));
+			FRDGTextureRef SourceTextureRef = GraphBuilder.RegisterExternalTexture(CreateRenderTarget(SourceResource->TextureRHI, TEXT("SourceTexture")));
+			FRDGTextureRef ValidityTextureRef = GraphBuilder.RegisterExternalTexture(CreateRenderTarget(ValidityResource->TextureRHI, TEXT("ValidityTexture")));
+
+			FLandscapeEditLayersHeightmapsGenerateNormalsPS::FParameters* PSParams = GraphBuilder.AllocParameters<FLandscapeEditLayersHeightmapsGenerateNormalsPS::FParameters>();
+			PSParams->RenderTargets[0] = FRenderTargetBinding(OutputTextureRef, ERenderTargetLoadAction::ENoAction);
+			PSParams->InTextureSize = FUintVector4(EffectiveTextureSize.X, EffectiveTextureSize.Y, SourceResource->GetSizeX(), SourceResource->GetSizeY());
+			PSParams->InLandscapeGridScale = static_cast<FVector3f>(LandscapeGridScale);
+			PSParams->InComponentSizeQuads = ComponentSizeQuads;
+			PSParams->InNumComponents = FUintVector2(NumComponentsInRenderArea.X, NumComponentsInRenderArea.Y);
+			PSParams->InSourceHeightmapSampler = TStaticSamplerState<SF_Point, AM_Clamp, AM_Clamp>::GetRHI();
+			PSParams->InSourceHeightmap = GraphBuilder.CreateSRV(FRDGTextureSRVDesc::Create(SourceTextureRef));
+			PSParams->InValidityTexture = GraphBuilder.CreateSRV(FRDGTextureSRVDesc::Create(ValidityTextureRef));
+			FLandscapeEditLayersHeightmapsGenerateNormalsPS::GenerateNormalsPS(
+				RDG_EVENT_NAME("GenerateNormals -> %s", *OutputResourceName), GraphBuilder, PSParams, EffectiveTextureSize);
+		};
 
-		GraphBuilder.Execute();
-	});
+	// We need to specify the final state of the external textures to prevent the graph builder from transitioning them to SRVMask (even those that end up as SRVMask at the end of this command,
+	//  because they will likely be part of another RDGCommand down the line so we need to maintain an accurate picture of every external texture ever involved in the recorded command so that 
+	//  we can set a proper access when the recorder is flushed (and the FRDGBuilder, executed) :
+	TArray<FRDGBuilderRecorder::FRDGExternalTextureAccessFinal> RDGExternalTextureAccessFinalList =
+	{
+		{ WriteRT->GetRenderTarget()->GetResource(), ERHIAccess::RTV },
+		{ ReadRT->GetRenderTarget()->GetResource(), ERHIAccess::SRVMask },
+		{ ValidityRT->GetRenderTarget()->GetResource(), ERHIAccess::SRVMask },
+	};
+	RDGBuilderRecorder.EnqueueRDGCommand(RDGCommand, RDGExternalTextureAccessFinalList);
 }
 
-void ULandscapeWeightmapWeightBlendedLayersRenderer::RenderLayer(FRenderParams& InRenderParams)
+void ULandscapeWeightmapWeightBlendedLayersRenderer::RenderLayer(FRenderParams& RenderParams, UE::Landscape::FRDGBuilderRecorder& RDGBuilderRecorder)
 {
+	using namespace UE::Landscape;
 	using namespace UE::Landscape::EditLayers;
 	using namespace UE::Landscape::EditLayers::Private;
 
-	check(!InRenderParams.MergeRenderContext->IsHeightmapMerge());
+	checkf(RDGBuilderRecorder.IsRecording(), TEXT("ERenderFlags::RenderMode_Recorded means the command recorder should be recording at this point"));
+	check(!RenderParams.MergeRenderContext->IsHeightmapMerge());
 
-	InRenderParams.MergeRenderContext->CycleBlendRenderTargets(ERHIAccess::RTV);
-	ULandscapeScratchRenderTarget* WriteRT = InRenderParams.MergeRenderContext->GetBlendRenderTargetWrite();
-	ULandscapeScratchRenderTarget* ReadRT = InRenderParams.MergeRenderContext->GetBlendRenderTargetRead();
+	RenderParams.MergeRenderContext->CycleBlendRenderTargets(ERHIAccess::RTV, RDGBuilderRecorder);
+	ULandscapeScratchRenderTarget* WriteRT = RenderParams.MergeRenderContext->GetBlendRenderTargetWrite();
+	ULandscapeScratchRenderTarget* ReadRT = RenderParams.MergeRenderContext->GetBlendRenderTargetRead();
 
 	TArray<FWeightmapPaintLayerInfo> WeightmapPaintLayerInfos;
-	WeightmapPaintLayerInfos.Reserve(InRenderParams.RenderGroupTargetLayerInfos.Num());
-	Algo::Transform(InRenderParams.RenderGroupTargetLayerInfos, WeightmapPaintLayerInfos, [](ULandscapeLayerInfoObject* InLayerInfo)
+	WeightmapPaintLayerInfos.Reserve(RenderParams.RenderGroupTargetLayerInfos.Num());
+	Algo::Transform(RenderParams.RenderGroupTargetLayerInfos, WeightmapPaintLayerInfos, [](ULandscapeLayerInfoObject* InLayerInfo)
 		{
 			check(InLayerInfo != nullptr); // There should only be valid layer infos at this point
 			FWeightmapPaintLayerInfo WeightmapPaintLayerInfo;
@@ -6508,43 +6514,46 @@ void ULandscapeWeightmapWeightBlendedLayersRenderer::RenderLayer(FRenderParams&
 			return WeightmapPaintLayerInfo;
 		});
 
-	ENQUEUE_RENDER_COMMAND(LandscapeLayers_Cmd_WeightmapsPerformLegacyWeightBlending)(
-		[TargetLayerNames = InRenderParams.RenderGroupTargetLayerNames
+	auto RDGCommand =
+		[TargetLayerNames = RenderParams.RenderGroupTargetLayerNames
 		, WeightmapPaintLayerInfos
 		, OutputResource = WriteRT->GetRenderTarget2DArray()->GetResource()
 		, OutputResourceName = WriteRT->GetDebugName()
 		, CurrentEditLayerResource = ReadRT->GetRenderTarget2DArray()->GetResource()
-		, EffectiveTextureSize = InRenderParams.MergeRenderContext->GetCurrentRenderBatch()->GetRenderTargetResolution(/*bInWithDuplicateBorders = */false)]
-		(FRHICommandListImmediate& InRHICmdList)
-	{
-		SCOPED_DRAW_EVENTF(InRHICmdList, LandscapeLayers, TEXT("LegacyWeightBlending -> %s"), OutputResourceName);
-		FRDGBuilder GraphBuilder(InRHICmdList, RDG_EVENT_NAME("WeightmapsPerformLegacyWeightBlending"));
-
-		FRDGTextureRef OutputTextureRef = GraphBuilder.RegisterExternalTexture(CreateRenderTarget(OutputResource->TextureRHI, TEXT("OutputTexture"))/*, ERDGTextureFlags::SkipTracking*/);
-		FRDGTextureRef CurrentEditLayerTextureRef = GraphBuilder.RegisterExternalTexture(CreateRenderTarget(CurrentEditLayerResource->TextureRHI, TEXT("CurrentEditLayerTexture")));
-		FRDGTextureSRVRef CurrentEditLayerTextureSRVRef = GraphBuilder.CreateSRV(FRDGTextureSRVDesc::Create(CurrentEditLayerTextureRef));
-
-		FRDGBufferRef PaintLayerInfosBuffer = CreateStructuredBuffer(GraphBuilder, TEXT("LandscapePaintLayerInfosBuffer"), WeightmapPaintLayerInfos);
-		FRDGBufferSRVRef PaintLayerInfosBufferSRV = GraphBuilder.CreateSRV(FRDGBufferSRVDesc(PaintLayerInfosBuffer));
-
-		const int32 NumTargetLayers = TargetLayerNames.Num();
-		for (int32 TargetLayerIndex = 0; TargetLayerIndex < NumTargetLayers; ++TargetLayerIndex)
+		, EffectiveTextureSize = RenderParams.MergeRenderContext->GetCurrentRenderBatch()->GetRenderTargetResolution(/*bInWithDuplicateBorders = */false)](FRDGBuilder& GraphBuilder)
 		{
-			FLandscapeEditLayersWeightmapsPerformLegacyWeightBlendingPS::FParameters* PSParams = GraphBuilder.AllocParameters<FLandscapeEditLayersWeightmapsPerformLegacyWeightBlendingPS::FParameters>();
-			PSParams->RenderTargets[0] = FRenderTargetBinding(OutputTextureRef, ERenderTargetLoadAction::ENoAction, /*InMipIndex = */0, /*InArraySlice = */TargetLayerIndex);
-			PSParams->InPaintLayerIndex = TargetLayerIndex;
-			PSParams->InNumPaintLayers = NumTargetLayers;
-			PSParams->InPaintLayerInfos = PaintLayerInfosBufferSRV;
-			PSParams->InCurrentEditLayerWeightmaps = CurrentEditLayerTextureSRVRef;
+			FRDGTextureRef OutputTextureRef = GraphBuilder.RegisterExternalTexture(CreateRenderTarget(OutputResource->TextureRHI, TEXT("OutputTexture")));
+			FRDGTextureRef CurrentEditLayerTextureRef = GraphBuilder.RegisterExternalTexture(CreateRenderTarget(CurrentEditLayerResource->TextureRHI, TEXT("CurrentEditLayerTexture")));
+			FRDGTextureSRVRef CurrentEditLayerTextureSRVRef = GraphBuilder.CreateSRV(FRDGTextureSRVDesc::Create(CurrentEditLayerTextureRef));
 
-			FLandscapeEditLayersWeightmapsPerformLegacyWeightBlendingPS::PerformLegacyWeightBlendingPS(GraphBuilder, PSParams, TargetLayerNames[TargetLayerIndex], EffectiveTextureSize);
-		}
+			FRDGBufferRef PaintLayerInfosBuffer = CreateStructuredBuffer(GraphBuilder, TEXT("LandscapePaintLayerInfosBuffer"), WeightmapPaintLayerInfos);
+			FRDGBufferSRVRef PaintLayerInfosBufferSRV = GraphBuilder.CreateSRV(FRDGBufferSRVDesc(PaintLayerInfosBuffer));
 
-		// We need to specify the final state of the external texture to prevent the graph builder from transitioning it to SRVMask :
-		GraphBuilder.SetTextureAccessFinal(OutputTextureRef, ERHIAccess::RTV);
+			const int32 NumTargetLayers = TargetLayerNames.Num();
+			for (int32 TargetLayerIndex = 0; TargetLayerIndex < NumTargetLayers; ++TargetLayerIndex)
+			{
+				FLandscapeEditLayersWeightmapsPerformLegacyWeightBlendingPS::FParameters* PSParams = GraphBuilder.AllocParameters<FLandscapeEditLayersWeightmapsPerformLegacyWeightBlendingPS::FParameters>();
+				PSParams->RenderTargets[0] = FRenderTargetBinding(OutputTextureRef, ERenderTargetLoadAction::ENoAction, /*InMipIndex = */0, /*InArraySlice = */TargetLayerIndex);
+				PSParams->InPaintLayerIndex = TargetLayerIndex;
+				PSParams->InNumPaintLayers = NumTargetLayers;
+				PSParams->InPaintLayerInfos = PaintLayerInfosBufferSRV;
+				PSParams->InCurrentEditLayerWeightmaps = CurrentEditLayerTextureSRVRef;
 
-		GraphBuilder.Execute();
-	});
+				FLandscapeEditLayersWeightmapsPerformLegacyWeightBlendingPS::PerformLegacyWeightBlendingPS(
+					RDG_EVENT_NAME("LegacyWeightBlending(%s) -> %s", *TargetLayerNames[TargetLayerIndex].ToString(), *OutputResourceName),
+					GraphBuilder, PSParams, EffectiveTextureSize);
+			}
+		};
+	
+	// We need to specify the final state of the external textures to prevent the graph builder from transitioning them to SRVMask (even those that end up as SRVMask at the end of this command,
+	//  because they will likely be part of another RDGCommand down the line so we need to maintain an accurate picture of every external texture ever involved in the recorded command so that 
+	//  we can set a proper access when the recorder is flushed (and the FRDGBuilder, executed) :
+	TArray<FRDGBuilderRecorder::FRDGExternalTextureAccessFinal> RDGExternalTextureAccessFinalList =
+	{
+		{ WriteRT->GetRenderTarget()->GetResource(), ERHIAccess::RTV },
+		{ ReadRT->GetRenderTarget()->GetResource(), ERHIAccess::SRVMask },
+	};
+	RDGBuilderRecorder.EnqueueRDGCommand(RDGCommand, RDGExternalTextureAccessFinalList);
 }
 
 UE::Landscape::EditLayers::FMergeRenderContext ALandscape::PrepareEditLayersMergeRenderContext(const UE::Landscape::EditLayers::FMergeContext& InMergeContext, const UE::Landscape::EditLayers::FMergeRenderParams& InMergeRenderParams)
@@ -6571,9 +6580,9 @@ UE::Landscape::EditLayers::FMergeRenderContext ALandscape::PrepareEditLayersMerg
 	if (!InMergeContext.IsHeightmapMerge())
 	{
 		for (FName TargetLayerName : InMergeRenderParams.WeightmapLayerNames)
-		{
+	{
 			if (!InMergeContext.IsValidTargetLayerName(TargetLayerName))
-			{
+		{ 
 				UE_LOG(LogLandscape, Warning, TEXT("Target layer \"%s\" was requested by the merge but is invalid (missing its layer info object asset). It will be ignored"), *TargetLayerName.ToString());
 			}
 		}
@@ -6636,7 +6645,7 @@ UE::Landscape::EditLayers::FMergeRenderContext ALandscape::PrepareEditLayersMerg
 	}
 	else
 	{
-		// First, let's work out the weightmaps inter-dependencies (i.e. horizontal dependencies) : weight-blending requires some weightmaps to be processed together: 
+	// First, let's work out the weightmaps inter-dependencies (i.e. horizontal dependencies) : weight-blending requires some weightmaps to be processed together: 
 		TRACE_CPUPROFILER_EVENT_SCOPE(PrepareRenderGroups);
 
 		TArray<TArray<TBitArray<>>> RelevantRenderGroupsPerRenderer;
@@ -6763,6 +6772,7 @@ UE::Landscape::EditLayers::FMergeRenderContext ALandscape::PrepareEditLayersMerg
 	const TArray<ULandscapeComponent*> AllComponents = Component2DIndexer.GetAllValues();
 	const TBitArray<> AllValidComponentBitIndices = Component2DIndexer.GetValidValueBitIndices();
 	const int32 NumAllComponents = AllComponents.Num();
+	const FEditLayerRendererState DummyRendererState(&MergeRenderContext);
 
 	// The list of all components that will end up being rendered across all renderers (one bit per component) :
 	TBitArray<> FinalComponentsToRenderInfoBitIndices(false, NumAllComponents);
@@ -7146,21 +7156,49 @@ UE::Landscape::EditLayers::FMergeRenderContext ALandscape::PrepareEditLayersMerg
 			// Now, we have all the info to build our list of successive render steps : process group by group :
 			for (const TBitArray<>& RenderGroup : FinalRenderGroups)
 			{
+				ERenderFlags LastRenderMode = ERenderFlags::None;
+
 				// Execute the required operations for the whole stack of renderers for this group :
 				for (int32 RendererIndex = 0; RendererIndex < LocalNumRenderers; ++RendererIndex)
 				{
 					const FEditLayerRendererRenderInfo& EditLayerRendererRenderInfo = OrderedEditLayerRendererRenderInfos[RendererIndex];
 					TArrayView<ULandscapeComponent*> ComponentsToRender = MakeArrayView(ComponentsForRenderers[RendererIndex]);
 
-					if (!ComponentsToRender.IsEmpty() && EditLayerRendererRenderInfo.RendererState.GetRenderer()->CanRender())
+					if (!ComponentsToRender.IsEmpty())
 					{
+						const ERenderFlags CurrentRenderFlags = EditLayerRendererRenderInfo.RendererState.GetRenderer()->GetRenderFlags(&MergeRenderContext);
+						const ERenderFlags CurrentRenderMode = CurrentRenderFlags & ERenderFlags::RenderMode_Mask;
+						checkf(FMath::CountBits(static_cast<uint8>(CurrentRenderMode)) == 1, TEXT("Render mode should be either Immediate or Recorded"));
+
+						// Initiate or terminate the "render layer on render thread" sequence if necessary :
+						if (CurrentRenderMode == ERenderFlags::RenderMode_Recorded)
+						{
+							if (LastRenderMode != ERenderFlags::RenderMode_Recorded)
+							{
+								MergeRenderBatch.RenderSteps.Emplace(FMergeRenderStep::EType::BeginRenderCommandRecorder, DummyRendererState, RenderGroup, MergeRenderBatch.ComponentsToRender.Array());
+							}
+						}
+						else if (LastRenderMode == ERenderFlags::RenderMode_Recorded)
+						{
+							MergeRenderBatch.RenderSteps.Emplace(FMergeRenderStep::EType::EndRenderCommandRecorder, DummyRendererState, RenderGroup, MergeRenderBatch.ComponentsToRender.Array());
+						}
+						LastRenderMode = CurrentRenderMode;
+
 						// Render the content of this render layer : 
-						MergeRenderBatch.RenderSteps.Emplace(FMergeRenderStep::EType::RenderLayer, EditLayerRendererRenderInfo.RendererState, RenderGroup, ComponentsToRender);
+						FMergeRenderStep::EType RenderStepType = (CurrentRenderMode == ERenderFlags::RenderMode_Immediate)
+							? FMergeRenderStep::EType::RenderLayer_Immediate : FMergeRenderStep::EType::RenderLayer_Recorded;
+						MergeRenderBatch.RenderSteps.Emplace(RenderStepType, EditLayerRendererRenderInfo.RendererState, RenderGroup, ComponentsToRender);
 					}
 				}
 
+				// Insert a final render sequence end step if we ended on a render on the render thread :
+				if (LastRenderMode == ERenderFlags::RenderMode_Recorded)
+				{
+					MergeRenderBatch.RenderSteps.Emplace(FMergeRenderStep::EType::EndRenderCommandRecorder, DummyRendererState, RenderGroup, MergeRenderBatch.ComponentsToRender.Array());
+				}
+
 				// Finally, signal the group as done for this batch
-				MergeRenderBatch.RenderSteps.Emplace(FMergeRenderStep::EType::SignalBatchMergeGroupDone, FEditLayerRendererState(&MergeRenderContext), RenderGroup, MergeRenderBatch.ComponentsToRender.Array());
+				MergeRenderBatch.RenderSteps.Emplace(FMergeRenderStep::EType::SignalBatchMergeGroupDone, DummyRendererState, RenderGroup, MergeRenderBatch.ComponentsToRender.Array());
 
 				MergeRenderContext.MaxNeededNumSlices = FMath::Max(MergeRenderContext.MaxNeededNumSlices, RenderGroup.CountSetBits());
 			}
@@ -7490,8 +7528,10 @@ int32 ALandscape::PerformLayersHeightmapsBatchedMerge(const FUpdateLayersContent
 	TSet<ULandscapeComponent*> ResolvedComponents;
 	ResolvedComponents.Reserve(InUpdateLayersContentContext.LandscapeComponentsHeightmapsToRender.Num());
 
+	// TODO [jonathan.bard] : this could also be recorded into RDGBuilderRecorder to avoid using additional FRDGBuilders at this step...
 	// Callback executed each time a render batch is done computing the requested into, just before releasing the render resources : 
-	auto OnRenderBatchGroupDone = [&ResolvedComponents, &InUpdateLayersContentContext, OnEditLayersMergedDelegate = &OnEditLayersMergedDelegate](const FMergeRenderContext::FOnRenderBatchTargetGroupDoneParams& InParams)
+	auto OnRenderBatchGroupDone = [&ResolvedComponents, &InUpdateLayersContentContext, OnEditLayersMergedDelegate = &OnEditLayersMergedDelegate]
+		(const FMergeRenderContext::FOnRenderBatchTargetGroupDoneParams& InParams, UE::Landscape::FRDGBuilderRecorder& RDGBuilderRecorder)
 	{
 		ALandscape* Landscape = InParams.MergeRenderContext->GetLandscape();
 
@@ -7537,7 +7577,7 @@ int32 ALandscape::PerformLayersHeightmapsBatchedMerge(const FUpdateLayersContent
 			RHI_BREADCRUMB_EVENT_GAMETHREAD("CopyMip0AndExpand");
 
 			// Recompose mip0 of the final heightmaps, subsection by subsection, to duplicate borders : 
-			InParams.MergeRenderContext->CycleBlendRenderTargets(/*InDesiredWriteAccess = */ERHIAccess::SRVMask); // TODO [jonathan.bard] : This should be CopyDst but ExecuteCopyLayersTexture doesn't allow for it ATM
+			InParams.MergeRenderContext->CycleBlendRenderTargets(/*InDesiredWriteAccess = */ERHIAccess::SRVMask, RDGBuilderRecorder); // TODO [jonathan.bard] : This should be CopyDst but ExecuteCopyLayersTexture doesn't allow for it ATM
 			ULandscapeScratchRenderTarget* ReadRT = InParams.MergeRenderContext->GetBlendRenderTargetRead();
 
 			OnEditLayersMergedDelegate->Broadcast(FOnLandscapeEditLayersMergedParams(ReadRT->GetRenderTarget(),
@@ -7597,7 +7637,7 @@ int32 ALandscape::PerformLayersHeightmapsBatchedMerge(const FUpdateLayersContent
 
 			// "Expand" the scratch render target : 
 			// Copy sub-section by sub-section in order to duplicate borders :
-			InParams.MergeRenderContext->RenderExpandedRenderTarget(*InParams.RenderBatch);
+			InParams.MergeRenderContext->RenderExpandedRenderTarget(*InParams.RenderBatch, RDGBuilderRecorder);
 		}
 
 		// Generate the mips from the expanded RT and copy to the final texture mips
@@ -7610,7 +7650,7 @@ int32 ALandscape::PerformLayersHeightmapsBatchedMerge(const FUpdateLayersContent
 			FIntPoint CurrentMipSubsectionSize = Landscape->SubsectionSizeQuads + 1;
 			for (int32 MipIndex = 1; MipIndex < NumMips; ++MipIndex)
 			{
-				InParams.MergeRenderContext->CycleBlendRenderTargets(/*InDesiredWriteAccess = */ERHIAccess::RTV);
+				InParams.MergeRenderContext->CycleBlendRenderTargets(/*InDesiredWriteAccess = */ERHIAccess::RTV, RDGBuilderRecorder);
 				ULandscapeScratchRenderTarget* WriteRT = InParams.MergeRenderContext->GetBlendRenderTargetWrite();
 				ULandscapeScratchRenderTarget* ReadRT = InParams.MergeRenderContext->GetBlendRenderTargetRead();
 				CurrentMipResolution /= 2;
@@ -7652,7 +7692,7 @@ int32 ALandscape::PerformLayersHeightmapsBatchedMerge(const FUpdateLayersContent
 				{
 					RHI_BREADCRUMB_EVENT_GAMETHREAD("Copy mip %i", MipIndex);
 
-					WriteRT->TransitionTo(ERHIAccess::SRVMask);
+					WriteRT->TransitionTo(ERHIAccess::SRVMask, RDGBuilderRecorder);
 					TArray<FLandscapeLayersCopyTextureParams> DeferredCopyTextures;
 					for (ULandscapeComponent* Component : ComponentsToResolveThisBatch)
 					{
@@ -10399,7 +10439,8 @@ int32 ALandscape::PerformLayersWeightmapsBatchedMerge(FUpdateLayersContentContex
 	}
 
 	// Callback executed each time a render batch is done computing the requested into, just before releasing the render resources : 
-	auto OnRenderBatchGroupDone = [&ResolvedLandscapeComponents, &TexturesNeedingResolve, &InUpdateLayersContentContext, OnEditLayersMergedDelegate = &OnEditLayersMergedDelegate](const FMergeRenderContext::FOnRenderBatchTargetGroupDoneParams& InParams)
+	auto OnRenderBatchGroupDone = [&ResolvedLandscapeComponents, &TexturesNeedingResolve, &InUpdateLayersContentContext, OnEditLayersMergedDelegate = &OnEditLayersMergedDelegate]
+		(const FMergeRenderContext::FOnRenderBatchTargetGroupDoneParams& InParams, FRDGBuilderRecorder& RDGBuilderRecorder)
 	{
 		// We can now finalize the weightmaps : since we don't use SUPPORTS_LANDSCAPE_EDITORONLY_UBER_MATERIAL yet, each component that has been rendered should have a complete set of weightmaps that has been regenerated
 		//  already so we need to repack those into the appropriate weightmap channels
@@ -10407,7 +10448,7 @@ int32 ALandscape::PerformLayersWeightmapsBatchedMerge(FUpdateLayersContentContex
 		ALandscape* Landscape = InParams.MergeRenderContext->GetLandscape();
 
 		// Cycle render targets one last time so we can read from the last RT we've written to :
-		InParams.MergeRenderContext->CycleBlendRenderTargets();
+			InParams.MergeRenderContext->CycleBlendRenderTargets(/*InDesiredWriteAccess = */ERHIAccess::None, RDGBuilderRecorder);
 		ULandscapeScratchRenderTarget* ReadRT = InParams.MergeRenderContext->GetBlendRenderTargetRead();
 
 		OnEditLayersMergedDelegate->Broadcast(FOnLandscapeEditLayersMergedParams(ReadRT->GetRenderTarget(),
diff --git a/Engine/Source/Runtime/Landscape/Private/LandscapeEditResourcesSubsystem.cpp b/Engine/Source/Runtime/Landscape/Private/LandscapeEditResourcesSubsystem.cpp
--- a/Engine/Source/Runtime/Landscape/Private/LandscapeEditResourcesSubsystem.cpp
+++ b/Engine/Source/Runtime/Landscape/Private/LandscapeEditResourcesSubsystem.cpp
@@ -7,6 +7,7 @@
 #include "TextureResource.h"
 #include "UObject/Package.h"
 #include "RenderGraphUtils.h"
+#include "LandscapeUtils.h"
 
 #if WITH_EDITOR
 #include "LandscapeRender.h"
@@ -45,7 +46,8 @@ FRHITransitionInfo ULandscapeScratchRenderTarget::FTransitionInfo::ToRHITransiti
 
 // ----------------------------------------------------------------------------------
 
-ULandscapeScratchRenderTarget::FTransitionBatcherScope::FTransitionBatcherScope()
+ULandscapeScratchRenderTarget::FTransitionBatcherScope::FTransitionBatcherScope(UE::Landscape::FRDGBuilderRecorder& RDGBuilderRecorder)
+	: RDGBuilderRecorder(RDGBuilderRecorder)
 {
 
 }
@@ -54,13 +56,18 @@ ULandscapeScratchRenderTarget::FTransitionBatcherScope::~FTransitionBatcherScope
 {
 	if (!PendingTransitions.IsEmpty())
 	{
-		ENQUEUE_RENDER_COMMAND(TransitionScratchRTs)(
-			[Transitions = MoveTemp(PendingTransitions)](FRHICommandListImmediate& RHICmdList) mutable
+		// Don't transition when recording : the graph builder will do it automatically. It is simply required that the user specifies to the command recorder the state of each external texture
+		//  ever used in one of the recorded RDG commands, if they want to prevent the auto-transition to SRVMask when the commands are flushed and the FRDGBuilder, executed.
+		if (!RDGBuilderRecorder.IsRecording())
 		{
-			TArray<FRHITransitionInfo, TInlineAllocator<8>> RHITransitions;
-			Algo::Transform(Transitions, RHITransitions, [](const FTransitionInfo& InTransitionInfo) { return InTransitionInfo.ToRHITransitionInfo(); });
-			RHICmdList.Transition(RHITransitions);
-		});
+			auto PerformRHITransitions = [Transitions = MoveTemp(PendingTransitions)](FRHICommandListImmediate& RHICmdList) mutable
+				{
+					TArray<FRHITransitionInfo, TInlineAllocator<8>> RHITransitions;
+					Algo::Transform(Transitions, RHITransitions, [](const FTransitionInfo& InTransitionInfo) { return InTransitionInfo.ToRHITransitionInfo(); });
+					RHICmdList.Transition(RHITransitions);
+				};
+			RDGBuilderRecorder.EnqueueRenderCommand(PerformRHITransitions);
+		}
 	}
 }
 
@@ -162,62 +169,45 @@ ETextureRenderTargetFormat ULandscapeScratchRenderTarget::GetFormat() const
 	return RenderTargetFormat;
 }
 
-void ULandscapeScratchRenderTarget::TransitionTo(ERHIAccess InDesiredState)
+void ULandscapeScratchRenderTarget::TransitionTo(ERHIAccess InDesiredState, UE::Landscape::FRDGBuilderRecorder& RDGBuilderRecorder)
 {
 	check(RenderTarget != nullptr);
 
-	FTransitionBatcherScope TransitionScope;
+	FTransitionBatcherScope TransitionScope(RDGBuilderRecorder);
 	TransitionScope.TransitionTo(this, InDesiredState);
 }
 
-void ULandscapeScratchRenderTarget::Clear()
+void ULandscapeScratchRenderTarget::Clear(UE::Landscape::FRDGBuilderRecorder& RDGBuilderRecorder)
 {
+	using namespace UE::Landscape;
+
 	check(RenderTarget != nullptr);
 
-	TransitionTo(ERHIAccess::RTV);
+	TransitionTo(ERHIAccess::RTV, RDGBuilderRecorder);
 
-	ENQUEUE_RENDER_COMMAND(ClearScratchRT)(
+	auto RDGCommand = 
 		[ Resource = RenderTarget->GameThread_GetRenderTargetResource()
-		, ResourceName = GetDebugName()
-		, EffectiveResolution = GetEffectiveResolution()]
-		(FRHICommandListImmediate& InRHICmdList) mutable
-	{
-		SCOPED_DRAW_EVENTF(InRHICmdList, LandscapeLayers, TEXT("Clear %s"), ResourceName);
-
-		const FRHITextureDesc& TextureDesc = Resource->TextureRHI->GetDesc();
-		TArray<FRHIRenderPassInfo> RenderPassInfos;
-		// If texture array, we need multiple clear passes : 
-		if (TextureDesc.IsTextureArray())
+		, EffectiveNumSlices = GetEffectiveNumSlices()]
+		(FRDGBuilder& GraphBuilder)
 		{
-			for (int32 SliceIndex = 0; SliceIndex < TextureDesc.ArraySize;)
+			FRDGTextureRef TextureRef = GraphBuilder.RegisterExternalTexture(CreateRenderTarget(Resource->GetTextureRHI(), TEXT("ClearTexture")));
+			FRDGTextureClearInfo ClearInfo;
+			if (TextureRef->Desc.IsTextureArray())
 			{
-				// We can clear up to MaxSimultaneousRenderTargets slices in a single pass: 
-				FRHIRenderPassInfo& RenderPassInfo = RenderPassInfos.Emplace_GetRef();
-				for (int32 RenderTargetPassIndex = 0; (RenderTargetPassIndex < MaxSimultaneousRenderTargets) && (SliceIndex < TextureDesc.ArraySize); ++SliceIndex, ++RenderTargetPassIndex)
-				{
-					RenderPassInfo.ColorRenderTargets[RenderTargetPassIndex].RenderTarget = Resource->TextureRHI;
-					RenderPassInfo.ColorRenderTargets[RenderTargetPassIndex].ArraySlice = SliceIndex;
-					RenderPassInfo.ColorRenderTargets[RenderTargetPassIndex].Action = ERenderTargetActions::Clear_Store;
-				}
+				check(EffectiveNumSlices <= TextureRef->Desc.ArraySize);
+				ClearInfo.NumSlices = EffectiveNumSlices;
 			}
-		}
-		else
-		{
-			RenderPassInfos.Emplace(Resource->TextureRHI, ERenderTargetActions::Clear_Store);
-		}
+			AddClearRenderTargetPass(GraphBuilder, TextureRef, ClearInfo);
+		};
 
-		for (const FRHIRenderPassInfo& RenderPassInfo : RenderPassInfos)
-		{
-			InRHICmdList.BeginRenderPass(RenderPassInfo, TEXT("Clear"));
-			InRHICmdList.EndRenderPass();
-		}
-	});
+	// We need to specify the final state of the external texture to prevent the graph builder from transitioning it to SRVMask :
+	RDGBuilderRecorder.EnqueueRDGCommand(RDGCommand, { { RenderTarget->GetResource(), ERHIAccess::RTV } });
 }
 
 namespace UE::Landscape::Private
 {
 
-void EnqueueCopyToScratchRTRenderCommand(const ULandscapeScratchRenderTarget::FCopyFromParams& InCopyParams, FTextureResource* InSourceTextureResource, FTextureResource* InDestTextureResource)
+void EnqueueCopyToScratchRTRenderCommand(const ULandscapeScratchRenderTarget::FCopyFromParams& InCopyParams, FTextureResource* InSourceTextureResource, FTextureResource* InDestTextureResource, UE::Landscape::FRDGBuilderRecorder& RDGBuilderRecorder)
 {
 	FIntPoint SourceSize(InSourceTextureResource->GetSizeX() >> InCopyParams.SourceMip, InSourceTextureResource->GetSizeY() >> InCopyParams.SourceMip);
 	FIntPoint DestSize(InDestTextureResource->GetSizeX() >> InCopyParams.DestMip, InDestTextureResource->GetSizeY() >> InCopyParams.DestMip);
@@ -244,35 +234,48 @@ void EnqueueCopyToScratchRTRenderCommand(const ULandscapeScratchRenderTarget::FC
 	check(Info.DestPosition.X + Info.Size.X <= DestSize.X);
 	check(Info.DestPosition.Y + Info.Size.Y <= DestSize.Y);
 
-	ENQUEUE_RENDER_COMMAND(CopyToScratchRT)(
-		[Info, InSourceTextureResource, InDestTextureResource](FRHICommandListImmediate& InRHICmdList) mutable
+	auto RDGCommand =
+		[ InSourceTextureResource
+		, InDestTextureResource
+		, Info]
+		(FRDGBuilder& GraphBuilder)
+		{
+			FRDGTextureRef SourceTextureRef = GraphBuilder.RegisterExternalTexture(CreateRenderTarget(InSourceTextureResource->GetTextureRHI(), TEXT("CopySourceTexture")));
+			FRDGTextureRef DestTextureRef = GraphBuilder.RegisterExternalTexture(CreateRenderTarget(InDestTextureResource->GetTextureRHI(), TEXT("CopyDestTexture")));
+			AddCopyTexturePass(GraphBuilder, SourceTextureRef, DestTextureRef, Info);
+		};
+
+	// We need to specify the final state of the external textures to prevent the graph builder from transitioning them to SRVMask :
+	TArray<FRDGBuilderRecorder::FRDGExternalTextureAccessFinal> RDGExternalTextureAccessFinalList =
 	{
-		InRHICmdList.CopyTexture(InSourceTextureResource->TextureRHI, InDestTextureResource->TextureRHI, Info);
-	});
+		{ InSourceTextureResource, ERHIAccess::CopySrc },
+		{ InDestTextureResource, ERHIAccess::CopyDest }
+	};
+	RDGBuilderRecorder.EnqueueRDGCommand(RDGCommand, RDGExternalTextureAccessFinalList);
 }
 
 } // end namespace UE::Landscape::Private
 
 
-void ULandscapeScratchRenderTarget::CopyFrom(const FCopyFromTextureParams& InCopyParams)
+void ULandscapeScratchRenderTarget::CopyFrom(const FCopyFromTextureParams& InCopyParams, UE::Landscape::FRDGBuilderRecorder& RDGBuilderRecorder)
 {
 	using namespace UE::Landscape::Private;
 
 	// The source is expected to be in CopySrc state already. We need to transition the scratch RT to the appropriate state, though : 
-	TransitionTo(ERHIAccess::CopyDest);
+	TransitionTo(ERHIAccess::CopyDest, RDGBuilderRecorder);
 
-	EnqueueCopyToScratchRTRenderCommand(InCopyParams, InCopyParams.SourceTexture->GetResource(), RenderTarget->GameThread_GetRenderTargetResource());
+	EnqueueCopyToScratchRTRenderCommand(InCopyParams, InCopyParams.SourceTexture->GetResource(), RenderTarget->GameThread_GetRenderTargetResource(), RDGBuilderRecorder);
 }
 
-void ULandscapeScratchRenderTarget::CopyFrom(const FCopyFromScratchRenderTargetParams& InCopyParams)
+void ULandscapeScratchRenderTarget::CopyFrom(const FCopyFromScratchRenderTargetParams& InCopyParams, UE::Landscape::FRDGBuilderRecorder& RDGBuilderRecorder)
 {
 	using namespace UE::Landscape::Private;
 
 	// We need to transition both the source and destination scratch RT to the appropriate state:
-	InCopyParams.SourceScratchRenderTarget->TransitionTo(ERHIAccess::CopySrc);
-	TransitionTo(ERHIAccess::CopyDest);
+	InCopyParams.SourceScratchRenderTarget->TransitionTo(ERHIAccess::CopySrc, RDGBuilderRecorder);
+	TransitionTo(ERHIAccess::CopyDest, RDGBuilderRecorder);
 
-	EnqueueCopyToScratchRTRenderCommand(InCopyParams, InCopyParams.SourceScratchRenderTarget->GetRenderTarget()->GameThread_GetRenderTargetResource(), RenderTarget->GameThread_GetRenderTargetResource());
+	EnqueueCopyToScratchRTRenderCommand(InCopyParams, InCopyParams.SourceScratchRenderTarget->GetRenderTarget()->GameThread_GetRenderTargetResource(), RenderTarget->GameThread_GetRenderTargetResource(), RDGBuilderRecorder);
 }
 
 bool ULandscapeScratchRenderTarget::IsCompatibleWith(const UE::Landscape::FScratchRenderTargetParams& InParams) const
@@ -335,6 +338,8 @@ bool ULandscapeScratchRenderTarget::IsCompatibleWith(const UE::Landscape::FScrat
 
 void ULandscapeScratchRenderTarget::OnRequested(const UE::Landscape::FScratchRenderTargetParams& InParams)
 {
+	using namespace UE::Landscape;
+
 	check(!IsInUse());
 
 	// If it's not been initialized yet, create the render target now : 
@@ -375,7 +380,9 @@ void ULandscapeScratchRenderTarget::OnRequested(const UE::Landscape::FScratchRen
 
 	if (InParams.InitialState != ERHIAccess::None)
 	{
-		TransitionTo(InParams.InitialState);
+		// Cannot be requested when recording RDGRenderCommandRecorder so we use an immediate recorder : 
+		FRDGBuilderRecorder RDGBuilderRecorderImmediate;
+		TransitionTo(InParams.InitialState, RDGBuilderRecorderImmediate);
 	}
 }
 
diff --git a/Engine/Source/Runtime/Landscape/Private/LandscapeUtils.cpp b/Engine/Source/Runtime/Landscape/Private/LandscapeUtils.cpp
--- a/Engine/Source/Runtime/Landscape/Private/LandscapeUtils.cpp
+++ b/Engine/Source/Runtime/Landscape/Private/LandscapeUtils.cpp
@@ -8,6 +8,8 @@
 #include "LandscapeProxy.h"
 #include "LandscapeEditTypes.h"
 #include "Algo/Transform.h"
+#include "RenderGraphBuilder.h"
+#include "TextureResource.h"
 
 #if WITH_EDITOR
 #include "AssetRegistry/AssetRegistryModule.h"
@@ -61,8 +63,136 @@ FString GetLandscapeToolTargetTypeFlagsAsString(ELandscapeToolTargetTypeFlags In
 	return *FString::Join(TargetTypeStrings, TEXT(","));
 }
 
+
+// ----------------------------------------------------------------------------------
+
+FRDGBuilderRecorder::~FRDGBuilderRecorder()
+{
+	checkf((State == EState::Immediate) && IsEmpty(), 
+	   TEXT("The command recorder has %d commands pending while being destroyed. These commands will not get executed unless they are appended to a render command : use Flush() (or Clear() to remove all commands if this is intended)."), RDGCommands.Num());
+}
+
+void FRDGBuilderRecorder::StartRecording()
+{
+	if (State == EState::Immediate)
+	{
+		State = EState::Recording;
+	}
+}
+
+void FRDGBuilderRecorder::StopRecording()
+{
+	if (State == EState::Recording)
+	{
+		State = EState::Immediate;
+	}
+}
+
+void FRDGBuilderRecorder::StopRecordingAndFlush(FRDGEventName&& EventName)
+{
+	if (State == EState::Recording)
+	{
+		Flush(MoveTemp(EventName));
+		State = EState::Immediate;
+	}
+}
+
+void FRDGBuilderRecorder::Flush(FRDGEventName&& EventName)
+{
+	if (!IsEmpty())
+	{
+		checkf(State == EState::Immediate, TEXT("StopRecording needs to be called before flushing the recorded commands"));
+
+		ENQUEUE_RENDER_COMMAND(FRDGBuilderRecorder_Flush)(
+			[RDGCommands = MoveTemp(RDGCommands), RDGExternalTextureAccessFinal = MoveTemp(RDGExternalTextureAccessFinal), EventName = MoveTemp(EventName)](FRHICommandListImmediate& InRHICmdList) mutable
+			{
+				SCOPED_DRAW_EVENTF(InRHICmdList, FRDGBuilderRecorder_Flush, TEXT("%s"), FString(EventName.GetTCHAR()));
+
+				FRDGBuilder GraphBuilder(InRHICmdList, EventName);
+
+				for (const FRDGRecorderRDGCommand& Command : RDGCommands)
+				{
+					Command(GraphBuilder);
+				}
+				
+				for (auto ItPair : RDGExternalTextureAccessFinal)
+				{
+					FRDGTextureRef TextureRef = GraphBuilder.RegisterExternalTexture(CreateRenderTarget(ItPair.Key->GetTextureRHI(), TEXT("ExternalTexture")));
+					check(TextureRef != nullptr);
+					GraphBuilder.SetTextureAccessFinal(TextureRef, ItPair.Value);
+				}
+
+				GraphBuilder.Execute();
+			});
+		check(IsEmpty());
+
+		Clear();
+	}
+}
+
+void FRDGBuilderRecorder::EnqueueRDGCommand(FRDGRecorderRDGCommand InRDGCommand, TConstArrayView<FRDGExternalTextureAccessFinal> InRDGExternalTextureAccessFinalList)
+{
+	if (State == EState::Recording)
+	{
+		RDGCommands.Add(InRDGCommand);
+		for (const FRDGExternalTextureAccessFinal& TextureAccess : InRDGExternalTextureAccessFinalList)
+		{
+			// Replace the existing value if any : this specifies the state of the texture at the moment the FRDGBuilder executes : 
+			RDGExternalTextureAccessFinal.FindOrAdd(TextureAccess.TextureResource) = TextureAccess.Access;
+		}
+	}
+	else
+	{
+		ENQUEUE_RENDER_COMMAND(FRDGBuilderRecorder_RDGCommand)([InRDGCommand, LocalRDGExternalTextureAccessFinal = TArray<FRDGExternalTextureAccessFinal>(InRDGExternalTextureAccessFinalList)](FRHICommandListImmediate& InRHICmdList)
+			{ 
+				FRDGBuilder GraphBuilder(InRHICmdList, RDG_EVENT_NAME("RDGImmediateRDGCommand"));
+				
+				InRDGCommand(GraphBuilder);
+
+				for (const FRDGExternalTextureAccessFinal& TextureAccess : LocalRDGExternalTextureAccessFinal)
+				{
+					FRDGTextureRef TextureRef = GraphBuilder.RegisterExternalTexture(CreateRenderTarget(TextureAccess.TextureResource->GetTextureRHI(), TEXT("ExternalTexture")));
+					check(TextureRef != nullptr);
+					GraphBuilder.SetTextureAccessFinal(TextureRef, TextureAccess.Access);
+				}
+
+				GraphBuilder.Execute();
+			});
+	}
+}
+
+void FRDGBuilderRecorder::EnqueueRenderCommand(FRDGRecorderRenderCommand InRenderCommand)
+{
+	if (State == EState::Recording)
+	{
+		auto RDGCommand = [InRenderCommand](FRDGBuilder& GraphBuilder)
+			{
+				GraphBuilder.AddPass(RDG_EVENT_NAME("RDGRecordedRenderCommand"), ERDGPassFlags::NeverCull, [InRenderCommand](FRHICommandListImmediate& InRHICmdList) mutable { InRenderCommand(InRHICmdList); });
+			}; 
+		RDGCommands.Add(RDGCommand);
+	}
+	else
+	{
+		ENQUEUE_RENDER_COMMAND(RDGImmediateRenderCommand)(InRenderCommand);
+	}
+}
+
+bool FRDGBuilderRecorder::IsEmpty() const
+{
+	return RDGCommands.IsEmpty();
+}
+
+void FRDGBuilderRecorder::Clear()
+{
+	RDGCommands.Reset();
+	RDGExternalTextureAccessFinal.Reset();
+}
+
+
 #if WITH_EDITOR
 
+// ----------------------------------------------------------------------------------
+
 FString GetSharedAssetsPath(const FString& InPath)
 {
 	FString Path = InPath + TEXT("_sharedassets/");
diff --git a/Engine/Source/Runtime/Landscape/Public/GlobalMergeLegacySupportUtil.h b/Engine/Source/Runtime/Landscape/Public/GlobalMergeLegacySupportUtil.h
--- a/Engine/Source/Runtime/Landscape/Public/GlobalMergeLegacySupportUtil.h
+++ b/Engine/Source/Runtime/Landscape/Public/GlobalMergeLegacySupportUtil.h
@@ -58,7 +58,7 @@ protected:
 
 	// Used in RenderAsBlueprintBrush to cache things the same way that blueprint brushes do
 	FTransform CurrentRenderAreaWorldTransform;
-	FIntPoint CurrentRenderAreaSize;
-	FIntPoint CurrentRenderTargetSize;
+	FIntPoint CurrentRenderAreaSize = FIntPoint(ForceInitToZero);
+	FIntPoint CurrentRenderTargetSize = FIntPoint(ForceInitToZero);
 #endif
 };
diff --git a/Engine/Source/Runtime/Landscape/Public/LandscapeBlueprintBrushBase.h b/Engine/Source/Runtime/Landscape/Public/LandscapeBlueprintBrushBase.h
--- a/Engine/Source/Runtime/Landscape/Public/LandscapeBlueprintBrushBase.h
+++ b/Engine/Source/Runtime/Landscape/Public/LandscapeBlueprintBrushBase.h
@@ -80,8 +80,8 @@ protected:
 	uint32 LastRequestLayersContentUpdateFrameNumber;
 
 	FTransform CurrentRenderAreaWorldTransform;
-	FIntPoint CurrentRenderAreaSize = FIntPoint(MAX_int32, MAX_int32);
-	FIntPoint CurrentRenderTargetSize = FIntPoint(MAX_int32, MAX_int32);
+	FIntPoint CurrentRenderAreaSize = FIntPoint(ForceInitToZero);
+	FIntPoint CurrentRenderTargetSize = FIntPoint(ForceInitToZero);
 
 	bool bCaptureBoundaryNormals = false;	// HACK [chris.tchou] remove once we have a better boundary normal solution
 #endif // WITH_EDITORONLY_DATA
@@ -113,10 +113,12 @@ public:
 	LANDSCAPE_EDIT_LAYERS_BATCHED_MERGE_EXPERIMENTAL
 	LANDSCAPE_API virtual void GetRendererStateInfo(const UE::Landscape::EditLayers::FMergeContext* InMergeContext,
 		UE::Landscape::EditLayers::FEditLayerTargetTypeState& OutSupportedTargetTypeState, UE::Landscape::EditLayers::FEditLayerTargetTypeState& OutEnabledTargetTypeState, TArray<TBitArray<>>& OutRenderGroups) const override;
+	LANDSCAPE_EDIT_LAYERS_BATCHED_MERGE_EXPERIMENTAL	
+	virtual UE::Landscape::EditLayers::ERenderFlags GetRenderFlags(const UE::Landscape::EditLayers::FMergeContext* InMergeContext) const override { return UE::Landscape::EditLayers::ERenderFlags::RenderMode_Immediate; }
 	LANDSCAPE_EDIT_LAYERS_BATCHED_MERGE_EXPERIMENTAL
 	LANDSCAPE_API virtual TArray<UE::Landscape::EditLayers::FEditLayerRenderItem> GetRenderItems(const UE::Landscape::EditLayers::FMergeContext* InMergeContext) const override;
 	LANDSCAPE_EDIT_LAYERS_BATCHED_MERGE_EXPERIMENTAL
-	LANDSCAPE_API virtual void RenderLayer(ILandscapeEditLayerRenderer::FRenderParams& InRenderParams) override;
+	LANDSCAPE_API virtual void RenderLayer(ILandscapeEditLayerRenderer::FRenderParams& RenderParams, UE::Landscape::FRDGBuilderRecorder& RDGBuilderRecorder) override;
 	LANDSCAPE_EDIT_LAYERS_BATCHED_MERGE_EXPERIMENTAL
 	LANDSCAPE_API virtual FString GetEditLayerRendererDebugName() const override;
 	//~ End ILandscapeEditLayerRenderer implementation
diff --git a/Engine/Source/Runtime/Landscape/Public/LandscapeEditLayerRenderer.h b/Engine/Source/Runtime/Landscape/Public/LandscapeEditLayerRenderer.h
--- a/Engine/Source/Runtime/Landscape/Public/LandscapeEditLayerRenderer.h
+++ b/Engine/Source/Runtime/Landscape/Public/LandscapeEditLayerRenderer.h
@@ -4,6 +4,7 @@
 
 #include "CoreMinimal.h"
 #include "LandscapeEditTypes.h"
+#include "LandscapeUtils.h"
 #include "Engine/TextureRenderTarget2D.h"
 #include "RHIAccess.h"
 #include "EngineDefines.h"
@@ -17,6 +18,12 @@ class ULandscapeLayerInfoObject;
 class ULandscapeScratchRenderTarget;
 class ULandscapeEditLayerRenderer;
 class ILandscapeEditLayerRenderer;
+class FRDGBuilder;
+class FRHICommandListImmediate;
+namespace UE::Landscape
+{
+	class FRDGBuilderRecorder;
+} // namespace UE::Landscape
 
 namespace UE::Landscape::EditLayers
 {
@@ -642,10 +649,17 @@ struct FMergeRenderParams
  */
 struct FMergeRenderStep
 {
+	/** Describes what kind of operation this render step executes. Despite some of the names in there, those steps all run on the game thread but the _RenderThread ones will defer operations to the render thread. */
 	enum class EType
 	{
-		RenderLayer, // Perform the rendering of a render group on an edit layer on a given world region (i.e. in a batch)
-		SignalBatchMergeGroupDone, // Final step when rendering a render group on a given world region (i.e. in a batch) : allows to retrieve the result of the merge and do something with it (e.g. resolve the final textures)
+		Invalid,
+
+		BeginRenderCommandRecorder, // Initiates a render command recorder that will batch one or several RenderLayer_Recorded operations, runs on the game thread
+		RenderLayer_Recorded, // Performs the rendering of a render group on an edit layer on a given world region (i.e. in a batch), runs on the game thread but this is where render operation (a lambda that will run on the render thread) can be enqueued on the command recorder initiated by BeginRenderCommandRecorder
+		EndRenderCommandRecorder, // Ends the render command recorder initiated by the last BeginRenderCommandRecorder, runs on the game thread but this is where the render command executing the recorded operations from (several) RenderLayer_Recorded step(s) will be pushed (on the render thread, this is where the FRDBuilder is created and executed)
+
+		RenderLayer_Immediate, // Perform the rendering of a render group on an edit layer on a given world region (i.e. in a batch), runs on the game thread
+		SignalBatchMergeGroupDone, // Final step when rendering a render group on a given world region (i.e. in a batch) : runs on the game thread and allows to retrieve the result of the merge and do something with it (e.g. resolve the final textures)
 	};
 
 	FMergeRenderStep(EType InType, const FEditLayerRendererState& InRendererState, const TBitArray<>& InRenderGroupBitIndices, const TArrayView<ULandscapeComponent*>& InComponentsToRender)
@@ -656,7 +670,7 @@ struct FMergeRenderStep
 	{}
 
 	/** Type of operation for this step */
-	EType Type = EType::RenderLayer;
+	EType Type = EType::Invalid;
 	
 	/** Renderer state to be used this step (only when Type == EType::RenderLayer). This includes the renderer as well as its precise step (e.g. which weightmap are supported? which are enabled?) */
 	FEditLayerRendererState RendererState;
@@ -781,14 +795,15 @@ public:
 	/**
 	 * Cycle between the 3 render targets used for blending:
 	 *  Write becomes Read -> Read becomes ReadPrevious -> ReadPrevious becomes Write
-	 *	The new Write RT will be transitioned to state == InWriteAccess 
+	 *	The new Write RT will be transitioned to state == InDesiredWriteAccess (if != ERHIAccess::None)
 	 *  The new Read RT will be transitioned to state ERHIAccess::SRVMask
 	 *  The new ReadPrevious RT will stay in state ERHIAccess::SRVMask
 	 * 
-	 * @param InDesiredWriteAccess RHI state of the write RT after the operation, if specified
+	 * @param InDesiredWriteAccess RHI state of the write RT after the operation (if InDesiredWriteAccess != ERHIAccess::None)
+	 * @param RDGBuilderRecorder recorder to append operations to a single FRDGBuilder, if the the recorder in "recording" mode. In "immediate" mode, a render command will be enqueued immediately
 	 */
 	LANDSCAPE_EDIT_LAYERS_BATCHED_MERGE_EXPERIMENTAL
-	LANDSCAPE_API void CycleBlendRenderTargets(ERHIAccess InDesiredWriteAccess = ERHIAccess::None);
+	LANDSCAPE_API void CycleBlendRenderTargets(ERHIAccess InDesiredWriteAccess, UE::Landscape::FRDGBuilderRecorder& RDGBuilderRecorder);
 
 	LANDSCAPE_EDIT_LAYERS_BATCHED_MERGE_EXPERIMENTAL
 	LANDSCAPE_API ULandscapeScratchRenderTarget* GetBlendRenderTargetWrite() const;
@@ -828,7 +843,7 @@ public:
 		/** Additional info about the components that have been processed in this batch render */
 		const TArray<FComponentMergeRenderInfo> SortedComponentMergeRenderInfos;
 	};
-	void Render(TFunction<void(const FOnRenderBatchTargetGroupDoneParams&)> OnRenderBatchTargetGroupDone);
+	void Render(TFunction<void(const FOnRenderBatchTargetGroupDoneParams& /*InParams*/, UE::Landscape::FRDGBuilderRecorder& /*RDGBuilderRecorder*/)> OnRenderBatchTargetGroupDone);
 
 	LANDSCAPE_EDIT_LAYERS_BATCHED_MERGE_EXPERIMENTAL
 	inline FIntPoint GetMaxNeededResolution() const { return MaxNeededResolution; }
@@ -851,11 +866,19 @@ public:
 	LANDSCAPE_API bool IsVisualLogEnabled() const;
 #endif // ENABLE_VISUAL_LOG
 
-	/** Render the stencil render targets for each target layer in this merge for this batch */
-	void RenderValidityRenderTargets(const FMergeRenderBatch& InRenderBatch);
+	/** 
+	 * Render the stencil render targets for each target layer in this merge for this batch 
+	 * @param InRenderBatch current render batch
+	* @param RDGBuilderRecorder recorder to append operations to a single FRDGBuilder, if the the recorder in "recording" mode. In "immediate" mode, a render command will be enqueued immediately
+	 */
+	void RenderValidityRenderTargets(const FMergeRenderBatch& InRenderBatch, UE::Landscape::FRDGBuilderRecorder& RDGBuilderRecorder);
 
-	/** Duplicates the vertex data from the (sub-)sections of the batch, assuming GetBlendRenderTargetRead() is the RT that is read from and GetBlendRenderTargetWrite() the one that is written to */
-	void RenderExpandedRenderTarget(const FMergeRenderBatch& InRenderBatch);
+	/** 
+	 * Duplicates the vertex data from the (sub-)sections of the batch, assuming GetBlendRenderTargetRead() is the RT that is read from and GetBlendRenderTargetWrite() the one that is written to 
+	 * @param InRenderBatch current render batch
+	* @param RDGBuilderRecorder recorder to append operations to a single FRDGBuilder, if the the recorder in "recording" mode. In "immediate" mode, a render command will be enqueued immediately
+	 */
+	void RenderExpandedRenderTarget(const FMergeRenderBatch& InRenderBatch, UE::Landscape::FRDGBuilderRecorder& RDGBuilderRecorder);
 
 	const TBitArray<>& GetFinalTargetLayerBitIndices() const 
 	{ 
@@ -1105,6 +1128,17 @@ public:
 };
 #endif // WITH_EDITOR
 
+enum class ERenderFlags : uint8
+{
+	None = 0,
+
+	// Render mode flags
+	RenderMode_Recorded = (1 << 0), // This renderer can record its render commands into a single FRDGBuilder on the render thread (prefer this if possible), via a FRDGBuilderRecorder in "recording" mode
+	RenderMode_Immediate = (1 << 1), // This renderer enqueues its render commands immediately either via the FRDGBuilderRecorder in "immediate" mode or just enqueuing render commands the usual way
+	RenderMode_Mask = RenderMode_Recorded | RenderMode_Immediate
+};
+ENUM_CLASS_FLAGS(ERenderFlags);
+
 } //namespace UE::Landscape::EditLayers
 
 // ----------------------------------------------------------------------------------
@@ -1217,18 +1251,33 @@ public:
 		PURE_VIRTUAL(ILandscapeEditLayerRenderer::GetRenderItems, return { }; );
 
 	/**
-	 * Indicates whether the renderer actually does anything in the render phase
+	 * @param InMergeContext context containing all sorts of information related to this merge operation
+	 * @return details about how this renderer's render method is implemented 
 	*/
-	LANDSCAPE_EDIT_LAYERS_BATCHED_MERGE_EXPERIMENTAL
-	LANDSCAPE_API virtual bool CanRender() const { return true; }
+	LANDSCAPE_API virtual UE::Landscape::EditLayers::ERenderFlags GetRenderFlags(const UE::Landscape::EditLayers::FMergeContext* InMergeContext) const
+		PURE_VIRTUAL(ILandscapeEditLayerRenderer::GetRenderFlags, return UE::Landscape::EditLayers::ERenderFlags::None; );
 
 	/**
 	 * RenderLayer is where the renderer has a chance to render its content and eventually blend it with the merged result of all preceding layers
 	 * It operates on a limited set of components (depending on the size of the render batches) and on a set of target layers (e.g. multiple weightmaps).
-	 * It guarantees access to merged result from preceding layers of each target layer
+	 * It guarantees access to merged result from preceding layers of each target layer. 
+	 * 
+	 * @param RenderParams contains parameters necessary for this render (affected components, merge context, etc.)
+	 * 
+	 * @param RDGBuilderRecorder recorder to append operations to a single FRDGBuilder, if the the recorder in "recording" mode. In "immediate" mode, a render command will be enqueued immediately. 
+	 *  The recorder can be in 2 states, depending on the result from GetRenderFlags
+	 *  - ERenderFlags::RenderMode_Immediate : the recorder is in immediate mode renderer can enqueue render commands just like any other game thread-based renderer
+	 *  - ERenderFlags::RenderMode_Recorded : the function runs on the game thread but is *not* meant to enqueue render commands directly. Instead, it registers
+	 *     consecutive "render commands" via "RDG/render command" accumulated on a FRDGBuilderRecorder. This allows to coalesce several render commands onto the same FRDGBuilder (which is critical 
+	 *     for performance) while still allowing to interleave game thread-based renders (flushing the command recorder, enqueuing the game thread-based render, and starting recording on the same (or a 
+	 *     new) command recorder again...) 
+	 *     TLDR: when using the ERenderFlags::Recorded method, the user should use the command recorder to enqueue lambdas instead if enqueuing render commands directly. 
+	 *     Corollary: Any render command issued by RenderLayer will end up being pushed before the render operations recorded on the RDG command recorder (unless it is flushed), so there's no reason 
+	 *     to actually do it, unless you want those commands to run before the recorded commands
 	*/
 	LANDSCAPE_EDIT_LAYERS_BATCHED_MERGE_EXPERIMENTAL
-	LANDSCAPE_API virtual void RenderLayer(ILandscapeEditLayerRenderer::FRenderParams& InRenderParams) {}
+	LANDSCAPE_API virtual void RenderLayer(ILandscapeEditLayerRenderer::FRenderParams& RenderParams, UE::Landscape::FRDGBuilderRecorder& RDGBuilderRecorder) 
+		PURE_VIRTUAL(ILandscapeEditLayerRenderer::RenderLayer, );
 
 #endif // WITH_EDITOR
 };
diff --git a/Engine/Source/Runtime/Landscape/Public/LandscapeEditResourcesSubsystem.h b/Engine/Source/Runtime/Landscape/Public/LandscapeEditResourcesSubsystem.h
--- a/Engine/Source/Runtime/Landscape/Public/LandscapeEditResourcesSubsystem.h
+++ b/Engine/Source/Runtime/Landscape/Public/LandscapeEditResourcesSubsystem.h
@@ -18,9 +18,12 @@
 struct FResourceSizeEx;
 class UTextureRenderTarget2D;
 class UTextureRenderTarget2DArray;
+class FRDGBuilder;
+class FRHICommandListImmediate;
 
 namespace UE::Landscape 
 {
+class FRDGBuilderRecorder;
 
 struct LANDSCAPE_EDIT_LAYERS_BATCHED_MERGE_EXPERIMENTAL FScratchRenderTargetParams
 {
@@ -147,12 +150,13 @@ public:
 	class FTransitionBatcherScope : public FNoncopyable
 	{
 	public:
-		FTransitionBatcherScope();
+		FTransitionBatcherScope(UE::Landscape::FRDGBuilderRecorder& RDGBuilderRecorder);
 		~FTransitionBatcherScope();
 
 		void TransitionTo(ULandscapeScratchRenderTarget* InScratchRenderTarget, ERHIAccess InStateAfter);
 
 	private:
+		UE::Landscape::FRDGBuilderRecorder& RDGBuilderRecorder;
 		TArray<FTransitionInfo> PendingTransitions;
 	};
 
@@ -187,21 +191,45 @@ public:
 
 		ULandscapeScratchRenderTarget* SourceScratchRenderTarget = nullptr;
 	};
-
-	/** Copies the content of the texture in parameter to the scratch texture (assuming the input texture is in CopySrc state already). Transitions the scratch texture's RHIAccess */
+	
+	/**
+	 * Copies the content of the texture in parameter to the scratch texture (assuming the input texture is in CopySrc state already). Transitions the scratch texture's RHIAccess
+	 *  @param InCopyParams parameters for the texture to copy from 
+	 *  @param RDGBuilderRecorder when the recorder is in "recording" mode, the copy's render command enqueuing will be delayed to the render thread, after a FRDGBuilder has been created
+	 *   but the scratch render target's current state will still be mutated to reflect the state change on the game thread. In "immediate" mode, a render command will be enqueued immediately
+	 */
 	LANDSCAPE_EDIT_LAYERS_BATCHED_MERGE_EXPERIMENTAL
-	LANDSCAPE_API void CopyFrom(const FCopyFromTextureParams& InCopyParams);
-	/** Copies the content of the scratch texture in parameter to the scratch texture. Transitions both scratch texture's RHIAccess */
+	LANDSCAPE_API void CopyFrom(const FCopyFromTextureParams& InCopyParams, UE::Landscape::FRDGBuilderRecorder& RDGBuilderRecorder);
+
+	/**
+	 * Copies the content of the scratch texture in parameter to the scratch texture. Transitions both scratch textures' RHIAccess
+	 * @param InCopyParams parameters for the scratch texture to copy from
+	 * @param RDGBuilderRecorder when the recorder is in "recording" mode, the copy's render command enqueuing will be delayed to the render thread, after a FRDGBuilder has been created
+	*  but the scratch render targets' current state will still be mutated to reflect the state changes on the game thread. In "immediate" mode, a render command will be enqueued immediately
+	 */
 	LANDSCAPE_EDIT_LAYERS_BATCHED_MERGE_EXPERIMENTAL
-	LANDSCAPE_API void CopyFrom(const FCopyFromScratchRenderTargetParams& InCopyParams);
-	/** Perform a transition of the scratch texture's RHIAccess */
+	LANDSCAPE_API void CopyFrom(const FCopyFromScratchRenderTargetParams& InCopyParams, UE::Landscape::FRDGBuilderRecorder& RDGBuilderRecorder);
+
+	/** 
+	 * Perform a transition of the scratch texture's RHIAccess 
+	 *  @param InStateAfter desired RHIAccess of the scratch render target
+	*  @param RDGBuilderRecorder  when the recorder is in "recording" mode, the transition's render command enqueuing will not be executed (we assume the FRDGBuilder will take care of it)
+	 *   but the scratch render target's current state will still be mutated to reflect the state change on the game thread. In "immediate" mode, a render command will be enqueued immediately
+	 */
 	LANDSCAPE_EDIT_LAYERS_BATCHED_MERGE_EXPERIMENTAL
-	LANDSCAPE_API void TransitionTo(ERHIAccess InStateAfter);
-	/** Perform a clear of the scratch texture */
+	LANDSCAPE_API void TransitionTo(ERHIAccess InStateAfter, UE::Landscape::FRDGBuilderRecorder& RDGBuilderRecorder);
+
+	/** 
+	 * Perform a clear of the scratch texture 
+	 *  @param RDGBuilderRecorder  when the recorder is in "recording" mode, the clear's render command enqueuing will be delayed to the render thread, after a FRDGBuilder has been created
+	 *   but the scratch render target's current state will still be mutated to reflect the state change on the game thread. In "immediate" mode, a render command will be enqueued immediately
+	 */
 	LANDSCAPE_EDIT_LAYERS_BATCHED_MERGE_EXPERIMENTAL
-	LANDSCAPE_API void Clear();
+	LANDSCAPE_API void Clear(UE::Landscape::FRDGBuilderRecorder& RDGBuilderRecorder);
+
 	/** @return true if this scratch texture is compatible (and can therefore be used) with the requested render target params */
 	bool IsCompatibleWith(const UE::Landscape::FScratchRenderTargetParams& InParams) const;
+
 	/** @return true if the scratch texture is currently in use (it has been requested but not released) */
 	bool IsInUse() const { return bIsInUse; }
 
diff --git a/Engine/Source/Runtime/Landscape/Public/LandscapeUtils.h b/Engine/Source/Runtime/Landscape/Public/LandscapeUtils.h
--- a/Engine/Source/Runtime/Landscape/Public/LandscapeUtils.h
+++ b/Engine/Source/Runtime/Landscape/Public/LandscapeUtils.h
@@ -8,13 +8,18 @@
 #include "AssetRegistry/AssetData.h"
 #include "LandscapeTextureHash.h"
 #include "Materials/Material.h"
+#include "RenderGraphDefinitions.h"
+#include "RenderGraphEvent.h"
 
-class ULevel;
 enum EShaderPlatform : uint16;
+class FRDGBuilder;
+class FRDGEventName;
 class ULandscapeComponent;
 class ULandscapeLayerInfoObject;
-class UTexture2D;
+class ULevel;
 class UTexture;
+class UTexture2D;
+class FTextureResource;
 enum class ELandscapeToolTargetType : uint8;
 enum class ELandscapeToolTargetTypeFlags : uint8;
 
@@ -32,8 +37,198 @@ LANDSCAPE_API ELandscapeToolTargetTypeFlags GetLandscapeToolTargetTypeAsFlags(EL
 LANDSCAPE_API ELandscapeToolTargetType GetLandscapeToolTargetTypeSingleFlagAsType(ELandscapeToolTargetTypeFlags InSingleFlag);
 LANDSCAPE_API FString GetLandscapeToolTargetTypeFlagsAsString(ELandscapeToolTargetTypeFlags InTargetTypeFlags);
 
+
+// ----------------------------------------------------------------------------------
+
+/** This struct is usually meant to be allocated on the game thread (where there's no FRDGBuilder, which is Render Thread-only) and allows to queue successive operations (lambdas) onto a single render command,
+ *   sharing the same FRDGBuilder. This allows to sequence a list of RDG passes from the game thread and makes it possible to interleave render thread operations (in a single render command) with game thread-initiated
+ *   render commands. 
+ *   Here's an extensive use case :
+ * 
+ *   FRDGBuilderRecorder Recorder;
+ * 
+ *   // Start recording commands to a single graph builder : 
+ *   Recorder.StartRecording(); 
+ * 
+ *   Recorder.EnqueueRDGCommand([](FRDGBuilder* GraphBuilder) { GraphBuilder->AddPass(); }); // Append a Pass_A
+ *   ENQUEUE_RENDER_COMMAND(...)  // Push Render_Command_A (immediately)
+ *   Recorder.EnqueueRDGCommand([](FRDGBuilder* GraphBuilder) { GraphBuilder->AddPass(); }, { { SomeTexture, ERHIAccess::RTV } }); // Append a Pass_B and inform of the final state of a given texture to prevent the RDG from auto-transitioning to SRVMask at the end
+ *   Recorder.EnqueueRenderCommand([](FRHICommandListImmediate& InRHICmdList) mutable { [...]; }); // Append a Render_Command_B to the graph builder
+ *   Recorder.EnqueueRDGCommand([](FRDGBuilder* GraphBuilder) { GraphBuilder->AddPass(); }, { { }SomeTexture, ERHIAccess::CopySrc } }); // Append a Pass_C and inform of the final state of a given texture to prevent the RDG from auto-transitioning to SRVMask at the end
+ *   
+ *   // Stop recording and issue a render command with all that's been recorded so far
+ *   Recorder.StopRecordingAndFlush(RDG_EVENT_NAME("Pass ABC")); 
+ * 
+ *   // Enqueue some game-thread render commands (e.g. BP render) either directly or via the recorder in immediate mode :
+ *   ENQUEUE_RENDER_COMMAND(...)  // Render_Command_C
+ *   Recorder.EnqueueRDGCommand([](FRDGBuilder* GraphBuilder) { GraphBuilder->AddPass(); }, { { SomeTexture, ERHIAccess::CopyDst } }); // Append a Pass_D and inform of the final state of a given texture to prevent the RDG from auto-transitioning to SRVMask at the end
+ *   Recorder.EnqueueRenderCommand([RHIOperation](FRHICommandListImmediate& InRHICmdList) mutable { [...]; }); // Append a Render_Command_D
+ *   
+ *   // Start recording commands again to a new single graph builder : 
+ *   Recorder.StartRecording(); 
+ *   Recorder.EnqueueRDGCommand([](FRDGBuilder* GraphBuilder) { GraphBuilder->AddPass(); }); // Append a Pass_E
+ *  
+ *   // Stop recording and issue a render command with all that's been recorded so far
+ *   Recorder.StopRecordingAndFlush(RDG_EVENT_NAME("Pass E"));
+ * 
+ *   --> Will yield the following sequence on the render thread: 
+ *   + Render_Command_A
+ *   + Render Command "Pass ABC" 
+ *     +- (RDGBuilder_0)
+ *     +- RDGBuilder_0.Pass_A
+ *     +- RDGBuilder_0.Pass_B
+ *     +- RDGBuilder_0.LambdaPass (Render_Command_B)
+ *     +- RDGBuilder_0.Pass_C
+ *     +- RDGBuilder_0.SetTextureAccessFinal(OutputTexture, ERHIAccess::CopySrc); // Only the final state recorded for a given texture is set
+ *     +- RDGBuilder_0.Execute
+ *   + Render_Command_C
+ *   + Render Command 
+ *     +- (RDGBuilder_1)
+ *     +- RDGBuilder_1.Pass_D
+ *     +- RDGBuilder_1.SetTextureAccessFinal(OutputTexture, ERHIAccess::CopyDst); 
+ *     +- RDGBuilder_1.Execute
+ *   + Render_Command_D
+ *   + Render Command "Pass E"
+ *     +- (RDGBuilder_2)
+ *       +- RDGBuilder_2.Pass_E
+ *       +- RDGBuilder_2.Execute
+ */
+class FRDGBuilderRecorder final
+{
+public:
+	using FRDGRecorderRDGCommand = TFunction<void(FRDGBuilder& /*GraphBuilder*/)>;
+	using FRDGRecorderRenderCommand = TFunction<void(FRHICommandListImmediate& /*InRHICmdList*/)>;
+	struct FRDGExternalTextureAccessFinal
+	{
+		FTextureResource* TextureResource = nullptr;
+		ERHIAccess Access = ERHIAccess::None;
+	};
+
+	LANDSCAPE_API FRDGBuilderRecorder() = default;
+	LANDSCAPE_API ~FRDGBuilderRecorder();
+	
+	enum class EState
+	{
+		Immediate, // In immediate mode, any command that is enqueued will be pushed to the render thread immediately (effectively acting like a ENQUEUE_RENDER_COMMAND) 
+		Recording, // In recording mode, any command that is enqueued will be deferred to the render thread (effectively acting like a ENQUEUE_RENDER_COMMAND) 
+	};
+
+	inline EState GetState() const { return State; }
+	inline bool IsRecording() const { return State == EState::Recording; }
+
+	/**
+	* Starts recording commands
+	 */
+	LANDSCAPE_API void StartRecording();
+
+	/**
+	* Stops recording commands. A call to Flush is needed to ensure any pending command is flushed to the render thread (use StopRecordingAndFlush to do both)
+	*/
+	LANDSCAPE_API void StopRecording();
+
+	/**
+	 * Stops recording commands and flushes them to the render thread. Expects the recorder to be in Recording mode and changes it to Immediate mode.
+	 * @param EventName RDG event name to use on the render command's FRDGBuilder
+	*/
+	LANDSCAPE_API void StopRecordingAndFlush(FRDGEventName&& EventName);
+
+	/**
+	 * Flushes any pending command to the render thread
+	 * @param EventName RDG event name to use on the render command's FRDGBuilder
+	*/
+	LANDSCAPE_API void Flush(FRDGEventName&& EventName);
+
+	/**
+	 * Records a FRDGRecorderRDGCommand to execute when registering passes to the single FRDGBuilder when in Recording mode or pushes it immediately to the render thread when in Immediate mode
+	 * 
+	 * @param InRDGCommand Lambda to execute on the render thread
+	 * @param InRDGTextureAccessFinalList (optional) : list of external texture with the RHIAccess they should have when executing the FRDGBuilder. This is to prevent the RDG from auto-transitioning to SRVMask at the end
+	 */
+	LANDSCAPE_API void EnqueueRDGCommand(FRDGRecorderRDGCommand InRDGCommand, TConstArrayView<FRDGExternalTextureAccessFinal> InRDGExternalTextureAccessFinalList = {});
+
+	/**
+	 * Records a FRDGRecorderRenderCommand to execute when registering passes to the single FRDGBuilder when in Recording mode or pushes it immediately to the render thread when in Immediate mode
+	 *
+	 * @param InRenderCommand Lambda to execute on the render thread
+	 */
+	LANDSCAPE_API void EnqueueRenderCommand(FRDGRecorderRenderCommand InRenderCommand);
+
+	/**
+	 * @return true if there's no command currently recorded
+	 */
+	LANDSCAPE_API bool IsEmpty() const;
+
+	/** 
+	 * Cancels all recorder operations. This must be used if the FRDGEventRenderRecorder is "cancelled" (i.e. its sequence of operations is not flushed to a render command).
+	 *  Otherwise, there will be an assert in ~FRDGBuilderRecorder().
+	 */
+	LANDSCAPE_API void Clear();
+
+private:
+	EState State = EState::Immediate;
+
+	// List of callbacks to call on the render thread after the render command was initiated
+	TArray<FRDGRecorderRDGCommand> RDGCommands;
+
+	// Map of textures and the RHI access they should have when leaving the FRDGBuilder :
+	TMap<FTextureResource*, ERHIAccess> RDGExternalTextureAccessFinal;
+
+#if RDG_EVENTS
+public:
+	/** 
+	 * Scope object meant to insert a RDG event in the RDG operations, as if it was inserted on the render thread on a FRDGBuilder. 
+	 *  Use RDG_RENDER_COMMAND_RECORDER_BREADCRUMB_EVENT to create one.
+	 */
+	class FScopedBreadcrumbEvent final
+	{
+	public:
+		template<typename... TArgs>
+		FScopedBreadcrumbEvent(FRDGBuilderRecorder& InRecorder, TArgs&&... InArgs)
+			: Recorder(&InRecorder)
+			, RDGEvent(MakeShared<TOptional<TRDGEventScopeGuard<FRDGScope_RHI>>>())
+		{
+			//  We use a shared ptr to create a TOptional<TRDGEventScopeGuard> immediately, then capture this shared ptr in these additional operations' lambdas (so that the object
+			//  continues to live until the closing tag operation)
+			Recorder->EnqueueRDGCommand([
+				RDGEvent = RDGEvent, 
+				Values = std::make_tuple(std::forward<TArgs>(InArgs) ...)
+				](FRDGBuilder& GraphBuilder) mutable
+				{
+					std::apply([RDGEvent, &GraphBuilder](auto&&... InnerArgs)
+						{
+							// Allocate the TOptional<TRDGEventScopeGuard> now in order to insert the tag : 
+							RDGEvent->Emplace(GraphBuilder, ERDGScopeFlags::None, FRHIBreadcrumbData(__FILE__, __LINE__, TStatId(), NAME_None), std::forward<TArgs>(InnerArgs)...);
+						}, Values);
+				});
+		}
+
+		~FScopedBreadcrumbEvent()
+		{
+			Recorder->EnqueueRDGCommand([RDGEvent = RDGEvent](FRDGBuilder& GraphBuilder)
+				{
+					// Reset the TOptional in order to delete the TRDGEventScopeGuard, which will remove the tag : 
+					RDGEvent->Reset();
+				});
+		}
+
+	private:
+		FRDGBuilderRecorder* Recorder = nullptr;
+		TSharedPtr<TOptional<TRDGEventScopeGuard<FRDGScope_RHI>>> RDGEvent;
+	};
+#endif // RDG_EVENTS
+};
+
+#if RDG_EVENTS
+#define RDG_RENDER_COMMAND_RECORDER_BREADCRUMB_EVENT(Recorder, Format, ...) FRDGBuilderRecorder::FScopedBreadcrumbEvent ANONYMOUS_VARIABLE(BreadcrumbEvent)(Recorder, RDG_EVENT_NAME(Format, ##__VA_ARGS__))
+#else // RDG_EVENTS
+#define RDG_RENDER_COMMAND_RECORDER_BREADCRUMB_EVENT(Recorder, Format, ...) do { } while(0)
+#endif // !RDG_EVENTS
+
+
 #if WITH_EDITOR
 
+// ----------------------------------------------------------------------------------
+
 struct FTextureCopyRequest
 {
 	UTexture2D* Source = nullptr;
