Landscape edit layers batched merge : renamed RenderGroup to TargetLayerGroup to avoid confusion with upcoming change

diff --git a/Engine/Plugins/Experimental/LandscapePatch/Source/LandscapePatch/Private/LandscapeCircleHeightPatch.cpp b/Engine/Plugins/Experimental/LandscapePatch/Source/LandscapePatch/Private/LandscapeCircleHeightPatch.cpp
--- a/Engine/Plugins/Experimental/LandscapePatch/Source/LandscapePatch/Private/LandscapeCircleHeightPatch.cpp
+++ b/Engine/Plugins/Experimental/LandscapePatch/Source/LandscapePatch/Private/LandscapeCircleHeightPatch.cpp
@@ -96,12 +96,12 @@ void ULandscapeCircleHeightPatch::RenderLayer(FRenderParams& InRenderParams, UE:
 	UTextureRenderTarget2DArray* TextureArray = LandscapeRT->TryGetRenderTarget2DArray();
 	check(TextureArray != nullptr);
 
-	int32 NumLayers = InRenderParams.RenderGroupTargetLayerInfos.Num();
+	int32 NumLayers = InRenderParams.TargetLayerGroupLayerInfos.Num();
 	check(TextureArray->Slices == NumLayers);
 
 	for (int32 LayerIndex = 0; LayerIndex < NumLayers; ++LayerIndex)
 	{
-		if (UE::Landscape::IsVisibilityLayer(InRenderParams.RenderGroupTargetLayerInfos[LayerIndex]))
+		if (UE::Landscape::IsVisibilityLayer(InRenderParams.TargetLayerGroupLayerInfos[LayerIndex]))
 		{
 			ApplyCirclePatch(RDGBuilderRecorder, /*bIsVisibilityLayer = */true, TextureArray->GetResource(),
 				LayerIndex, InRenderParams.RenderAreaSectionRect.Size(), HeightmapCoordsToWorld);
@@ -205,7 +205,7 @@ void ULandscapeCircleHeightPatch::ApplyCirclePatch(UE::Landscape::FRDGBuilderRec
 void ULandscapeCircleHeightPatch::GetRendererStateInfo(const UE::Landscape::EditLayers::FMergeContext* InMergeContext,
 	FEditLayerTargetTypeState& OutSupported,
 	FEditLayerTargetTypeState& OutEnabled,
-	TArray<TBitArray<>>& OutRenderGroups) const
+	TArray<TBitArray<>>& OutTargetLayerGroups) const
 {
 	OutSupported.AddTargetType(bEditVisibility ? ELandscapeToolTargetType::Visibility
 		: ELandscapeToolTargetType::Heightmap);
diff --git a/Engine/Plugins/Experimental/LandscapePatch/Source/LandscapePatch/Private/LandscapePatchManager.cpp b/Engine/Plugins/Experimental/LandscapePatch/Source/LandscapePatch/Private/LandscapePatchManager.cpp
--- a/Engine/Plugins/Experimental/LandscapePatch/Source/LandscapePatch/Private/LandscapePatchManager.cpp
+++ b/Engine/Plugins/Experimental/LandscapePatch/Source/LandscapePatch/Private/LandscapePatchManager.cpp
@@ -236,8 +236,8 @@ TArray<UE::Landscape::EditLayers::FEditLayerRendererState> ALandscapePatchManage
 	// Inherit the base brush's enabled state and intersect it with all patches' state, so that the heightmap/weightmap/visibility toggle at the brush's level applies to them too : 
 	FEditLayerTargetTypeState BaseBrushSupportedState(InMergeContext);
 	FEditLayerTargetTypeState BaseBrushEnabledState(InMergeContext);
-	TArray<TBitArray<>> BaseBrushRenderGroups;
-	Super::GetRendererStateInfo(InMergeContext, BaseBrushSupportedState, BaseBrushEnabledState, BaseBrushRenderGroups);
+	TArray<TBitArray<>> BaseTargetLayerGroups;
+	Super::GetRendererStateInfo(InMergeContext, BaseBrushSupportedState, BaseBrushEnabledState, BaseTargetLayerGroups);
 
 	TArray<FEditLayerRendererState> RendererStates;
 	RendererStates.Reserve(PatchComponents.Num());
diff --git a/Engine/Plugins/Experimental/LandscapePatch/Source/LandscapePatch/Private/LandscapeTexturePatch.cpp b/Engine/Plugins/Experimental/LandscapePatch/Source/LandscapePatch/Private/LandscapeTexturePatch.cpp
--- a/Engine/Plugins/Experimental/LandscapePatch/Source/LandscapePatch/Private/LandscapeTexturePatch.cpp
+++ b/Engine/Plugins/Experimental/LandscapePatch/Source/LandscapePatch/Private/LandscapeTexturePatch.cpp
@@ -182,20 +182,20 @@ void ULandscapeTexturePatch::RenderLayer(FRenderParams& InRenderParams, UE::Land
 	UTextureRenderTarget2DArray* TextureArray = LandscapeScratchRT->TryGetRenderTarget2DArray();
 	check(TextureArray != nullptr);
 	
-	int32 NumTargetLayersInRenderGroup = InRenderParams.RenderGroupTargetLayerNames.Num();
-	check(TextureArray->Slices == NumTargetLayersInRenderGroup);
+	int32 NumTargetLayersInGroup = InRenderParams.TargetLayerGroupLayerNames.Num();
+	check(TextureArray->Slices == NumTargetLayersInGroup);
 
-	for (int32 TargetLayerIndexInRenderGroup = 0; TargetLayerIndexInRenderGroup < NumTargetLayersInRenderGroup; ++TargetLayerIndexInRenderGroup)
+	for (int32 TargetLayerIndexInGroup = 0; TargetLayerIndexInGroup < NumTargetLayersInGroup; ++TargetLayerIndexInGroup)
 	{
-		bool bIsVisibilityLayer = ensure(TargetLayerIndexInRenderGroup < InRenderParams.RenderGroupTargetLayerInfos.Num())
-			&& UE::Landscape::IsVisibilityLayer(InRenderParams.RenderGroupTargetLayerInfos[TargetLayerIndexInRenderGroup]);
+		bool bIsVisibilityLayer = ensure(TargetLayerIndexInGroup < InRenderParams.TargetLayerGroupLayerInfos.Num())
+			&& UE::Landscape::IsVisibilityLayer(InRenderParams.TargetLayerGroupLayerInfos[TargetLayerIndexInGroup]);
 		
 		// Try to find the weight patch
 		ULandscapeWeightPatchTextureInfo* WeightPatchInfo = nullptr;
 		for (const TObjectPtr<ULandscapeWeightPatchTextureInfo>& WeightPatchEntry : WeightPatches)
 		{
 			if ((bIsVisibilityLayer && WeightPatchEntry->bEditVisibilityLayer)
-				|| (WeightPatchEntry->WeightmapLayerName == InRenderParams.RenderGroupTargetLayerNames[TargetLayerIndexInRenderGroup]))
+				|| (WeightPatchEntry->WeightmapLayerName == InRenderParams.TargetLayerGroupLayerNames[TargetLayerIndexInGroup]))
 			{
 				WeightPatchInfo = WeightPatchEntry;
 				break;
@@ -213,14 +213,14 @@ void ULandscapeTexturePatch::RenderLayer(FRenderParams& InRenderParams, UE::Land
 			WeightPatchInfo->bReinitializeOnNextRender = false;
 			checkf(!RDGBuilderRecorder.IsRecording(), TEXT("We should be using ERenderFlags::RenderMode_Immediate when reinitializing weight"));
 			ReinitializeWeightPatch(WeightPatchInfo, TextureArray->GetResource(),
-				FIntPoint(TextureArray->SizeX, TextureArray->SizeY), TargetLayerIndexInRenderGroup, LandscapeHeightmapToWorld);
+				FIntPoint(TextureArray->SizeX, TextureArray->SizeY), TargetLayerIndexInGroup, LandscapeHeightmapToWorld);
 		}
 		else
 		{
 			checkf(RDGBuilderRecorder.IsRecording(), TEXT("We should be using ERenderFlags::RenderMode_Recorded in the typical case"));
 			ApplyToWeightmap(RDGBuilderRecorder, WeightPatchInfo,
 				TextureArray->GetResource(),
-				TargetLayerIndexInRenderGroup, 
+				TargetLayerIndexInGroup,
 				InRenderParams.RenderAreaSectionRect.Size(), 
 				GetHeightmapToWorld(InRenderParams.RenderAreaWorldTransform));
 		}
@@ -471,48 +471,48 @@ void ULandscapeTexturePatch::ApplyToWeightmap(UE::Landscape::FRDGBuilderRecorder
 
 	auto RDGCommand =
 		[InMergedLandscapeTextureResource, LandscapeTextureSliceIndex, ShaderParamsToCopy, Patch, DestinationBounds](FRDGBuilder& GraphBuilder)
-		{
-			TRACE_CPUPROFILER_EVENT_SCOPE(LandscapeTextureWeightPatch_Render);
+	{
+		TRACE_CPUPROFILER_EVENT_SCOPE(LandscapeTextureWeightPatch_Render);
 
 			TRefCountPtr<IPooledRenderTarget> DestinationRenderTarget = CreateRenderTarget(InMergedLandscapeTextureResource->GetTextureRHI(), TEXT("LandscapeTextureWeightPatchOutput"));
-			FRDGTextureRef DestinationTexture = GraphBuilder.RegisterExternalTexture(DestinationRenderTarget);
-
-			// Make a copy of the portion of our weightmap input that we're writing to so that we can 
-			// read and write at the same time (needed for blending)
-			FRDGTextureDesc InputCopyDescription = DestinationTexture->Desc;
-			InputCopyDescription.Dimension = ETextureDimension::Texture2D;
-			InputCopyDescription.ArraySize = 1;
-			InputCopyDescription.NumMips = 1;
-			InputCopyDescription.Extent = DestinationBounds.Size();
-			FRDGTextureRef InputCopy = GraphBuilder.CreateTexture(InputCopyDescription, TEXT("LandscapeTextureWeightPatchInputCopy"));
-
-			FRHICopyTextureInfo CopyTextureInfo;
-			CopyTextureInfo.SourceMipIndex = 0;
-			CopyTextureInfo.NumMips = 1;
-			CopyTextureInfo.SourceSliceIndex = LandscapeTextureSliceIndex;
-			CopyTextureInfo.NumSlices = 1;
-			CopyTextureInfo.SourcePosition = FIntVector(DestinationBounds.Min.X, DestinationBounds.Min.Y, 0);
-			CopyTextureInfo.Size = FIntVector(InputCopyDescription.Extent.X, InputCopyDescription.Extent.Y, 0);
-			AddCopyTexturePass(GraphBuilder, DestinationTexture, InputCopy, CopyTextureInfo);
-
-			FApplyLandscapeTextureWeightPatchPS::FParameters* ShaderParams =
-				GraphBuilder.AllocParameters<FApplyLandscapeTextureWeightPatchPS::FParameters>();
-			*ShaderParams = ShaderParamsToCopy;
-
-			TRefCountPtr<IPooledRenderTarget> PatchRenderTarget = CreateRenderTarget(Patch->GetTexture2DRHI(), TEXT("LandscapeTextureWeightPatch"));
-			FRDGTextureRef PatchTexture = GraphBuilder.RegisterExternalTexture(PatchRenderTarget);
-			FRDGTextureSRVRef PatchSRV = GraphBuilder.CreateSRV(FRDGTextureSRVDesc::CreateForMipLevel(PatchTexture, 0));
-			ShaderParams->InWeightPatch = PatchSRV;
-			ShaderParams->InWeightPatchSampler = TStaticSamplerState<SF_Bilinear, AM_Clamp, AM_Clamp>::GetRHI();
-
-			FRDGTextureSRVRef InputCopySRV = GraphBuilder.CreateSRV(FRDGTextureSRVDesc::CreateForMipLevel(InputCopy, 0));
-			ShaderParams->InSourceWeightmap = InputCopySRV;
-			ShaderParams->InSourceWeightmapCoordOffset = DestinationBounds.Min;
-
-			ShaderParams->RenderTargets[0] = FRenderTargetBinding(DestinationTexture, ERenderTargetLoadAction::ENoAction,
-				/*InMipIndex = */0, LandscapeTextureSliceIndex);
-
-			FApplyLandscapeTextureWeightPatchPS::AddToRenderGraph(GraphBuilder, ShaderParams, DestinationBounds);
+		FRDGTextureRef DestinationTexture = GraphBuilder.RegisterExternalTexture(DestinationRenderTarget);
+
+		// Make a copy of the portion of our weightmap input that we're writing to so that we can 
+		// read and write at the same time (needed for blending)
+		FRDGTextureDesc InputCopyDescription = DestinationTexture->Desc;
+		InputCopyDescription.Dimension = ETextureDimension::Texture2D;
+		InputCopyDescription.ArraySize = 1;
+		InputCopyDescription.NumMips = 1;
+		InputCopyDescription.Extent = DestinationBounds.Size();
+		FRDGTextureRef InputCopy = GraphBuilder.CreateTexture(InputCopyDescription, TEXT("LandscapeTextureWeightPatchInputCopy"));
+
+		FRHICopyTextureInfo CopyTextureInfo;
+		CopyTextureInfo.SourceMipIndex = 0;
+		CopyTextureInfo.NumMips = 1;
+		CopyTextureInfo.SourceSliceIndex = LandscapeTextureSliceIndex;
+		CopyTextureInfo.NumSlices = 1;
+		CopyTextureInfo.SourcePosition = FIntVector(DestinationBounds.Min.X, DestinationBounds.Min.Y, 0);
+		CopyTextureInfo.Size = FIntVector(InputCopyDescription.Extent.X, InputCopyDescription.Extent.Y, 0);
+		AddCopyTexturePass(GraphBuilder, DestinationTexture, InputCopy, CopyTextureInfo);
+
+		FApplyLandscapeTextureWeightPatchPS::FParameters* ShaderParams =
+			GraphBuilder.AllocParameters<FApplyLandscapeTextureWeightPatchPS::FParameters>();
+		*ShaderParams = ShaderParamsToCopy;
+
+		TRefCountPtr<IPooledRenderTarget> PatchRenderTarget = CreateRenderTarget(Patch->GetTexture2DRHI(), TEXT("LandscapeTextureWeightPatch"));
+		FRDGTextureRef PatchTexture = GraphBuilder.RegisterExternalTexture(PatchRenderTarget);
+		FRDGTextureSRVRef PatchSRV = GraphBuilder.CreateSRV(FRDGTextureSRVDesc::CreateForMipLevel(PatchTexture, 0));
+		ShaderParams->InWeightPatch = PatchSRV;
+		ShaderParams->InWeightPatchSampler = TStaticSamplerState<SF_Bilinear, AM_Clamp, AM_Clamp>::GetRHI();
+
+		FRDGTextureSRVRef InputCopySRV = GraphBuilder.CreateSRV(FRDGTextureSRVDesc::CreateForMipLevel(InputCopy, 0));
+		ShaderParams->InSourceWeightmap = InputCopySRV;
+		ShaderParams->InSourceWeightmapCoordOffset = DestinationBounds.Min;
+
+		ShaderParams->RenderTargets[0] = FRenderTargetBinding(DestinationTexture, ERenderTargetLoadAction::ENoAction,
+			/*InMipIndex = */0, LandscapeTextureSliceIndex);
+
+		FApplyLandscapeTextureWeightPatchPS::AddToRenderGraph(GraphBuilder, ShaderParams, DestinationBounds);
 		};
 
 	// We need to specify the final state of the external texture to prevent the graph builder from transitioning it to SRVMask :
@@ -2037,29 +2037,29 @@ void ULandscapeTexturePatch::SetEditVisibilityLayer(const FName& InWeightmapLaye
 #if WITH_EDITOR
 void ULandscapeTexturePatch::GetRendererStateInfo(const UE::Landscape::EditLayers::FMergeContext* InMergeContext,
 	FEditLayerTargetTypeState& OutSupported, FEditLayerTargetTypeState& OutEnabled,
-	TArray<TBitArray<>>& OutRenderGroups) const
+	TArray<TBitArray<>>& OutTargetLayerGroups) const
 {
 	if (InMergeContext->IsHeightmapMerge())
-{
-	if (CanAffectHeightmap())
 	{
-		OutSupported.AddTargetType(ELandscapeToolTargetType::Heightmap);
-	}
+		if (CanAffectHeightmap())
+		{
+			OutSupported.AddTargetType(ELandscapeToolTargetType::Heightmap);
+		}
 	}
 	else
 	{
-	for (const TObjectPtr<ULandscapeWeightPatchTextureInfo>& WeightPatch : WeightPatches)
-	{
-		if (IsValid(WeightPatch) && WeightPatchCanRender(*WeightPatch))
+		for (const TObjectPtr<ULandscapeWeightPatchTextureInfo>& WeightPatch : WeightPatches)
 		{
-			if (WeightPatch->bEditVisibilityLayer)
+			if (IsValid(WeightPatch) && WeightPatchCanRender(*WeightPatch))
 			{
-				OutSupported.AddTargetType(ELandscapeToolTargetType::Visibility);
-			}
+				if (WeightPatch->bEditVisibilityLayer)
+				{
+					OutSupported.AddTargetType(ELandscapeToolTargetType::Visibility);
+				}
 				else if (InMergeContext->IsValidTargetLayerName(WeightPatch->WeightmapLayerName))
-			{
+				{
 					const int32 TargetLayerIndex = InMergeContext->GetTargetLayerIndexForNameChecked(WeightPatch->WeightmapLayerName);
-				OutSupported.AddTargetType(ELandscapeToolTargetType::Weightmap);
+					OutSupported.AddTargetType(ELandscapeToolTargetType::Weightmap);
 					OutSupported.AddWeightmap(TargetLayerIndex);
 				}
 			}
@@ -2090,19 +2090,19 @@ TArray<UE::Landscape::EditLayers::FEditLayerRenderItem> ULandscapeTexturePatch::
 
 	if (InMergeContext->IsHeightmapMerge())
 	{
-	if (CanAffectHeightmap())
-	{
+		if (CanAffectHeightmap())
+		{
 			FEditLayerTargetTypeState TargetInfo(InMergeContext, ELandscapeToolTargetTypeFlags::Heightmap);
-		FEditLayerRenderItem Item(TargetInfo, InputWorldArea, OutputWorldArea, false);
-		AffectedAreas.Add(Item);
-	}
+			FEditLayerRenderItem Item(TargetInfo, InputWorldArea, OutputWorldArea, false);
+			AffectedAreas.Add(Item);
+		}
 	}
 	else
 	{
-	for (const TObjectPtr<ULandscapeWeightPatchTextureInfo>& WeightPatch : WeightPatches)
-	{
-		if (IsValid(WeightPatch) && WeightPatchCanRender(*WeightPatch))
+		for (const TObjectPtr<ULandscapeWeightPatchTextureInfo>& WeightPatch : WeightPatches)
 		{
+			if (IsValid(WeightPatch) && WeightPatchCanRender(*WeightPatch))
+			{
 				FEditLayerTargetTypeState TargetInfo(InMergeContext);
 				if (WeightPatch->bEditVisibilityLayer)
 				{
@@ -2115,10 +2115,10 @@ TArray<UE::Landscape::EditLayers::FEditLayerRenderItem> ULandscapeTexturePatch::
 					TargetInfo.AddWeightmap(TargetLayerIndex);
 				}
 				FEditLayerRenderItem Item(TargetInfo, InputWorldArea, OutputWorldArea, /*bInModifyExistingWeightmapsOnly = */false);
-			AffectedAreas.Add(Item);
+				AffectedAreas.Add(Item);
+			}
 		}
 	}
-	}
 
 	return AffectedAreas;
 }
diff --git a/Engine/Plugins/Experimental/LandscapePatch/Source/LandscapePatch/Public/LandscapeCircleHeightPatch.h b/Engine/Plugins/Experimental/LandscapePatch/Source/LandscapePatch/Public/LandscapeCircleHeightPatch.h
--- a/Engine/Plugins/Experimental/LandscapePatch/Source/LandscapePatch/Public/LandscapeCircleHeightPatch.h
+++ b/Engine/Plugins/Experimental/LandscapePatch/Source/LandscapePatch/Public/LandscapeCircleHeightPatch.h
@@ -37,7 +37,7 @@ public:
 	virtual void GetRendererStateInfo(const UE::Landscape::EditLayers::FMergeContext* InMergeContext,
 		FEditLayerTargetTypeState& OutSupportedTargetTypeState, 
 		FEditLayerTargetTypeState& OutEnabledTargetTypeState, 
-		TArray<TBitArray<>>& OutRenderGroups) const override;
+		TArray<TBitArray<>>& OutTargetLayerGroups) const override;
 	virtual FString GetEditLayerRendererDebugName() const override;
 	virtual TArray<FEditLayerRenderItem> GetRenderItems(const UE::Landscape::EditLayers::FMergeContext* InMergeContext) const override;
 	virtual UE::Landscape::EditLayers::ERenderFlags GetRenderFlags(const UE::Landscape::EditLayers::FMergeContext* InMergeContext) const override { return UE::Landscape::EditLayers::ERenderFlags::RenderMode_Recorded; }
diff --git a/Engine/Plugins/Experimental/LandscapePatch/Source/LandscapePatch/Public/LandscapePatchManager.h b/Engine/Plugins/Experimental/LandscapePatch/Source/LandscapePatch/Public/LandscapePatchManager.h
--- a/Engine/Plugins/Experimental/LandscapePatch/Source/LandscapePatch/Public/LandscapePatchManager.h
+++ b/Engine/Plugins/Experimental/LandscapePatch/Source/LandscapePatch/Public/LandscapePatchManager.h
@@ -46,7 +46,7 @@ public:
 	virtual void GetRendererStateInfo(const UE::Landscape::EditLayers::FMergeContext* InMergeContext,
 		UE::Landscape::EditLayers::FEditLayerTargetTypeState& OutSupportedTargetTypeState,
 		UE::Landscape::EditLayers::FEditLayerTargetTypeState& OutEnabledTargetTypeState,
-		TArray<TBitArray<>>& OutRenderGroups) const override {}
+		TArray<TBitArray<>>& OutTargetLayerGroups) const override {}
 	virtual TArray<UE::Landscape::EditLayers::FEditLayerRenderItem> GetRenderItems(const UE::Landscape::EditLayers::FMergeContext* InMergeContext) const override { return {}; }
 	virtual void RenderLayer(ILandscapeEditLayerRenderer::FRenderParams& InRenderParams, UE::Landscape::FRDGBuilderRecorder& RDGBuilderRecorder) override { };
 #endif // WITH_EDITOR
diff --git a/Engine/Plugins/Experimental/LandscapePatch/Source/LandscapePatch/Public/LandscapeTexturePatch.h b/Engine/Plugins/Experimental/LandscapePatch/Source/LandscapePatch/Public/LandscapeTexturePatch.h
--- a/Engine/Plugins/Experimental/LandscapePatch/Source/LandscapePatch/Public/LandscapeTexturePatch.h
+++ b/Engine/Plugins/Experimental/LandscapePatch/Source/LandscapePatch/Public/LandscapeTexturePatch.h
@@ -235,7 +235,7 @@ public:
 	virtual void GetRendererStateInfo(const UE::Landscape::EditLayers::FMergeContext* InMergeContext,
 		FEditLayerTargetTypeState& OutSupportedTargetTypeState,
 		FEditLayerTargetTypeState& OutEnabledTargetTypeState,
-		TArray<TBitArray<>>& OutRenderGroups) const override;
+		TArray<TBitArray<>>& OutTargetLayerGroups) const override;
 	virtual FString GetEditLayerRendererDebugName() const override;
 	virtual TArray<FEditLayerRenderItem> GetRenderItems(const UE::Landscape::EditLayers::FMergeContext* InMergeContext) const override;
 
diff --git a/Engine/Shaders/Private/Landscape/LandscapeEditLayersWeightmaps.usf b/Engine/Shaders/Private/Landscape/LandscapeEditLayersWeightmaps.usf
--- a/Engine/Shaders/Private/Landscape/LandscapeEditLayersWeightmaps.usf
+++ b/Engine/Shaders/Private/Landscape/LandscapeEditLayersWeightmaps.usf
@@ -30,8 +30,8 @@ uint InPaintLayerIndex; // Index of the paint layer being processed
 uint InEditLayerPaintLayerBlendMode; // See EEditLayerWeightmapBlendMode. Paint layer's blend mode, which is per-edit layer and per-paint layer and can therefore differ from one edit layer to another
 float InEditLayerAlpha; // Global alpha value of the edit layer currently being merged
 
-Texture2DArray<float> InCurrentEditLayerWeightmaps; // The paint layers of the current edit layer to merge (1 slice per weightmap in the render group)
-Texture2DArray<float> InPreviousEditLayersWeightmaps; // The result from the merge of all prior edit layers (1 slice per weightmap in the render group)
+Texture2DArray<float> InCurrentEditLayerWeightmaps; // The paint layers of the current edit layer to merge (1 slice per weightmap in the target layer group)
+Texture2DArray<float> InPreviousEditLayersWeightmaps; // The result from the merge of all prior edit layers (1 slice per weightmap in the target layer group)
 // Array that contains per-edit layer / per-paint layer information (e.g. the paint layer's blend mode, which is per-edit layer and can therefore differ from one edit layer to another)
 #endif // MERGE_EDIT_LAYER
 
@@ -48,7 +48,7 @@ uint InPaintLayerIndex; // Index of the paint layer being processed
 uint InNumPaintLayers; // Number of paint layers (i.e. == InPaintLayerInfos.Num())
 // Array that contains per-paint layer information (e.g. paint layer's flags)
 StructuredBuffer<FWeightmapPaintLayerInfo> InPaintLayerInfos;
-Texture2DArray<float> InCurrentEditLayerWeightmaps; // The texture to horizontally blend (1 slice per weightmap in the render group)
+Texture2DArray<float> InCurrentEditLayerWeightmaps; // The texture to horizontally blend (1 slice per weightmap in the target layer group)
 #endif // PERFORM_LEGACY_WEIGHT_BLENDING
 
 #if PACK_WEIGHTMAP
diff --git a/Engine/Source/Runtime/Landscape/Classes/LandscapeEditLayer.h b/Engine/Source/Runtime/Landscape/Classes/LandscapeEditLayer.h
--- a/Engine/Source/Runtime/Landscape/Classes/LandscapeEditLayer.h
+++ b/Engine/Source/Runtime/Landscape/Classes/LandscapeEditLayer.h
@@ -195,7 +195,7 @@ public:
 	//~ Begin ILandscapeEditLayerRenderer implementation
 	LANDSCAPE_EDIT_LAYERS_BATCHED_MERGE_EXPERIMENTAL
 	LANDSCAPE_API virtual void GetRendererStateInfo(const UE::Landscape::EditLayers::FMergeContext* InMergeContext,
-		UE::Landscape::EditLayers::FEditLayerTargetTypeState& OutSupportedTargetTypeState, UE::Landscape::EditLayers::FEditLayerTargetTypeState& OutEnabledTargetTypeState, TArray<TBitArray<>>& OutRenderGroups) const override;
+		UE::Landscape::EditLayers::FEditLayerTargetTypeState& OutSupportedTargetTypeState, UE::Landscape::EditLayers::FEditLayerTargetTypeState& OutEnabledTargetTypeState, TArray<TBitArray<>>& OutTargetLayerGroups) const override;
 	LANDSCAPE_EDIT_LAYERS_BATCHED_MERGE_EXPERIMENTAL
 	virtual UE::Landscape::EditLayers::ERenderFlags GetRenderFlags(const UE::Landscape::EditLayers::FMergeContext* InMergeContext) const override { return UE::Landscape::EditLayers::ERenderFlags::RenderMode_Recorded; }
 	LANDSCAPE_EDIT_LAYERS_BATCHED_MERGE_EXPERIMENTAL
diff --git a/Engine/Source/Runtime/Landscape/Private/LandscapeBlueprintBrushBase.cpp b/Engine/Source/Runtime/Landscape/Private/LandscapeBlueprintBrushBase.cpp
--- a/Engine/Source/Runtime/Landscape/Private/LandscapeBlueprintBrushBase.cpp
+++ b/Engine/Source/Runtime/Landscape/Private/LandscapeBlueprintBrushBase.cpp
@@ -222,7 +222,7 @@ void ALandscapeBlueprintBrushBase::SetCanAffectVisibilityLayer(bool bInCanAffect
 #if WITH_EDITOR
 
 void ALandscapeBlueprintBrushBase::GetRendererStateInfo(const UE::Landscape::EditLayers::FMergeContext* InMergeContext,
-	UE::Landscape::EditLayers::FEditLayerTargetTypeState& OutSupportedTargetTypeState, UE::Landscape::EditLayers::FEditLayerTargetTypeState& OutEnabledTargetTypeState, TArray<TBitArray<>>& OutRenderGroups) const
+	UE::Landscape::EditLayers::FEditLayerTargetTypeState& OutSupportedTargetTypeState, UE::Landscape::EditLayers::FEditLayerTargetTypeState& OutEnabledTargetTypeState, TArray<TBitArray<>>& OutTargetLayerGroups) const
 {
 	// What can the brush do?
 	if (CanAffectHeightmap())
@@ -254,7 +254,7 @@ void ALandscapeBlueprintBrushBase::GetRendererStateInfo(const UE::Landscape::Edi
 
 	// Mark which weightmap is supported/enabled : 
 	if (!InMergeContext->IsHeightmapMerge() && CanAffectWeightmap())
-		{
+	{
 		InMergeContext->ForEachValidTargetLayer(
 			[InMergeContext, this, &OutSupportedTargetTypeState, &OutEnabledTargetTypeState](int32 InTargetLayerIndex, const FName& InTargetLayerName, ULandscapeLayerInfoObject* InWeightmapLayerInfo)
 			{
@@ -264,8 +264,8 @@ void ALandscapeBlueprintBrushBase::GetRendererStateInfo(const UE::Landscape::Edi
 					if (AffectsWeightmapLayer(InTargetLayerName))
 					{
 						OutEnabledTargetTypeState.AddWeightmap(InTargetLayerIndex);
-			}
-		}
+					}
+				}
 				return true;
 			});
 	}
@@ -277,8 +277,8 @@ TArray<UE::Landscape::EditLayers::FEditLayerRenderItem> ALandscapeBlueprintBrush
 
 	FEditLayerTargetTypeState SupportedTargetTypeState(InMergeContext);
 	FEditLayerTargetTypeState EnabledTargetTypeState(InMergeContext);
-	TArray<TBitArray<>> DummyRenderGroups;
-	GetRendererStateInfo(InMergeContext, SupportedTargetTypeState, EnabledTargetTypeState, DummyRenderGroups);
+	TArray<TBitArray<>> DummyTargetLayerGroups;
+	GetRendererStateInfo(InMergeContext, SupportedTargetTypeState, EnabledTargetTypeState, DummyTargetLayerGroups);
 
 	// By default, for landscape BP brushes, we use FInputWorldArea::EType::Infinite, to indicate they can only reliably work when applied globally on the entire landscape
 	//  This allows full backwards-compatibility but will prevent landscapes from benefiting from batched merge. Users will be able to indicate their brush works in a local fashion
@@ -307,7 +307,7 @@ void ALandscapeBlueprintBrushBase::RenderLayer(ILandscapeEditLayerRenderer::FRen
 
 	const bool bIsHeightmapMerge = RenderParams.MergeRenderContext->IsHeightmapMerge();
 	const TArray<FName> EnabledWeightmaps = RenderParams.RendererState.GetActiveTargetWeightmaps();
-
+	
 	ULandscapeEditResourcesSubsystem* LandscapeEditResourcesSubsystem = GEngine->GetEngineSubsystem<ULandscapeEditResourcesSubsystem>();
 	check(LandscapeEditResourcesSubsystem != nullptr);
 	ULandscapeScratchRenderTarget* CurrentLayerReadRT2D = CurrentLayerReadRT;
@@ -332,21 +332,21 @@ void ALandscapeBlueprintBrushBase::RenderLayer(ILandscapeEditLayerRenderer::FRen
 	// The original texture array will be accessed as ERHIAccess::CopySrc all along : 
 	CurrentLayerReadRT->TransitionTo(ERHIAccess::CopySrc, RDGBuilderRecorder);
 
-	const int32 NumTargetLayersInRenderGroup = RenderParams.RenderGroupTargetLayerNames.Num();
-	for (int32 TargetLayerIndexInRenderGroup = 0; TargetLayerIndexInRenderGroup < NumTargetLayersInRenderGroup; ++TargetLayerIndexInRenderGroup)
+	const int32 NumTargetLayersInGroup = RenderParams.TargetLayerGroupLayerNames.Num();
+	for (int32 TargetLayerIndexInGroup = 0; TargetLayerIndexInGroup < NumTargetLayersInGroup; ++TargetLayerIndexInGroup)
 	{
-		const FName TargetLayerName = RenderParams.RenderGroupTargetLayerNames[TargetLayerIndexInRenderGroup];
+		const FName TargetLayerName = RenderParams.TargetLayerGroupLayerNames[TargetLayerIndexInGroup];
 		RHI_BREADCRUMB_EVENT_GAMETHREAD("Render %s", TargetLayerName);
 
 		// If necessary, copy from the source render target (slice) to the scratch render target 2D : 
-		RHI_BREADCRUMB_EVENT_GAMETHREAD("Copy Source (slice %i) -> %s", TargetLayerIndexInRenderGroup, CurrentLayerReadRT2D->GetDebugName());
+		RHI_BREADCRUMB_EVENT_GAMETHREAD("Copy Source (slice %i) -> %s", TargetLayerIndexInGroup, CurrentLayerReadRT2D->GetDebugName());
 
 		{
 			ULandscapeScratchRenderTarget::FCopyFromScratchRenderTargetParams CopyParams(CurrentLayerReadRT);
 			// It's important to copy only the needed size, because CurrentLayerReadRT2D comes from ScratchTexture, whose size may be larger than RenderParams.RenderAreaSectionRect)
 			CopyParams.CopySize = RenderParams.RenderAreaSectionRect.Size();
 			// Copy from the proper slice in the texture array if any : 
-			CopyParams.SourceSliceIndex = TargetLayerIndexInRenderGroup;
+			CopyParams.SourceSliceIndex = TargetLayerIndexInGroup;
 			CurrentLayerReadRT2D->CopyFrom(CopyParams, RDGBuilderRecorder);
 		}
 		CurrentLayerReadRT2D->TransitionTo(ERHIAccess::SRVMask, RDGBuilderRecorder);
@@ -358,7 +358,7 @@ void ALandscapeBlueprintBrushBase::RenderLayer(ILandscapeEditLayerRenderer::FRen
 		// If the BP brush failed to render, we still need to passthrough from the read RT to the write RT in order not to lose what has been merged so far : 
 		UTextureRenderTarget2D* OutputRT2D = ReadRT2D;
 
-		// Only render the target layer if it's effectively enabled for this merge : it's possible there are target layers in the render group that we don't support or are not enabled so we have to 
+		// Only render the target layer if it's effectively enabled for this merge : it's possible there are target layers in the group that we don't support or are not enabled so we have to 
 		//  do the validation here first :
 		if (bIsHeightmapMerge || EnabledWeightmaps.Contains(TargetLayerName))
 		{
@@ -376,7 +376,7 @@ void ALandscapeBlueprintBrushBase::RenderLayer(ILandscapeEditLayerRenderer::FRen
 
 		// Resolve back to the write RT 
 		{
-			RHI_BREADCRUMB_EVENT_GAMETHREAD("Copy BP Render Result -> %s (slice %i)", WriteRT->GetDebugName(), TargetLayerIndexInRenderGroup);
+			RHI_BREADCRUMB_EVENT_GAMETHREAD("Copy BP Render Result -> %s (slice %i)", WriteRT->GetDebugName(), TargetLayerIndexInGroup);
 
 			// The RT returned by the brush is in SRV state so we need a transition: 
 			ENQUEUE_RENDER_COMMAND(TransitionToCopySrc)([Resource = OutputRT2D->GetResource()](FRHICommandListImmediate& InRHICmdList) mutable
@@ -388,7 +388,7 @@ void ALandscapeBlueprintBrushBase::RenderLayer(ILandscapeEditLayerRenderer::FRen
 			// It's important to copy only the needed size, because OutputRT2D is currently sized after ReadRT2D, which can be a render target larger than it needs to be (because 
 			//  it can come from ScratchTexture, whose size may be larger than RenderParams.RenderAreaSectionRect)
 			CopyParams.CopySize = RenderParams.RenderAreaSectionRect.Size();
-			CopyParams.DestSliceIndex = TargetLayerIndexInRenderGroup;
+			CopyParams.DestSliceIndex = TargetLayerIndexInGroup;
 			WriteRT->CopyFrom(CopyParams, RDGBuilderRecorder);
 
 			// It's also expected we leave the RT return by the BP as SRV :
diff --git a/Engine/Source/Runtime/Landscape/Private/LandscapeEditLayerRenderer.cpp b/Engine/Source/Runtime/Landscape/Private/LandscapeEditLayerRenderer.cpp
--- a/Engine/Source/Runtime/Landscape/Private/LandscapeEditLayerRenderer.cpp
+++ b/Engine/Source/Runtime/Landscape/Private/LandscapeEditLayerRenderer.cpp
@@ -425,12 +425,12 @@ void FEditLayerTargetTypeState::SetTargetTypeMask(ELandscapeToolTargetTypeFlags
 		const int32 VisibilityTargetLayerIndex = MergeContext->GetVisibilityTargetLayerIndex();
 		if (VisibilityTargetLayerIndex != INDEX_NONE)
 		{
-		if (EnumHasAnyFlags(TargetTypeMask, ELandscapeToolTargetTypeFlags::Visibility))
-		{
+			if (EnumHasAnyFlags(TargetTypeMask, ELandscapeToolTargetTypeFlags::Visibility))
+			{
 				AddWeightmap(VisibilityTargetLayerIndex);
-		}
+			}
 			else if (WeightmapTargetLayerBitIndices[VisibilityTargetLayerIndex])
-		{
+			{
 				RemoveWeightmap(VisibilityTargetLayerIndex);
 			}
 		}
@@ -549,22 +549,22 @@ FEditLayerRendererState::FEditLayerRendererState(const FMergeContext* InMergeCon
 	, EnabledTargetTypeState(InMergeContext)
 	, ActiveTargetTypeState(InMergeContext)
 {
-	Renderer->GetRendererStateInfo(InMergeContext, SupportedTargetTypeState, EnabledTargetTypeState, RenderGroups);
+	Renderer->GetRendererStateInfo(InMergeContext, SupportedTargetTypeState, EnabledTargetTypeState, TargetLayerGroups);
 
-	// Make sure that each supported weightmap belongs to one render group and one only. For those that are in no render group, put them in their own group, that simply means this renderer 
+	// Make sure that each supported weightmap belongs to one target layer group and one only. For those that are in no target layer group, put them in their own group, that simply means this renderer 
 	//  can render them without requesting the presence of other target layers (e.g. no weight-blending)
 	InMergeContext->ForEachTargetLayer(SupportedTargetTypeState.GetActiveWeightmapBitIndices(), 
 		[InMergeContext, InRenderer, this](int32 InTargetLayerIndex, const FName& InTargetLayerName, ULandscapeLayerInfoObject* InWeightmapLayerInfo)
-	{
-			const int32 RenderGroupCount = Algo::CountIf(RenderGroups, [InTargetLayerIndex](const TBitArray<>& InRenderGroup) { return InRenderGroup[InTargetLayerIndex]; });
-		checkf(RenderGroupCount < 2, TEXT("Target layer %s belongs to more than 1 render group in edit layer renderer %s. This is forbidden: in the end, it must belong to 1 and 1 only."),
-				*InTargetLayerName.ToString(), *InRenderer->GetEditLayerRendererDebugName());
-		if (RenderGroupCount == 0)
 		{
-				TBitArray<> RenderGroup = InMergeContext->BuildTargetLayerBitIndices(false);
-				RenderGroup[InTargetLayerIndex] = true;
-				RenderGroups.Add(RenderGroup);
-}
+			const int32 TargetLayerGroupCount = Algo::CountIf(TargetLayerGroups, [InTargetLayerIndex](const TBitArray<>& InTargetLayerGroup) { return InTargetLayerGroup[InTargetLayerIndex]; });
+			checkf(TargetLayerGroupCount < 2, TEXT("Target layer %s belongs to more than 1 target layer group in edit layer renderer %s. This is forbidden: in the end, it must belong to 1 and 1 only."),
+				*InTargetLayerName.ToString(), *InRenderer->GetEditLayerRendererDebugName());
+			if (TargetLayerGroupCount == 0)
+			{
+				TBitArray<> TargetLayerGroup = InMergeContext->BuildTargetLayerBitIndices(false);
+				TargetLayerGroup[InTargetLayerIndex] = true;
+				TargetLayerGroups.Add(TargetLayerGroup);
+			}
 			return true;
 		});
 
@@ -608,7 +608,7 @@ void FEditLayerRendererState::EnableWeightmap(ELandscapeToolTargetType InTargetT
 		UpdateActiveTargetTypeState();
 	}
 }
-	
+
 void FEditLayerRendererState::EnableWeightmapChecked(ELandscapeToolTargetType InTargetType, const FName& InWeightmapLayerName)
 {
 	EnableWeightmap(InTargetType, (InWeightmapLayerName != NAME_None) ? MergeContext->GetTargetLayerIndexForNameChecked(InWeightmapLayerName) : INDEX_NONE);
@@ -1023,32 +1023,32 @@ void FMergeRenderContext::RenderValidityRenderTargets(const FMergeRenderBatch& I
 
 	auto RDGCommand =
 		[TexturesAndRects = MoveTemp(TexturesAndRects)](FRDGBuilder& GraphBuilder)
+	{
+		for (const FTextureAndRects& TextureAndRects : TexturesAndRects)
 		{
-			for (const FTextureAndRects& TextureAndRects : TexturesAndRects)
-			{
-				FRDGBufferRef RectBuffer = CreateUploadBuffer(GraphBuilder, TEXT("MarkValidityRects"), MakeArrayView(TextureAndRects.Rects));
-				FRDGBufferSRVRef RectBufferSRV = GraphBuilder.CreateSRV(FRDGBufferSRVDesc(RectBuffer, PF_R32G32B32A32_UINT));
-				FRDGTextureRef OutputTexture = GraphBuilder.RegisterExternalTexture(CreateRenderTarget(TextureAndRects.TextureResource->GetTexture2DRHI(), TEXT("ValidityMask")));
-
-				FMarkValidityPSParameters* PassParameters = GraphBuilder.AllocParameters<FMarkValidityPSParameters>();
-				PassParameters->RenderTargets[0] = FRenderTargetBinding(OutputTexture, ERenderTargetLoadAction::ELoad);
-				FGlobalShaderMap* ShaderMap = GetGlobalShaderMap(GMaxRHIFeatureLevel);
-				TShaderRef<FMarkValidityPS> PixelShader = ShaderMap->GetShader<FMarkValidityPS>();
-
-				FPixelShaderUtils::AddRasterizeToRectsPass<FMarkValidityPS>(GraphBuilder,
-					ShaderMap,
-					RDG_EVENT_NAME("MarkValidity(%s) -> %s", *TextureAndRects.TargetLayerName.ToString(), *TextureAndRects.TextureDebugName),
-					PixelShader,
-					PassParameters,
-					/*ViewportSize = */OutputTexture->Desc.Extent,
-					RectBufferSRV,
-					TextureAndRects.Rects.Num(),
-					/*BlendState = */ nullptr,
-					/*RasterizerState = */ nullptr,
-					/*DepthStencilState = */ nullptr,
-					/*StencilRef = */ 0,
-					/*TextureSize = */ OutputTexture->Desc.Extent);
-			}
+			FRDGBufferRef RectBuffer = CreateUploadBuffer(GraphBuilder, TEXT("MarkValidityRects"), MakeArrayView(TextureAndRects.Rects));
+			FRDGBufferSRVRef RectBufferSRV = GraphBuilder.CreateSRV(FRDGBufferSRVDesc(RectBuffer, PF_R32G32B32A32_UINT));
+			FRDGTextureRef OutputTexture = GraphBuilder.RegisterExternalTexture(CreateRenderTarget(TextureAndRects.TextureResource->GetTexture2DRHI(), TEXT("ValidityMask")));
+
+			FMarkValidityPSParameters* PassParameters = GraphBuilder.AllocParameters<FMarkValidityPSParameters>();
+			PassParameters->RenderTargets[0] = FRenderTargetBinding(OutputTexture, ERenderTargetLoadAction::ELoad);
+			FGlobalShaderMap* ShaderMap = GetGlobalShaderMap(GMaxRHIFeatureLevel);
+			TShaderRef<FMarkValidityPS> PixelShader = ShaderMap->GetShader<FMarkValidityPS>();
+
+			FPixelShaderUtils::AddRasterizeToRectsPass<FMarkValidityPS>(GraphBuilder,
+				ShaderMap,
+				RDG_EVENT_NAME("MarkValidity(%s) -> %s", *TextureAndRects.TargetLayerName.ToString(), *TextureAndRects.TextureDebugName),
+				PixelShader,
+				PassParameters,
+				/*ViewportSize = */OutputTexture->Desc.Extent,
+				RectBufferSRV,
+				TextureAndRects.Rects.Num(),
+				/*BlendState = */ nullptr,
+				/*RasterizerState = */ nullptr,
+				/*DepthStencilState = */ nullptr,
+				/*StencilRef = */ 0,
+				/*TextureSize = */ OutputTexture->Desc.Extent);
+		}
 		};
 
 	// We need to specify the final state of the external textures to prevent the graph builder from transitioning them to SRVMask :
@@ -1056,7 +1056,7 @@ void FMergeRenderContext::RenderValidityRenderTargets(const FMergeRenderBatch& I
 	Algo::Transform(TexturesAndRects, RDGExternalTextureAccessFinalList, [](const FTextureAndRects& InTextureAndRects) -> FRDGBuilderRecorder::FRDGExternalTextureAccessFinal
 		{ 
 			return { InTextureAndRects.TextureResource, ERHIAccess::RTV };
-		});
+	});
 	RDGBuilderRecorder.EnqueueRDGCommand(RDGCommand, RDGExternalTextureAccessFinalList);
 }
 
@@ -1087,46 +1087,46 @@ void FMergeRenderContext::RenderExpandedRenderTarget(const FMergeRenderBatch& In
 		, SourceResource = ReadRT->GetRenderTarget2D()->GetResource()
 		, SceneInterface] (FRDGBuilder& GraphBuilder)
 	{
-			FRDGBufferRef RectBuffer = CreateUploadBuffer(GraphBuilder, TEXT("ExpandRects"), MakeArrayView(DestinationRects));
-			FRDGBufferSRVRef RectBufferSRV = GraphBuilder.CreateSRV(FRDGBufferSRVDesc(RectBuffer, PF_R32G32B32A32_UINT));
+		FRDGBufferRef RectBuffer = CreateUploadBuffer(GraphBuilder, TEXT("ExpandRects"), MakeArrayView(DestinationRects));
+		FRDGBufferSRVRef RectBufferSRV = GraphBuilder.CreateSRV(FRDGBufferSRVDesc(RectBuffer, PF_R32G32B32A32_UINT));
 
-			FRDGBufferRef RectUVBuffer = CreateUploadBuffer(GraphBuilder, TEXT("ExpandRectsUVs"), MakeArrayView(SourceRects));
-			FRDGBufferSRVRef RectUVBufferSRV = GraphBuilder.CreateSRV(FRDGBufferSRVDesc(RectUVBuffer, PF_R32G32B32A32_UINT));
+		FRDGBufferRef RectUVBuffer = CreateUploadBuffer(GraphBuilder, TEXT("ExpandRectsUVs"), MakeArrayView(SourceRects));
+		FRDGBufferSRVRef RectUVBufferSRV = GraphBuilder.CreateSRV(FRDGBufferSRVDesc(RectUVBuffer, PF_R32G32B32A32_UINT));
 
-			FRDGTextureRef OutputTexture = GraphBuilder.RegisterExternalTexture(CreateRenderTarget(OutputResource->GetTexture2DRHI(), TEXT("OutputTexture")));
-			FRDGTextureRef SourceTexture = GraphBuilder.RegisterExternalTexture(CreateRenderTarget(SourceResource->GetTexture2DRHI(), TEXT("SourceTexture")));
+		FRDGTextureRef OutputTexture = GraphBuilder.RegisterExternalTexture(CreateRenderTarget(OutputResource->GetTexture2DRHI(), TEXT("OutputTexture")));
+		FRDGTextureRef SourceTexture = GraphBuilder.RegisterExternalTexture(CreateRenderTarget(SourceResource->GetTexture2DRHI(), TEXT("SourceTexture")));
 
-			// TODO [jonathan.bard] this is just a rhi validation error for unoptimized shaders... once validation is made to not issue those errors, we can remove this
-			// Create a SceneView to please the shader bindings, but it's unused in practice 
-			FSceneViewFamilyContext ViewFamily(FSceneViewFamily::ConstructionValues(nullptr, SceneInterface, FEngineShowFlags(ESFIM_Game)).SetTime(FGameTime::GetTimeSinceAppStart()));
-			FSceneViewInitOptions ViewInitOptions;
-			ViewInitOptions.ViewFamily = &ViewFamily;
-			ViewInitOptions.SetViewRectangle(FIntRect(0, 0, 1, 1)); // Use a dummy rect to avoid a check(slow)
+		// TODO [jonathan.bard] this is just a rhi validation error for unoptimized shaders... once validation is made to not issue those errors, we can remove this
+		// Create a SceneView to please the shader bindings, but it's unused in practice 
+		FSceneViewFamilyContext ViewFamily(FSceneViewFamily::ConstructionValues(nullptr, SceneInterface, FEngineShowFlags(ESFIM_Game)).SetTime(FGameTime::GetTimeSinceAppStart()));
+		FSceneViewInitOptions ViewInitOptions;
+		ViewInitOptions.ViewFamily = &ViewFamily;
+		ViewInitOptions.SetViewRectangle(FIntRect(0, 0, 1, 1)); // Use a dummy rect to avoid a check(slow)
 			GetRendererModule().CreateAndInitSingleView(GraphBuilder.RHICmdList, &ViewFamily, &ViewInitOptions);
-			const FSceneView* View = ViewFamily.Views[0];
-
-			FCopyQuadsPSParameters* PassParameters = GraphBuilder.AllocParameters<FCopyQuadsPSParameters>();
-			PassParameters->RenderTargets[0] = FRenderTargetBinding(OutputTexture, ERenderTargetLoadAction::ELoad);
-			PassParameters->PS.View = View->ViewUniformBuffer;
-			PassParameters->PS.InSourceTexture = SourceTexture;
-
-			FGlobalShaderMap* ShaderMap = GetGlobalShaderMap(GMaxRHIFeatureLevel);
-			TShaderRef<FCopyQuadsPS> PixelShader = ShaderMap->GetShader<FCopyQuadsPS>();
-
-			FPixelShaderUtils::AddRasterizeToRectsPass<FCopyQuadsPS>(GraphBuilder,
-				ShaderMap,
-				RDG_EVENT_NAME("CopyQuadsPS"),
-				PixelShader,
-				PassParameters,
-				/*ViewportSize = */OutputTexture->Desc.Extent,
-				RectBufferSRV,
-				DestinationRects.Num(),
-				/*BlendState = */nullptr,
-				/*RasterizerState = */nullptr,
-				/*DepthStencilState = */nullptr,
-				/*StencilRef = */0,
-				/*TextureSize = */SourceTexture->Desc.Extent,
-				RectUVBufferSRV);
+		const FSceneView* View = ViewFamily.Views[0];
+
+		FCopyQuadsPSParameters* PassParameters = GraphBuilder.AllocParameters<FCopyQuadsPSParameters>();
+		PassParameters->RenderTargets[0] = FRenderTargetBinding(OutputTexture, ERenderTargetLoadAction::ELoad);
+		PassParameters->PS.View = View->ViewUniformBuffer;
+		PassParameters->PS.InSourceTexture = SourceTexture;
+
+		FGlobalShaderMap* ShaderMap = GetGlobalShaderMap(GMaxRHIFeatureLevel);
+		TShaderRef<FCopyQuadsPS> PixelShader = ShaderMap->GetShader<FCopyQuadsPS>();
+
+		FPixelShaderUtils::AddRasterizeToRectsPass<FCopyQuadsPS>(GraphBuilder,
+			ShaderMap,
+			RDG_EVENT_NAME("CopyQuadsPS"),
+			PixelShader,
+			PassParameters,
+			/*ViewportSize = */OutputTexture->Desc.Extent,
+			RectBufferSRV,
+			DestinationRects.Num(),
+			/*BlendState = */nullptr,
+			/*RasterizerState = */nullptr,
+			/*DepthStencilState = */nullptr,
+			/*StencilRef = */0,
+			/*TextureSize = */SourceTexture->Desc.Extent,
+			RectUVBufferSRV);
 		};
 
 	// We need to specify the final state of the external textures to prevent the graph builder from transitioning them to SRVMask (even those that end up as SRVMask at the end of this command,
@@ -1201,8 +1201,8 @@ void FMergeRenderContext::Render(TFunction<void(const FOnRenderBatchTargetGroupD
 		{
 			const FMergeRenderStep& RenderStep = RenderBatch.RenderSteps[RenderStepIndex];
 
-			TArray<FName> RenderGroupTargetLayerNames;
-			TArray<ULandscapeLayerInfoObject*> RenderGroupTargetLayerInfos;
+			TArray<FName> TargetLayerGroupLayerNames;
+			TArray<ULandscapeLayerInfoObject*> TargetLayerGroupLayerInfos;
 			TArray<FComponentMergeRenderInfo> SortedComponentMergeRenderInfos;
 
 			// Compute some data for the actual render steps : 
@@ -1210,26 +1210,26 @@ void FMergeRenderContext::Render(TFunction<void(const FOnRenderBatchTargetGroupD
 				|| (RenderStep.Type == FMergeRenderStep::EType::RenderLayer_Recorded)
 				|| (RenderStep.Type == FMergeRenderStep::EType::SignalBatchMergeGroupDone))
 			{
-				RenderGroupTargetLayerNames = ConvertTargetLayerBitIndicesToNames(RenderStep.RenderGroupBitIndices);
-				RenderGroupTargetLayerInfos = bIsHeightmapMerge ? TArray<ULandscapeLayerInfoObject*> { nullptr } : ConvertTargetLayerBitIndicesToLayerInfos(RenderStep.RenderGroupBitIndices);
-
-				// Compute all necessary info about the components affected by this renderer at this step
-				SortedComponentMergeRenderInfos.Reserve(RenderStep.ComponentsToRender.Num());
-				Algo::Transform(RenderStep.ComponentsToRender, SortedComponentMergeRenderInfos, [MinComponentKey = RenderBatch.MinComponentKey](ULandscapeComponent* InComponent)
-				{
-					FComponentMergeRenderInfo ComponentMergeRenderInfo;
-					ComponentMergeRenderInfo.Component = InComponent;
-
-					const FIntPoint ComponentKey = InComponent->GetComponentKey();
-					const FIntPoint LocalComponentKey = ComponentKey - MinComponentKey;
-					check((LocalComponentKey.X >= 0) && (LocalComponentKey.Y >= 0));
-					ComponentMergeRenderInfo.ComponentKeyInRenderArea = LocalComponentKey;
-					// Area in the render target for this component : 
-					ComponentMergeRenderInfo.ComponentRegionInRenderArea = FIntRect(LocalComponentKey * InComponent->ComponentSizeQuads, (LocalComponentKey + 1) * InComponent->ComponentSizeQuads);
-
-					return ComponentMergeRenderInfo;
-				});
-				SortedComponentMergeRenderInfos.Sort();
+				TargetLayerGroupLayerNames = ConvertTargetLayerBitIndicesToNames(RenderStep.TargetLayerGroupBitIndices);
+				TargetLayerGroupLayerInfos = bIsHeightmapMerge ? TArray<ULandscapeLayerInfoObject*> { nullptr } : ConvertTargetLayerBitIndicesToLayerInfos(RenderStep.TargetLayerGroupBitIndices);
+
+			    // Compute all necessary info about the components affected by this renderer at this step
+			    SortedComponentMergeRenderInfos.Reserve(RenderStep.ComponentsToRender.Num());
+			    Algo::Transform(RenderStep.ComponentsToRender, SortedComponentMergeRenderInfos, [MinComponentKey = RenderBatch.MinComponentKey](ULandscapeComponent* InComponent)
+			    {
+				    FComponentMergeRenderInfo ComponentMergeRenderInfo;
+				    ComponentMergeRenderInfo.Component = InComponent;
+    
+				    const FIntPoint ComponentKey = InComponent->GetComponentKey();
+				    const FIntPoint LocalComponentKey = ComponentKey - MinComponentKey;
+				    check((LocalComponentKey.X >= 0) && (LocalComponentKey.Y >= 0));
+				    ComponentMergeRenderInfo.ComponentKeyInRenderArea = LocalComponentKey;
+				    // Area in the render target for this component : 
+				    ComponentMergeRenderInfo.ComponentRegionInRenderArea = FIntRect(LocalComponentKey * InComponent->ComponentSizeQuads, (LocalComponentKey + 1) * InComponent->ComponentSizeQuads);
+    
+				    return ComponentMergeRenderInfo;
+			    });
+			    SortedComponentMergeRenderInfos.Sort();
 			}
 
 			switch (RenderStep.Type)
@@ -1264,73 +1264,73 @@ void FMergeRenderContext::Render(TFunction<void(const FOnRenderBatchTargetGroupD
 			}
 			case FMergeRenderStep::EType::RenderLayer_Immediate:
 			case FMergeRenderStep::EType::RenderLayer_Recorded:
-			{
-				TRACE_CPUPROFILER_EVENT_SCOPE(RenderLayer);
-
-				ILandscapeEditLayerRenderer* Renderer = RenderStep.RendererState.GetRenderer();
-				check((Renderer != nullptr) && EnumHasAnyFlags(Renderer->GetRenderFlags(this), ERenderFlags::RenderMode_Mask));
-
-				// TODO[jonathan.bard] offset the world transform to account for the half-pixel offset?
-				//RenderParams.RenderAreaWorldTransform = FTransform(LandscapeTransform.GetRotation(), LandscapeTransform.GetTranslation() + FVector(RenderBatch.SectionRect.Min) /** (double)ComponentSizeQuads - FVector(0.5, 0.5, 0)*/, LandscapeTransform.GetScale3D());
-
-				FString RenderStepProfilingEventName = FString::Format(TEXT("Step [{0}] ({1}): Render {2}"), { RenderLayerStepIndex, *ConvertTargetLayerNamesToString(RenderGroupTargetLayerNames), *Renderer->GetEditLayerRendererDebugName() });
-
-				// TODO [jonathan.bard] : this is more of a Batch world transform/section rect at the moment. Shall we have a RenderAreaWorldTransform/RenderAreaSectionRect in FRenderParams and a BatchRenderAreaWorldTransform in FMergeRenderBatch?
-				//  because currently the old BP brushes work with FMergeRenderBatch data (i.e. 1 transform for the batch and a section rect for the entire batch) but eventually, renderers might be interested in just their Render step context, 
-				//  that is : 1 matrix corresponding to the bottom left corner of their list of components to render?
-				FTransform RenderAreaWorldTransform = FTransform(FVector(RenderBatch.SectionRect.Min)) * LandscapeTransform;
-				FIntRect RenderAreaSectionRect = RenderBatch.SectionRect;
-
-				// Drop some visual cues to help understand how each renderer is applied :
-				UE_IFVLOG(
-					if (IsVisualLogEnabled() && bVisualLogMergeProcess)
-					{
-						FTransform RenderAreaWorldTransformForVisLog = FTransform(FVector(RenderBatch.SectionRect.Min)) * LandscapeWorldTransformForVisLog;
-						IncrementVisualLogOffset();
-						UE_VLOG_LOCATION(Landscape, LogLandscape, Log, ComputeVisualLogTransform(RenderAreaWorldTransformForVisLog).GetTranslation(), 10.0f, FColor::Red, TEXT("%s"), *RenderStepProfilingEventName);
-						UE_VLOG_WIREOBOX(Landscape, LogLandscape, Log, FBox(FVector(RenderBatch.SectionRect.Min) - FVector(0.5, 0.5, 0.0), FVector(RenderBatch.SectionRect.Max) - FVector(0.5, 0.5, 0.0)),
-							ComputeVisualLogTransform(LandscapeWorldTransformForVisLog).ToMatrixWithScale(), FColor::White, TEXT(""));
-
-						// Draw each component's bounds rendered by this renderer : 
-						for (const FComponentMergeRenderInfo& ComponentMergeRenderInfo : SortedComponentMergeRenderInfos)
-						{
-							UE_VLOG_WIREOBOX(Landscape, LogLandscape, Log, FBox(FVector(ComponentMergeRenderInfo.ComponentRegionInRenderArea.Min), FVector(ComponentMergeRenderInfo.ComponentRegionInRenderArea.Max)),
-								ComputeVisualLogTransform(RenderAreaWorldTransformForVisLog).ToMatrixWithScale(), FColor::White, TEXT(""));
-						}
-					});
-
-				ILandscapeEditLayerRenderer::FRenderParams RenderParams(this, RenderGroupTargetLayerNames, RenderGroupTargetLayerInfos, RenderStep.RendererState, SortedComponentMergeRenderInfos, RenderAreaWorldTransform, RenderAreaSectionRect);
-				if (RenderStep.Type == FMergeRenderStep::EType::RenderLayer_Recorded)
 				{
-					checkf(RDGBuilderRecorder.IsRecording(), TEXT("RenderLayer_Recorded must be preceded by a BeginRenderCommandRecorder which should create a command recorder"));
+					TRACE_CPUPROFILER_EVENT_SCOPE(RenderLayer);
 
-					RDG_RENDER_COMMAND_RECORDER_BREADCRUMB_EVENT(RDGBuilderRecorder, "%s", *RenderStepProfilingEventName);
-					Renderer->RenderLayer(RenderParams, RDGBuilderRecorder);
-				}
-				else
-				{
-					checkf(RDGBuilderRecorder.IsEmpty() && !RDGBuilderRecorder.IsRecording(), TEXT("RenderLayer_Immediate should be preceded by a EndRenderCommandRecorder which should finalize a command recorder and destroy it"));
+					ILandscapeEditLayerRenderer* Renderer = RenderStep.RendererState.GetRenderer();
+					check((Renderer != nullptr) && EnumHasAnyFlags(Renderer->GetRenderFlags(this), ERenderFlags::RenderMode_Mask));
 
-					RHI_BREADCRUMB_EVENT_GAMETHREAD("%s", RenderStepProfilingEventName);
-					Renderer->RenderLayer(RenderParams, RDGBuilderRecorder);
-				}
+					// TODO[jonathan.bard] offset the world transform to account for the half-pixel offset?
+					//RenderParams.RenderAreaWorldTransform = FTransform(LandscapeTransform.GetRotation(), LandscapeTransform.GetTranslation() + FVector(RenderBatch.SectionRect.Min) /** (double)ComponentSizeQuads - FVector(0.5, 0.5, 0)*/, LandscapeTransform.GetScale3D());
 
-				++RenderLayerStepIndex;
-				break;
-			}
-			case FMergeRenderStep::EType::SignalBatchMergeGroupDone:
-			{
-				TRACE_CPUPROFILER_EVENT_SCOPE(MergeGroupDone);
-				RHI_BREADCRUMB_EVENT_GAMETHREAD("Step [%i] (%s) : Render Group Done", RenderLayerStepIndex, ConvertTargetLayerNamesToString(RenderGroupTargetLayerNames));
+					FString RenderStepProfilingEventName = FString::Format(TEXT("Step [{0}] ({1}): Render {2}"), { RenderLayerStepIndex, *ConvertTargetLayerNamesToString(TargetLayerGroupLayerNames), *Renderer->GetEditLayerRendererDebugName() });
 
-				checkf(RDGBuilderRecorder.IsEmpty() && !RDGBuilderRecorder.IsRecording(), TEXT("SignalBatchMergeGroupDone should be preceded by a EndRenderCommandRecorder which should finalize a command recorder and destroy it"));
-				// The last render target we wrote to is the one containing the batch group's merge result : 
-				FOnRenderBatchTargetGroupDoneParams Params(this, RenderBatch, RenderGroupTargetLayerNames, RenderGroupTargetLayerInfos, SortedComponentMergeRenderInfos);
-				OnBatchTargetGroupDone(Params, RDGBuilderRecorder);
-				break;
-			}
-			default:
-				checkf(false, TEXT("Invalid render step type"));
+					// TODO [jonathan.bard] : this is more of a Batch world transform/section rect at the moment. Shall we have a RenderAreaWorldTransform/RenderAreaSectionRect in FRenderParams and a BatchRenderAreaWorldTransform in FMergeRenderBatch?
+					//  because currently the old BP brushes work with FMergeRenderBatch data (i.e. 1 transform for the batch and a section rect for the entire batch) but eventually, renderers might be interested in just their Render step context, 
+					//  that is : 1 matrix corresponding to the bottom left corner of their list of components to render?
+					FTransform RenderAreaWorldTransform = FTransform(FVector(RenderBatch.SectionRect.Min)) * LandscapeTransform;
+					FIntRect RenderAreaSectionRect = RenderBatch.SectionRect;
+				
+					// Drop some visual cues to help understand how each renderer is applied :
+					UE_IFVLOG(
+						if (IsVisualLogEnabled() && bVisualLogMergeProcess)
+						{
+							FTransform RenderAreaWorldTransformForVisLog = FTransform(FVector(RenderBatch.SectionRect.Min)) * LandscapeWorldTransformForVisLog;
+							IncrementVisualLogOffset();
+							UE_VLOG_LOCATION(Landscape, LogLandscape, Log, ComputeVisualLogTransform(RenderAreaWorldTransformForVisLog).GetTranslation(), 10.0f, FColor::Red, TEXT("%s"), *RenderStepProfilingEventName);
+							UE_VLOG_WIREOBOX(Landscape, LogLandscape, Log, FBox(FVector(RenderBatch.SectionRect.Min) - FVector(0.5, 0.5, 0.0), FVector(RenderBatch.SectionRect.Max) - FVector(0.5, 0.5, 0.0)),
+								ComputeVisualLogTransform(LandscapeWorldTransformForVisLog).ToMatrixWithScale(), FColor::White, TEXT(""));
+
+							// Draw each component's bounds rendered by this renderer : 
+							for (const FComponentMergeRenderInfo& ComponentMergeRenderInfo : SortedComponentMergeRenderInfos)
+							{
+								UE_VLOG_WIREOBOX(Landscape, LogLandscape, Log, FBox(FVector(ComponentMergeRenderInfo.ComponentRegionInRenderArea.Min), FVector(ComponentMergeRenderInfo.ComponentRegionInRenderArea.Max)),
+									ComputeVisualLogTransform(RenderAreaWorldTransformForVisLog).ToMatrixWithScale(), FColor::White, TEXT(""));
+							}
+						});
+			
+					ILandscapeEditLayerRenderer::FRenderParams RenderParams(this, TargetLayerGroupLayerNames, TargetLayerGroupLayerInfos, RenderStep.RendererState, SortedComponentMergeRenderInfos, RenderAreaWorldTransform, RenderAreaSectionRect);
+				    if (RenderStep.Type == FMergeRenderStep::EType::RenderLayer_Recorded)
+				    {
+					    checkf(RDGBuilderRecorder.IsRecording(), TEXT("RenderLayer_Recorded must be preceded by a BeginRenderCommandRecorder which should create a command recorder"));
+    
+					    RDG_RENDER_COMMAND_RECORDER_BREADCRUMB_EVENT(RDGBuilderRecorder, "%s", *RenderStepProfilingEventName);
+					    Renderer->RenderLayer(RenderParams, RDGBuilderRecorder);
+				    }
+				    else
+				    {
+					    checkf(RDGBuilderRecorder.IsEmpty() && !RDGBuilderRecorder.IsRecording(), TEXT("RenderLayer_Immediate should be preceded by a EndRenderCommandRecorder which should finalize a command recorder and destroy it"));
+    
+					    RHI_BREADCRUMB_EVENT_GAMETHREAD("%s", RenderStepProfilingEventName);
+					    Renderer->RenderLayer(RenderParams, RDGBuilderRecorder);
+				    }
+    
+				    ++RenderLayerStepIndex;
+				    break;
+			    }
+			    case FMergeRenderStep::EType::SignalBatchMergeGroupDone:
+			    {
+				    TRACE_CPUPROFILER_EVENT_SCOPE(MergeGroupDone);
+				    RHI_BREADCRUMB_EVENT_GAMETHREAD("Step [%i] (%s) : Render Group Done", RenderLayerStepIndex, ConvertTargetLayerNamesToString(TargetLayerGroupLayerNames));
+    
+				    checkf(RDGBuilderRecorder.IsEmpty() && !RDGBuilderRecorder.IsRecording(), TEXT("SignalBatchMergeGroupDone should be preceded by a EndRenderCommandRecorder which should finalize a command recorder and destroy it"));
+				    // The last render target we wrote to is the one containing the batch group's merge result : 
+				    FOnRenderBatchTargetGroupDoneParams Params(this, RenderBatch, TargetLayerGroupLayerNames, TargetLayerGroupLayerInfos, SortedComponentMergeRenderInfos);
+				    OnBatchTargetGroupDone(Params, RDGBuilderRecorder);
+				    break;
+			    }
+			    default:
+				    checkf(false, TEXT("Invalid render step type"));
 			}
 		}
 		
diff --git a/Engine/Source/Runtime/Landscape/Private/LandscapeEditLayerRendererPrivate.cpp b/Engine/Source/Runtime/Landscape/Private/LandscapeEditLayerRendererPrivate.cpp
--- a/Engine/Source/Runtime/Landscape/Private/LandscapeEditLayerRendererPrivate.cpp
+++ b/Engine/Source/Runtime/Landscape/Private/LandscapeEditLayerRendererPrivate.cpp
@@ -16,7 +16,7 @@
 // ----------------------------------------------------------------------------------
 
 void ULandscapeDefaultEditLayerRenderer::GetRendererStateInfo(const UE::Landscape::EditLayers::FMergeContext* InMergeContext,
-	UE::Landscape::EditLayers::FEditLayerTargetTypeState& OutSupportedTargetTypeState, UE::Landscape::EditLayers::FEditLayerTargetTypeState& OutEnabledTargetTypeState, TArray<TBitArray<>>& OutRenderGroups) const
+	UE::Landscape::EditLayers::FEditLayerTargetTypeState& OutSupportedTargetTypeState, UE::Landscape::EditLayers::FEditLayerTargetTypeState& OutEnabledTargetTypeState, TArray<TBitArray<>>& OutTargetLayerGroups) const
 {
 	using namespace UE::Landscape::EditLayers;
 
@@ -66,7 +66,7 @@ void ULandscapeDefaultEditLayerRenderer::RenderLayer(FRenderParams& RenderParams
 // ----------------------------------------------------------------------------------
 
 void ULandscapeHeightmapNormalsEditLayerRenderer::GetRendererStateInfo(const UE::Landscape::EditLayers::FMergeContext* InMergeContext,
-	UE::Landscape::EditLayers::FEditLayerTargetTypeState& OutSupportedTargetTypeState, UE::Landscape::EditLayers::FEditLayerTargetTypeState& OutEnabledTargetTypeState, TArray<TBitArray<>>& OutRenderGroups) const
+	UE::Landscape::EditLayers::FEditLayerTargetTypeState& OutSupportedTargetTypeState, UE::Landscape::EditLayers::FEditLayerTargetTypeState& OutEnabledTargetTypeState, TArray<TBitArray<>>& OutTargetLayerGroups) const
 {
 	// Only relevant for heightmaps :
 	OutSupportedTargetTypeState.SetTargetTypeMask(ELandscapeToolTargetTypeFlags::Heightmap);
@@ -100,7 +100,7 @@ FString ULandscapeHeightmapNormalsEditLayerRenderer::GetEditLayerRendererDebugNa
 // ----------------------------------------------------------------------------------
 
 void ULandscapeWeightmapWeightBlendedLayersRenderer::GetRendererStateInfo(const UE::Landscape::EditLayers::FMergeContext* InMergeContext,
-	UE::Landscape::EditLayers::FEditLayerTargetTypeState& OutSupportedTargetTypeState, UE::Landscape::EditLayers::FEditLayerTargetTypeState& OutEnabledTargetTypeState, TArray<TBitArray<>>& OutRenderGroups) const
+	UE::Landscape::EditLayers::FEditLayerTargetTypeState& OutSupportedTargetTypeState, UE::Landscape::EditLayers::FEditLayerTargetTypeState& OutEnabledTargetTypeState, TArray<TBitArray<>>& OutTargetLayerGroups) const
 {
 	using namespace UE::Landscape::EditLayers;
 
@@ -109,10 +109,10 @@ void ULandscapeWeightmapWeightBlendedLayersRenderer::GetRendererStateInfo(const
 	FEditLayerTargetTypeState OutputTargetTypeState(InMergeContext, ELandscapeToolTargetTypeFlags::Weightmap, WeightBlendedWeightmapLayerBitIndices);
 	OutEnabledTargetTypeState = OutSupportedTargetTypeState = OutputTargetTypeState;
 
-	// Now fill in the render groups : one group for all weight-blended layers : 
+	// Now fill in the target layer groups : one group for all weight-blended layers : 
 	if (WeightBlendedWeightmapLayerBitIndices.Contains(true))
 	{
-		OutRenderGroups.Add(WeightBlendedWeightmapLayerBitIndices);
+		OutTargetLayerGroups.Add(WeightBlendedWeightmapLayerBitIndices);
 	}
 }
 
diff --git a/Engine/Source/Runtime/Landscape/Private/LandscapeEditLayerRendererPrivate.h b/Engine/Source/Runtime/Landscape/Private/LandscapeEditLayerRendererPrivate.h
--- a/Engine/Source/Runtime/Landscape/Private/LandscapeEditLayerRendererPrivate.h
+++ b/Engine/Source/Runtime/Landscape/Private/LandscapeEditLayerRendererPrivate.h
@@ -30,7 +30,7 @@ class ULandscapeDefaultEditLayerRenderer : public UObject
 public:
 	//~ Begin ILandscapeEditLayerRenderer implementation
 	virtual void GetRendererStateInfo(const UE::Landscape::EditLayers::FMergeContext* InMergeContext,
-		UE::Landscape::EditLayers::FEditLayerTargetTypeState& OutSupportedTargetTypeState, UE::Landscape::EditLayers::FEditLayerTargetTypeState& OutEnabledTargetTypeState, TArray<TBitArray<>>& OutRenderGroups) const override;
+		UE::Landscape::EditLayers::FEditLayerTargetTypeState& OutSupportedTargetTypeState, UE::Landscape::EditLayers::FEditLayerTargetTypeState& OutEnabledTargetTypeState, TArray<TBitArray<>>& OutTargetLayerGroups) const override;
 	virtual TArray<UE::Landscape::EditLayers::FEditLayerRenderItem> GetRenderItems(const UE::Landscape::EditLayers::FMergeContext* InMergeContext) const override;
 	virtual FString GetEditLayerRendererDebugName() const override;
 	virtual UE::Landscape::EditLayers::ERenderFlags GetRenderFlags(const UE::Landscape::EditLayers::FMergeContext* InMergeContext) const override { return UE::Landscape::EditLayers::ERenderFlags::RenderMode_Recorded; }
@@ -56,7 +56,7 @@ class ULandscapeHeightmapNormalsEditLayerRenderer : public UObject
 public:
 	//~ Begin ILandscapeEditLayerRenderer implementation
 	virtual void GetRendererStateInfo(const UE::Landscape::EditLayers::FMergeContext* InMergeContext,
-		UE::Landscape::EditLayers::FEditLayerTargetTypeState& OutSupportedTargetTypeState, UE::Landscape::EditLayers::FEditLayerTargetTypeState& OutEnabledTargetTypeState, TArray<TBitArray<>>& OutRenderGroups) const override;
+		UE::Landscape::EditLayers::FEditLayerTargetTypeState& OutSupportedTargetTypeState, UE::Landscape::EditLayers::FEditLayerTargetTypeState& OutEnabledTargetTypeState, TArray<TBitArray<>>& OutTargetLayerGroups) const override;
 	virtual TArray<UE::Landscape::EditLayers::FEditLayerRenderItem> GetRenderItems(const UE::Landscape::EditLayers::FMergeContext* InMergeContext) const override;
 	virtual FString GetEditLayerRendererDebugName() const override;
 	virtual UE::Landscape::EditLayers::ERenderFlags GetRenderFlags(const UE::Landscape::EditLayers::FMergeContext* InMergeContext) const override { return UE::Landscape::EditLayers::ERenderFlags::RenderMode_Recorded; }
@@ -80,7 +80,7 @@ class ULandscapeWeightmapWeightBlendedLayersRenderer : public UObject
 public:
 	//~ Begin ILandscapeEditLayerRenderer implementation
 	virtual void GetRendererStateInfo(const UE::Landscape::EditLayers::FMergeContext* InMergeContext,
-		UE::Landscape::EditLayers::FEditLayerTargetTypeState& OutSupportedTargetTypeState, UE::Landscape::EditLayers::FEditLayerTargetTypeState& OutEnabledTargetTypeState, TArray<TBitArray<>>& OutRenderGroups) const override;
+		UE::Landscape::EditLayers::FEditLayerTargetTypeState& OutSupportedTargetTypeState, UE::Landscape::EditLayers::FEditLayerTargetTypeState& OutEnabledTargetTypeState, TArray<TBitArray<>>& OutTargetLayerGroups) const override;
 	virtual TArray<UE::Landscape::EditLayers::FEditLayerRenderItem> GetRenderItems(const UE::Landscape::EditLayers::FMergeContext* InMergeContext) const override;
 	virtual FString GetEditLayerRendererDebugName() const override;
 	virtual UE::Landscape::EditLayers::ERenderFlags GetRenderFlags(const UE::Landscape::EditLayers::FMergeContext* InMergeContext) const override { return UE::Landscape::EditLayers::ERenderFlags::RenderMode_Recorded; }
diff --git a/Engine/Source/Runtime/Landscape/Private/LandscapeEditLayers.cpp b/Engine/Source/Runtime/Landscape/Private/LandscapeEditLayers.cpp
--- a/Engine/Source/Runtime/Landscape/Private/LandscapeEditLayers.cpp
+++ b/Engine/Source/Runtime/Landscape/Private/LandscapeEditLayers.cpp
@@ -5454,7 +5454,7 @@ TArray<UE::Landscape::EditLayers::FEditLayerRendererState> FLandscapeLayerBrush:
 }
 
 void ULandscapeEditLayerPersistent::GetRendererStateInfo(const UE::Landscape::EditLayers::FMergeContext* InMergeContext,
-	UE::Landscape::EditLayers::FEditLayerTargetTypeState& OutSupportedTargetTypeState, UE::Landscape::EditLayers::FEditLayerTargetTypeState& OutEnabledTargetTypeState, TArray<TBitArray<>>& OutRenderGroups) const
+	UE::Landscape::EditLayers::FEditLayerTargetTypeState& OutSupportedTargetTypeState, UE::Landscape::EditLayers::FEditLayerTargetTypeState& OutEnabledTargetTypeState, TArray<TBitArray<>>& OutTargetLayerGroups) const
 {
 	using namespace UE::Landscape::EditLayers;
 
@@ -5553,7 +5553,7 @@ void ULandscapeEditLayerPersistent::RenderLayer(FRenderParams& RenderParams, UE:
 	const FLandscapeLayer* LayerStruct = GetOwningLayer();
 	check(LayerStruct != nullptr);
 
-	const int32 NumTargetLayersInRenderGroup = RenderParams.RenderGroupTargetLayerNames.Num();
+	const int32 NumTargetLayersInGroup = RenderParams.TargetLayerGroupLayerNames.Num();
 
 	// The first step is to copy all of the necessary components' textures to WriteRT and remove the duplicate borders when doing so (subsection by subsection). This is done with a 
 	//  "copy from multiple sources" shader instead of several texture copies, in order to reduce the amount of copy texture commands, which can have a big impact on render-thread performance
@@ -5575,7 +5575,7 @@ void ULandscapeEditLayerPersistent::RenderLayer(FRenderParams& RenderParams, UE:
 		uint8 ChannelSwizzleMask = RGBAToRGBASwizzleMask;
 	};
 	TArray<FCopyQuadParams> CopyQuadsParams;
-	CopyQuadsParams.Reserve(RenderParams.SortedComponentMergeRenderInfos.Num() * NumTargetLayersInRenderGroup* Landscape->NumSubsections * Landscape->NumSubsections);
+	CopyQuadsParams.Reserve(RenderParams.SortedComponentMergeRenderInfos.Num() * NumTargetLayersInGroup* Landscape->NumSubsections * Landscape->NumSubsections);
 
 	// This is kinda hacky, but we since we reuse AddRasterizeToRectsPass and it passes the source texture's sample coordinates from VS to PS via the UVs, dividing the source rect by a single TextureSize, 
 	//  we correct each quad's source rect by a common scale factor so that it samples at the right location in each source texture: 
@@ -5585,9 +5585,9 @@ void ULandscapeEditLayerPersistent::RenderLayer(FRenderParams& RenderParams, UE:
 	// Build the full list of quads to copy :
 	for (const FComponentMergeRenderInfo& ComponentMergeRenderInfo : RenderParams.SortedComponentMergeRenderInfos)
 	{
-		for (int32 TargetLayerIndexInRenderGroup = 0; TargetLayerIndexInRenderGroup < NumTargetLayersInRenderGroup; ++TargetLayerIndexInRenderGroup)
+		for (int32 TargetLayerIndexInGroup = 0; TargetLayerIndexInGroup < NumTargetLayersInGroup; ++TargetLayerIndexInGroup)
 		{
-			const FName TargetLayerName = RenderParams.RenderGroupTargetLayerNames[TargetLayerIndexInRenderGroup];
+			const FName TargetLayerName = RenderParams.TargetLayerGroupLayerNames[TargetLayerIndexInGroup];
 			UTexture2D* SourceTexture = nullptr;
 			FVector2D SourceTextureBias(ForceInit);
 			// Copy all channels by default :
@@ -5643,7 +5643,7 @@ void ULandscapeEditLayerPersistent::RenderLayer(FRenderParams& RenderParams, UE:
 						QuadParams.SourceTextureResource = SourceTextureResource;
 						QuadParams.SourceRect = FIntRect(SourceTexturePixelCoordinatesToQuadCoords(SourcePosition), SourceTexturePixelCoordinatesToQuadCoords(SourcePosition + FIntPoint(SubsectionSizeVerts, SubsectionSizeVerts)));
 						QuadParams.DestinationRect = FIntRect(DestinationPosition, DestinationPosition + FIntPoint(SubsectionSizeVerts, SubsectionSizeVerts));
-						QuadParams.DestinationArrayIndex = TargetLayerIndexInRenderGroup;
+						QuadParams.DestinationArrayIndex = TargetLayerIndexInGroup;
 						QuadParams.ChannelSwizzleMask = CopyTextureSwizzleMask;
 					}
 				}
@@ -5666,171 +5666,171 @@ void ULandscapeEditLayerPersistent::RenderLayer(FRenderParams& RenderParams, UE:
 
 	// First step : render the edit layer's component quads
 	{
-		// Sort the list of quads by output texture first, then by source texture, because we'll coalesce several texture copies into the same pass (one output, multiple inputs)
-		CopyQuadsParams.Sort([](const FCopyQuadParams& InLHS, const FCopyQuadParams& InRHS)
+	// Sort the list of quads by output texture first, then by source texture, because we'll coalesce several texture copies into the same pass (one output, multiple inputs)
+	CopyQuadsParams.Sort([](const FCopyQuadParams& InLHS, const FCopyQuadParams& InRHS)
+	{
+		if (InLHS.DestinationArrayIndex == InRHS.DestinationArrayIndex)
 		{
-			if (InLHS.DestinationArrayIndex == InRHS.DestinationArrayIndex)
+			if (InLHS.SourceTextureResource == InRHS.SourceTextureResource)
 			{
-				if (InLHS.SourceTextureResource == InRHS.SourceTextureResource)
-				{
-					return (InLHS.SourceRect.Min.Y == InRHS.SourceRect.Min.Y) 
-						? (InLHS.SourceRect.Min.X < InRHS.SourceRect.Min.X)
-						: (InLHS.SourceRect.Min.Y < InRHS.SourceRect.Min.Y);
-				}
-				return (InLHS.SourceTextureResource < InRHS.SourceTextureResource);
+				return (InLHS.SourceRect.Min.Y == InRHS.SourceRect.Min.Y) 
+					? (InLHS.SourceRect.Min.X < InRHS.SourceRect.Min.X)
+					: (InLHS.SourceRect.Min.Y < InRHS.SourceRect.Min.Y);
 			}
+			return (InLHS.SourceTextureResource < InRHS.SourceTextureResource);
+		}
 
-			return (InLHS.DestinationArrayIndex < InRHS.DestinationArrayIndex);
-		});
+		return (InLHS.DestinationArrayIndex < InRHS.DestinationArrayIndex);
+	});
 
-		// Now process this list of quads and prepare as many passes as necessary for performing all the copies :
+	// Now process this list of quads and prepare as many passes as necessary for performing all the copies :
 		auto RDGCommand =
-			[CopyQuadsParams = MoveTemp(CopyQuadsParams)
+			[ CopyQuadsParams = MoveTemp(CopyQuadsParams)
 			, OutputResource = WriteRT->GetRenderTarget()->GetResource()
 			, OutputResourceName = WriteRT->GetDebugName()
 			, SceneInterface
 			, CommonSourceTextureSize
 			, bUseSwizzleMask] (FRDGBuilder& GraphBuilder)
-			{
+		{
 				RDG_EVENT_SCOPE(GraphBuilder, "CopyEditLayer -> %s", *OutputResourceName);
 
-				FRDGTextureSRVRef BlackDummySRVRef = GraphBuilder.CreateSRV(FRDGTextureSRVDesc::Create(GSystemTextures.GetBlackDummy(GraphBuilder)));
+			FRDGTextureSRVRef BlackDummySRVRef = GraphBuilder.CreateSRV(FRDGTextureSRVDesc::Create(GSystemTextures.GetBlackDummy(GraphBuilder)));
 
-				FCopyQuadsMultiSourcePS::FPermutationDomain PermutationVector;
-				PermutationVector.Set< FCopyQuadsMultiSourcePS::FUseSwizleMask >(bUseSwizzleMask);
-				FGlobalShaderMap* ShaderMap = GetGlobalShaderMap(GMaxRHIFeatureLevel);
-				TShaderRef<FCopyQuadsMultiSourcePS> PixelShader = ShaderMap->GetShader<FCopyQuadsMultiSourcePS>(PermutationVector);
-
-				// TODO [jonathan.bard] this is just to avoid a RHI validation error for unoptimized shaders... once validation is made to not issue those errors, we can remove this
-				// Create a SceneView to please the shader bindings, but it's unused in practice 
-				FSceneViewFamilyContext ViewFamily(FSceneViewFamily::ConstructionValues(nullptr, SceneInterface, FEngineShowFlags(ESFIM_Game)).SetTime(FGameTime::GetTimeSinceAppStart()));
-				FSceneViewInitOptions ViewInitOptions;
-				ViewInitOptions.ViewFamily = &ViewFamily;
-				ViewInitOptions.SetViewRectangle(FIntRect(0, 0, 1, 1)); // Use a dummy rect to avoid a check(slow)
-				GetRendererModule().CreateAndInitSingleView(GraphBuilder.RHICmdList, &ViewFamily, &ViewInitOptions);
-				const FSceneView* View = ViewFamily.Views[0];
+			FCopyQuadsMultiSourcePS::FPermutationDomain PermutationVector;
+			PermutationVector.Set< FCopyQuadsMultiSourcePS::FUseSwizleMask >(bUseSwizzleMask);
+			FGlobalShaderMap* ShaderMap = GetGlobalShaderMap(GMaxRHIFeatureLevel);
+			TShaderRef<FCopyQuadsMultiSourcePS> PixelShader = ShaderMap->GetShader<FCopyQuadsMultiSourcePS>(PermutationVector);
 
-				// The following variables allow to accumulate quads for rendering in a single pass. Then, we flush them all out when the pass is full and we start again : 
-				int32 CurrentDestinationArrayIndex = INDEX_NONE;
-				TArray<FTextureResource*> CurrentSourceTextureResources;
-				TArray<FUintVector2> CurrentQuadInfos;
-				TArray<FUintVector4> CurrentSourceRects, CurrentDestinationRects;
+			// TODO [jonathan.bard] this is just to avoid a RHI validation error for unoptimized shaders... once validation is made to not issue those errors, we can remove this
+			// Create a SceneView to please the shader bindings, but it's unused in practice 
+			FSceneViewFamilyContext ViewFamily(FSceneViewFamily::ConstructionValues(nullptr, SceneInterface, FEngineShowFlags(ESFIM_Game)).SetTime(FGameTime::GetTimeSinceAppStart()));
+			FSceneViewInitOptions ViewInitOptions;
+			ViewInitOptions.ViewFamily = &ViewFamily;
+			ViewInitOptions.SetViewRectangle(FIntRect(0, 0, 1, 1)); // Use a dummy rect to avoid a check(slow)
+				GetRendererModule().CreateAndInitSingleView(GraphBuilder.RHICmdList, &ViewFamily, &ViewInitOptions);
+			const FSceneView* View = ViewFamily.Views[0];
 
-				FRDGTextureRef OutputTexture = GraphBuilder.RegisterExternalTexture(CreateRenderTarget(OutputResource->GetTextureRHI(), TEXT("OutputTexture")));
+			// The following variables allow to accumulate quads for rendering in a single pass. Then, we flush them all out when the pass is full and we start again : 
+			int32 CurrentDestinationArrayIndex = INDEX_NONE;
+			TArray<FTextureResource*> CurrentSourceTextureResources;
+			TArray<FUintVector2> CurrentQuadInfos;
+			TArray<FUintVector4> CurrentSourceRects, CurrentDestinationRects;
 
-				// When we've reached the limit of what we can render in one pass, flush all and add a render pass : 
-				auto FlushCopyQuads = [&GraphBuilder, View, OutputTexture, &CurrentDestinationArrayIndex, &CurrentSourceTextureResources, &CurrentSourceRects, &CurrentDestinationRects,
-					&CurrentQuadInfos, &CommonSourceTextureSize, BlackDummySRVRef, ShaderMap, PixelShader]()
-					{
-						if (CurrentSourceTextureResources.IsEmpty())
-						{
-							// Nothing to flush :
-							check(CurrentSourceRects.IsEmpty() && CurrentDestinationRects.IsEmpty() && CurrentQuadInfos.IsEmpty());
-							return;
-						}
+			FRDGTextureRef OutputTexture = GraphBuilder.RegisterExternalTexture(CreateRenderTarget(OutputResource->GetTextureRHI(), TEXT("OutputTexture")));
 
-						check(!CurrentSourceRects.IsEmpty() && (CurrentSourceRects.Num() == CurrentDestinationRects.Num()) && (CurrentSourceRects.Num() == CurrentQuadInfos.Num()));
+			// When we've reached the limit of what we can render in one pass, flush all and add a render pass : 
+			auto FlushCopyQuads = [&GraphBuilder, View, OutputTexture, &CurrentDestinationArrayIndex, &CurrentSourceTextureResources, &CurrentSourceRects, &CurrentDestinationRects,
+				&CurrentQuadInfos, &CommonSourceTextureSize, BlackDummySRVRef, ShaderMap, PixelShader]()
+			{
+				if (CurrentSourceTextureResources.IsEmpty())
+				{
+					// Nothing to flush :
+					check(CurrentSourceRects.IsEmpty() && CurrentDestinationRects.IsEmpty() && CurrentQuadInfos.IsEmpty());
+					return;
+				}
 
-						FRDGBufferRef RectBuffer = CreateUploadBuffer(GraphBuilder, TEXT("DestinationRects"), TConstArrayView<FUintVector4>(MakeArrayView(CurrentDestinationRects)));
-						FRDGBufferSRVRef RectBufferSRV = GraphBuilder.CreateSRV(FRDGBufferSRVDesc(RectBuffer, PF_R32G32B32A32_UINT));
+				check(!CurrentSourceRects.IsEmpty() && (CurrentSourceRects.Num() == CurrentDestinationRects.Num()) && (CurrentSourceRects.Num() == CurrentQuadInfos.Num()));
 
-						FRDGBufferRef RectUVBuffer = CreateUploadBuffer(GraphBuilder, TEXT("RectUVs"), TConstArrayView<FUintVector4>(MakeArrayView(CurrentSourceRects)));
-						FRDGBufferSRVRef RectUVBufferSRV = GraphBuilder.CreateSRV(FRDGBufferSRVDesc(RectUVBuffer, PF_R32G32B32A32_UINT));
+				FRDGBufferRef RectBuffer = CreateUploadBuffer(GraphBuilder, TEXT("DestinationRects"), TConstArrayView<FUintVector4>(MakeArrayView(CurrentDestinationRects)));
+				FRDGBufferSRVRef RectBufferSRV = GraphBuilder.CreateSRV(FRDGBufferSRVDesc(RectBuffer, PF_R32G32B32A32_UINT));
 
-						FRDGBufferRef QuadInfosBuffer = CreateUploadBuffer(GraphBuilder, TEXT("QuadInfos"), TConstArrayView<FUintVector2>(MakeArrayView(CurrentQuadInfos)));
-						FRDGBufferSRVRef QuadInfosBufferSRV = GraphBuilder.CreateSRV(FRDGBufferSRVDesc(QuadInfosBuffer, PF_R32G32_UINT));
+				FRDGBufferRef RectUVBuffer = CreateUploadBuffer(GraphBuilder, TEXT("RectUVs"), TConstArrayView<FUintVector4>(MakeArrayView(CurrentSourceRects)));
+				FRDGBufferSRVRef RectUVBufferSRV = GraphBuilder.CreateSRV(FRDGBufferSRVDesc(RectUVBuffer, PF_R32G32B32A32_UINT));
 
-						// We may target an array slice directly : 
-						int32 ArrayIndex = -1;
-						if (OutputTexture->Desc.IsTextureArray())
-						{
-							check(CurrentDestinationArrayIndex < OutputTexture->Desc.ArraySize);
-							ArrayIndex = CurrentDestinationArrayIndex;
-						}
-
-						FCopyQuadsMultiSourcePSParameters* PassParameters = GraphBuilder.AllocParameters<FCopyQuadsMultiSourcePSParameters>();
-						PassParameters->RenderTargets[0] = FRenderTargetBinding(OutputTexture, ERenderTargetLoadAction::ELoad, /*InMipIndex = */0, ArrayIndex);
-						PassParameters->PS.View = View->ViewUniformBuffer;
-						PassParameters->PS.InQuadInfos = QuadInfosBufferSRV;
-						check(CurrentSourceTextureResources.Num() <= FCopyQuadsMultiSourcePS::NumMultiSources);
-						for (int32 TextureIndex = 0; TextureIndex < FCopyQuadsMultiSourcePS::NumMultiSources; ++TextureIndex)
-						{
-							FRDGTextureRef SourceTexture = BlackDummySRVRef->GetParent();
-							if (CurrentSourceTextureResources.IsValidIndex(TextureIndex))
-							{
-								SourceTexture = GraphBuilder.RegisterExternalTexture(CreateRenderTarget(CurrentSourceTextureResources[TextureIndex]->GetTexture2DRHI(), TEXT("SourceTexture")));
-							}
-							PassParameters->PS.InSourceTexture[TextureIndex] = SourceTexture;
-						}
+				FRDGBufferRef QuadInfosBuffer = CreateUploadBuffer(GraphBuilder, TEXT("QuadInfos"), TConstArrayView<FUintVector2>(MakeArrayView(CurrentQuadInfos)));
+				FRDGBufferSRVRef QuadInfosBufferSRV = GraphBuilder.CreateSRV(FRDGBufferSRVDesc(QuadInfosBuffer, PF_R32G32_UINT));
+				
+				// We may target an array slice directly : 
+				int32 ArrayIndex = -1;
+				if (OutputTexture->Desc.IsTextureArray())
+				{
+					check(CurrentDestinationArrayIndex < OutputTexture->Desc.ArraySize);
+					ArrayIndex = CurrentDestinationArrayIndex;
+				}
 
-						FPixelShaderUtils::AddRasterizeToRectsPass<FCopyQuadsMultiSourcePS>(GraphBuilder,
-							ShaderMap,
-							RDG_EVENT_NAME("CopyQuadsMultiSourcePS"),
-							PixelShader,
-							PassParameters,
-							/*ViewportSize = */OutputTexture->Desc.Extent,
-							RectBufferSRV,
-							CurrentDestinationRects.Num(),
-							/*BlendState = */nullptr,
-							/*RasterizerState = */nullptr,
-							/*DepthStencilState = */nullptr,
-							/*StencilRef = */0,
-							/*TextureSize = */FIntPoint(FMath::FloorToInt32(CommonSourceTextureSize.X), FMath::FloorToInt32(CommonSourceTextureSize.Y)),
-							RectUVBufferSRV);
-
-						// Reset all for the next pass that comes along : 
-						CurrentDestinationArrayIndex = INDEX_NONE;
-						CurrentSourceTextureResources.Empty();
-						CurrentSourceRects.Empty();
-						CurrentDestinationRects.Empty();
-						CurrentQuadInfos.Empty();
-					};
-
-				for (const FCopyQuadParams& InSingleQuadParams : CopyQuadsParams)
-				{
-					check(InSingleQuadParams.DestinationArrayIndex != INDEX_NONE);
-					// If the output texture/slice has changed since the last iteration, it's time to flush and start a new pass :
-					if (CurrentDestinationArrayIndex != InSingleQuadParams.DestinationArrayIndex)
+				FCopyQuadsMultiSourcePSParameters* PassParameters = GraphBuilder.AllocParameters<FCopyQuadsMultiSourcePSParameters>();
+				PassParameters->RenderTargets[0] = FRenderTargetBinding(OutputTexture, ERenderTargetLoadAction::ELoad, /*InMipIndex = */0, ArrayIndex);
+				PassParameters->PS.View = View->ViewUniformBuffer;
+				PassParameters->PS.InQuadInfos = QuadInfosBufferSRV;
+				check(CurrentSourceTextureResources.Num() <= FCopyQuadsMultiSourcePS::NumMultiSources);
+				for (int32 TextureIndex = 0; TextureIndex < FCopyQuadsMultiSourcePS::NumMultiSources; ++TextureIndex)
+				{
+					FRDGTextureRef SourceTexture = BlackDummySRVRef->GetParent();
+					if (CurrentSourceTextureResources.IsValidIndex(TextureIndex))
 					{
-						if (CurrentDestinationArrayIndex != INDEX_NONE)
-						{
-							FlushCopyQuads();
-						}
-						CurrentDestinationArrayIndex = InSingleQuadParams.DestinationArrayIndex;
+						SourceTexture = GraphBuilder.RegisterExternalTexture(CreateRenderTarget(CurrentSourceTextureResources[TextureIndex]->GetTexture2DRHI(), TEXT("SourceTexture")));
 					}
+					PassParameters->PS.InSourceTexture[TextureIndex] = SourceTexture;
+				}
 
-					if (CurrentSourceTextureResources.IsEmpty())
+				FPixelShaderUtils::AddRasterizeToRectsPass<FCopyQuadsMultiSourcePS>(GraphBuilder,
+					ShaderMap,
+					RDG_EVENT_NAME("CopyQuadsMultiSourcePS"),
+					PixelShader,
+					PassParameters,
+					/*ViewportSize = */OutputTexture->Desc.Extent,
+					RectBufferSRV,
+					CurrentDestinationRects.Num(),
+					/*BlendState = */nullptr,
+					/*RasterizerState = */nullptr,
+					/*DepthStencilState = */nullptr,
+					/*StencilRef = */0,
+					/*TextureSize = */FIntPoint(FMath::FloorToInt32(CommonSourceTextureSize.X), FMath::FloorToInt32(CommonSourceTextureSize.Y)),
+					RectUVBufferSRV);
+
+				// Reset all for the next pass that comes along : 
+				CurrentDestinationArrayIndex = INDEX_NONE;
+				CurrentSourceTextureResources.Empty();
+				CurrentSourceRects.Empty();
+				CurrentDestinationRects.Empty();
+				CurrentQuadInfos.Empty();
+			};
+
+			for (const FCopyQuadParams& InSingleQuadParams : CopyQuadsParams)
+			{
+				check(InSingleQuadParams.DestinationArrayIndex != INDEX_NONE);
+				// If the output texture/slice has changed since the last iteration, it's time to flush and start a new pass :
+				if (CurrentDestinationArrayIndex != InSingleQuadParams.DestinationArrayIndex)
+				{
+					if (CurrentDestinationArrayIndex != INDEX_NONE)
 					{
-						CurrentSourceTextureResources.Add(InSingleQuadParams.SourceTextureResource);
+						FlushCopyQuads();
 					}
-					else if (InSingleQuadParams.SourceTextureResource != CurrentSourceTextureResources.Last())
+					CurrentDestinationArrayIndex = InSingleQuadParams.DestinationArrayIndex;
+				}
+
+				if (CurrentSourceTextureResources.IsEmpty())
+				{
+					CurrentSourceTextureResources.Add(InSingleQuadParams.SourceTextureResource);
+				}
+				else if (InSingleQuadParams.SourceTextureResource != CurrentSourceTextureResources.Last())
+				{
+					// If we've reached the amount of textures we can render in a single pass, we flush the pass and initiate a new one : 
+					if (CurrentSourceTextureResources.Num() == FCopyQuadsMultiSourcePS::NumMultiSources)
 					{
-						// If we've reached the amount of textures we can render in a single pass, we flush the pass and initiate a new one : 
-						if (CurrentSourceTextureResources.Num() == FCopyQuadsMultiSourcePS::NumMultiSources)
-						{
-							FlushCopyQuads();
-							check(CurrentSourceTextureResources.IsEmpty());
+						FlushCopyQuads();
+						check(CurrentSourceTextureResources.IsEmpty());
 
-							CurrentDestinationArrayIndex = InSingleQuadParams.DestinationArrayIndex;
-						}
-						CurrentSourceTextureResources.Add(InSingleQuadParams.SourceTextureResource);
+						CurrentDestinationArrayIndex = InSingleQuadParams.DestinationArrayIndex;
 					}
+					CurrentSourceTextureResources.Add(InSingleQuadParams.SourceTextureResource);
+				}
 
-					// If we are using the same texture as the last one, we can render it in the same pass, just append our quad : 
-					if (InSingleQuadParams.SourceTextureResource == CurrentSourceTextureResources.Last())
-					{
-						const int32 SourceTextureIndex = CurrentSourceTextureResources.Num() - 1;
-						CurrentQuadInfos.Add(FUintVector2(static_cast<uint32>(SourceTextureIndex), InSingleQuadParams.ChannelSwizzleMask));
-						CurrentSourceRects.Add(FUintVector4(InSingleQuadParams.SourceRect.Min.X, InSingleQuadParams.SourceRect.Min.Y, InSingleQuadParams.SourceRect.Max.X, InSingleQuadParams.SourceRect.Max.Y));
-						CurrentDestinationRects.Add(FUintVector4(InSingleQuadParams.DestinationRect.Min.X, InSingleQuadParams.DestinationRect.Min.Y, InSingleQuadParams.DestinationRect.Max.X, InSingleQuadParams.DestinationRect.Max.Y));
-					}
+				// If we are using the same texture as the last one, we can render it in the same pass, just append our quad : 
+				if (InSingleQuadParams.SourceTextureResource == CurrentSourceTextureResources.Last())
+				{
+					const int32 SourceTextureIndex = CurrentSourceTextureResources.Num() - 1;
+					CurrentQuadInfos.Add(FUintVector2(static_cast<uint32>(SourceTextureIndex), InSingleQuadParams.ChannelSwizzleMask));
+					CurrentSourceRects.Add(FUintVector4(InSingleQuadParams.SourceRect.Min.X, InSingleQuadParams.SourceRect.Min.Y, InSingleQuadParams.SourceRect.Max.X, InSingleQuadParams.SourceRect.Max.Y));
+					CurrentDestinationRects.Add(FUintVector4(InSingleQuadParams.DestinationRect.Min.X, InSingleQuadParams.DestinationRect.Min.Y, InSingleQuadParams.DestinationRect.Max.X, InSingleQuadParams.DestinationRect.Max.Y));
 				}
+			}
 
-				// Flush the remaining quads if any : 
-				FlushCopyQuads();
+			// Flush the remaining quads if any : 
+			FlushCopyQuads();
 			};
 
-		// We need to specify the final state of the external texture to prevent the graph builder from transitioning it to SRVMask :
+			// We need to specify the final state of the external texture to prevent the graph builder from transitioning it to SRVMask :
 		RDGBuilderRecorder.EnqueueRDGCommand(RDGCommand, { { WriteRT->GetRenderTarget()->GetResource(), ERHIAccess::RTV } });
 	}
 
@@ -5844,25 +5844,25 @@ void ULandscapeEditLayerPersistent::RenderLayer(FRenderParams& RenderParams, UE:
 		if (RenderParams.MergeRenderContext->IsHeightmapMerge())
 		{
 			auto RDGCommand =
-				[TargetLayerName = RenderParams.RenderGroupTargetLayerNames[0]
+				[ TargetLayerName = RenderParams.TargetLayerGroupLayerNames[0]
 				, OutputResource = WriteRT->GetRenderTarget2D()->GetResource()
 				, OutputResourceName = WriteRT->GetDebugName()
 				, CurrentEditLayerResource = CurrentLayerReadRT->GetRenderTarget2D()->GetResource()
 				, PreviousEditLayersResource = PreviousLayersReadRT->GetRenderTarget2D()->GetResource()
 				, EffectiveTextureSize = RenderParams.MergeRenderContext->GetCurrentRenderBatch()->GetRenderTargetResolution(/*bInWithDuplicateBorders = */false)
-				, BlendMode = LayerStruct->BlendMode, HeightmapAlpha = LayerStruct->HeightmapAlpha]
+				, BlendMode = LayerStruct->BlendMode, HeightmapAlpha = LayerStruct->HeightmapAlpha ]
 				(FRDGBuilder& GraphBuilder)
-				{
-					FRDGTextureRef OutputTextureRef = GraphBuilder.RegisterExternalTexture(CreateRenderTarget(OutputResource->TextureRHI, TEXT("OutputTexture")));
-					FRDGTextureRef CurrentEditLayerTextureRef = GraphBuilder.RegisterExternalTexture(CreateRenderTarget(CurrentEditLayerResource->TextureRHI, TEXT("CurrentEditLayerTexture")));
-					FRDGTextureRef PreviousEditLayersTextureRef = GraphBuilder.RegisterExternalTexture(CreateRenderTarget(PreviousEditLayersResource->TextureRHI, TEXT("PreviousEditLayersTexture")));
+			{
+				FRDGTextureRef OutputTextureRef = GraphBuilder.RegisterExternalTexture(CreateRenderTarget(OutputResource->TextureRHI, TEXT("OutputTexture")));
+				FRDGTextureRef CurrentEditLayerTextureRef = GraphBuilder.RegisterExternalTexture(CreateRenderTarget(CurrentEditLayerResource->TextureRHI, TEXT("CurrentEditLayerTexture")));
+				FRDGTextureRef PreviousEditLayersTextureRef = GraphBuilder.RegisterExternalTexture(CreateRenderTarget(PreviousEditLayersResource->TextureRHI, TEXT("PreviousEditLayersTexture")));
 
-					FLandscapeEditLayersHeightmapsMergeEditLayerPS::FParameters* PSParams = GraphBuilder.AllocParameters<FLandscapeEditLayersHeightmapsMergeEditLayerPS::FParameters>();
-					PSParams->RenderTargets[0] = FRenderTargetBinding(OutputTextureRef, ERenderTargetLoadAction::ENoAction);
-					PSParams->InEditLayerBlendMode = static_cast<uint32>(LandscapeBlendModeToEditLayerBlendMode(BlendMode));
-					PSParams->InEditLayerAlpha = HeightmapAlpha;
-					PSParams->InCurrentEditLayerHeightmap = GraphBuilder.CreateSRV(FRDGTextureSRVDesc::Create(CurrentEditLayerTextureRef));
-					PSParams->InPreviousEditLayersHeightmap = GraphBuilder.CreateSRV(FRDGTextureSRVDesc::Create(PreviousEditLayersTextureRef));
+				FLandscapeEditLayersHeightmapsMergeEditLayerPS::FParameters* PSParams = GraphBuilder.AllocParameters<FLandscapeEditLayersHeightmapsMergeEditLayerPS::FParameters>();
+				PSParams->RenderTargets[0] = FRenderTargetBinding(OutputTextureRef, ERenderTargetLoadAction::ENoAction);
+				PSParams->InEditLayerBlendMode = static_cast<uint32>(LandscapeBlendModeToEditLayerBlendMode(BlendMode));
+				PSParams->InEditLayerAlpha = HeightmapAlpha;
+				PSParams->InCurrentEditLayerHeightmap = GraphBuilder.CreateSRV(FRDGTextureSRVDesc::Create(CurrentEditLayerTextureRef));
+				PSParams->InPreviousEditLayersHeightmap = GraphBuilder.CreateSRV(FRDGTextureSRVDesc::Create(PreviousEditLayersTextureRef));
 
 					FLandscapeEditLayersHeightmapsMergeEditLayerPS::MergeEditLayerPS(
 						RDG_EVENT_NAME("MergeEditLayer(%s) -> %s", *TargetLayerName.ToString(), *OutputResourceName), GraphBuilder, PSParams, EffectiveTextureSize);
@@ -5895,8 +5895,8 @@ void ULandscapeEditLayerPersistent::RenderLayer(FRenderParams& RenderParams, UE:
 
 			// Prepare a list of all paint layers to merge and do it all in 1 command :
 			TArray<FMergePaintLayerParams> MergePaintLayersParams;
-			MergePaintLayersParams.Reserve(NumTargetLayersInRenderGroup);
-			for (const FName& TargetLayerName : RenderParams.RenderGroupTargetLayerNames)
+			MergePaintLayersParams.Reserve(NumTargetLayersInGroup);
+			for (const FName& TargetLayerName : RenderParams.TargetLayerGroupLayerNames)
 			{
 				FMergePaintLayerParams& MergePaintLayerParams = MergePaintLayersParams.Emplace_GetRef(TargetLayerName);
 
@@ -5917,36 +5917,36 @@ void ULandscapeEditLayerPersistent::RenderLayer(FRenderParams& RenderParams, UE:
 			}
 
 			auto RDGCommand =
-				[MergePaintLayersParams
+				[ MergePaintLayersParams
 				, OutputResource = WriteRT->GetRenderTarget2DArray()->GetResource()
 				, OutputResourceName = WriteRT->GetDebugName()
 				, CurrentEditLayerResource = CurrentLayerReadRT->GetRenderTarget2DArray()->GetResource()
 				, PreviousEditLayersResource = PreviousLayersReadRT->GetRenderTarget2DArray()->GetResource()
 				, EffectiveTextureSize = RenderParams.MergeRenderContext->GetCurrentRenderBatch()->GetRenderTargetResolution(/*bInWithDuplicateBorders = */false)]
 				(FRDGBuilder& GraphBuilder)
-				{
-					FRDGTextureRef OutputTextureRef = GraphBuilder.RegisterExternalTexture(CreateRenderTarget(OutputResource->TextureRHI, TEXT("OutputTexture")));
-					FRDGTextureRef CurrentEditLayerTextureRef = GraphBuilder.RegisterExternalTexture(CreateRenderTarget(CurrentEditLayerResource->TextureRHI, TEXT("CurrentEditLayerTexture")));
-					FRDGTextureSRVRef CurrentEditLayerTextureSRVRef = GraphBuilder.CreateSRV(FRDGTextureSRVDesc::Create(CurrentEditLayerTextureRef));
-					FRDGTextureRef PreviousEditLayersTextureRef = GraphBuilder.RegisterExternalTexture(CreateRenderTarget(PreviousEditLayersResource->TextureRHI, TEXT("PreviousEditLayersTexture")));
-					FRDGTextureSRVRef PreviousEditLayersTextureSRVRef = GraphBuilder.CreateSRV(FRDGTextureSRVDesc::Create(PreviousEditLayersTextureRef));
+			{
+				FRDGTextureRef OutputTextureRef = GraphBuilder.RegisterExternalTexture(CreateRenderTarget(OutputResource->TextureRHI, TEXT("OutputTexture")));
+				FRDGTextureRef CurrentEditLayerTextureRef = GraphBuilder.RegisterExternalTexture(CreateRenderTarget(CurrentEditLayerResource->TextureRHI, TEXT("CurrentEditLayerTexture")));
+				FRDGTextureSRVRef CurrentEditLayerTextureSRVRef = GraphBuilder.CreateSRV(FRDGTextureSRVDesc::Create(CurrentEditLayerTextureRef));
+				FRDGTextureRef PreviousEditLayersTextureRef = GraphBuilder.RegisterExternalTexture(CreateRenderTarget(PreviousEditLayersResource->TextureRHI, TEXT("PreviousEditLayersTexture")));
+				FRDGTextureSRVRef PreviousEditLayersTextureSRVRef = GraphBuilder.CreateSRV(FRDGTextureSRVDesc::Create(PreviousEditLayersTextureRef));
 
-					const int32 NumTargetLayers = MergePaintLayersParams.Num();
-					for (int32 TargetLayerIndex = 0; TargetLayerIndex < NumTargetLayers; ++TargetLayerIndex)
-					{
-						FLandscapeEditLayersWeightmapsMergeEditLayerPS::FParameters* PSParams = GraphBuilder.AllocParameters<FLandscapeEditLayersWeightmapsMergeEditLayerPS::FParameters>();
-						PSParams->RenderTargets[0] = FRenderTargetBinding(OutputTextureRef, ERenderTargetLoadAction::ENoAction, /*InMipIndex = */0, /*InArraySlice = */TargetLayerIndex);
-						PSParams->InPaintLayerIndex = TargetLayerIndex;
-						PSParams->InEditLayerPaintLayerBlendMode = static_cast<uint32>(MergePaintLayersParams[TargetLayerIndex].EditLayerPaintLayerBlendMode);
-						PSParams->InEditLayerAlpha = MergePaintLayersParams[TargetLayerIndex].EditLayerAlpha;
-						PSParams->InCurrentEditLayerWeightmaps = CurrentEditLayerTextureSRVRef;
-						PSParams->InPreviousEditLayersWeightmaps = PreviousEditLayersTextureSRVRef;
+				const int32 NumTargetLayers = MergePaintLayersParams.Num();
+				for (int32 TargetLayerIndex = 0; TargetLayerIndex < NumTargetLayers; ++TargetLayerIndex)
+				{
+					FLandscapeEditLayersWeightmapsMergeEditLayerPS::FParameters* PSParams = GraphBuilder.AllocParameters<FLandscapeEditLayersWeightmapsMergeEditLayerPS::FParameters>();
+					PSParams->RenderTargets[0] = FRenderTargetBinding(OutputTextureRef, ERenderTargetLoadAction::ENoAction, /*InMipIndex = */0, /*InArraySlice = */TargetLayerIndex);
+					PSParams->InPaintLayerIndex = TargetLayerIndex;
+					PSParams->InEditLayerPaintLayerBlendMode = static_cast<uint32>(MergePaintLayersParams[TargetLayerIndex].EditLayerPaintLayerBlendMode);
+					PSParams->InEditLayerAlpha = MergePaintLayersParams[TargetLayerIndex].EditLayerAlpha;
+					PSParams->InCurrentEditLayerWeightmaps = CurrentEditLayerTextureSRVRef;
+					PSParams->InPreviousEditLayersWeightmaps = PreviousEditLayersTextureSRVRef;
 
 						FLandscapeEditLayersWeightmapsMergeEditLayerPS::MergeEditLayerPS(
 							RDG_EVENT_NAME("MergeEditLayer(%s) -> %s", *MergePaintLayersParams[TargetLayerIndex].PaintLayerName.ToString(), *OutputResourceName), GraphBuilder, PSParams, EffectiveTextureSize);
-					}
+				}
 				};
-				
+
 			// We need to specify the final state of the external textures to prevent the graph builder from transitioning them to SRVMask (even those that end up as SRVMask at the end of this command,
 			//  because they will likely be part of another RDGCommand down the line so we need to maintain an accurate picture of every external texture ever involved in the recorded command so that 
 			//  we can set a proper access when the recorder is flushed (and the FRDGBuilder, executed) :
@@ -6432,12 +6432,12 @@ void ULandscapeHeightmapNormalsEditLayerRenderer::RenderLayer(FRenderParams& Ren
 	ALandscape* Landscape = RenderContext->GetLandscape();
 
 	checkf(RDGBuilderRecorder.IsRecording(), TEXT("ERenderFlags::RenderMode_Recorded means the command recorder should be recording at this point"));
-	checkf(RenderParams.RenderGroupTargetLayerNames.Num() == 1, TEXT("Normals should only be generated on heightmap merge, which should have 1 and only target layer"));
+	checkf(RenderParams.TargetLayerGroupLayerNames.Num() == 1, TEXT("Normals should only be generated on heightmap merge, which should have 1 and only target layer"));
 
 	RenderContext->CycleBlendRenderTargets(/*InDesiredWriteAccess = */ERHIAccess::RTV, RDGBuilderRecorder);
 	ULandscapeScratchRenderTarget* WriteRT = RenderContext->GetBlendRenderTargetWrite();
 	ULandscapeScratchRenderTarget* ReadRT = RenderContext->GetBlendRenderTargetRead();
-	ULandscapeScratchRenderTarget* ValidityRT = RenderContext->GetValidityRenderTarget(RenderParams.RenderGroupTargetLayerNames[0]);
+	ULandscapeScratchRenderTarget* ValidityRT = RenderContext->GetValidityRenderTarget(RenderParams.TargetLayerGroupLayerNames[0]);
 	check((WriteRT != nullptr) && (ReadRT != nullptr) && (ValidityRT != nullptr));
 
 	ValidityRT->TransitionTo(ERHIAccess::SRVMask, RDGBuilderRecorder);
@@ -6446,7 +6446,7 @@ void ULandscapeHeightmapNormalsEditLayerRenderer::RenderLayer(FRenderParams& Ren
 	check((NumComponentsInRenderArea.X > 0) && (NumComponentsInRenderArea.Y > 0));
 
 	auto RDGCommand =
-		[ OutputResource = WriteRT->GetRenderTarget2D()->GetResource()
+		[OutputResource = WriteRT->GetRenderTarget2D()->GetResource()
 		, OutputResourceName = WriteRT->GetDebugName()
 		, SourceResource = ReadRT->GetRenderTarget2D()->GetResource()
 		, ValidityResource = ValidityRT->GetRenderTarget2D()->GetResource()
@@ -6454,20 +6454,20 @@ void ULandscapeHeightmapNormalsEditLayerRenderer::RenderLayer(FRenderParams& Ren
 		, LandscapeGridScale = Landscape->GetRootComponent()->GetRelativeScale3D()
 		, ComponentSizeQuads = Landscape->ComponentSizeQuads
 		, NumComponentsInRenderArea](FRDGBuilder& GraphBuilder)
-		{
-			FRDGTextureRef OutputTextureRef = GraphBuilder.RegisterExternalTexture(CreateRenderTarget(OutputResource->TextureRHI, TEXT("OutputTexture")));
-			FRDGTextureRef SourceTextureRef = GraphBuilder.RegisterExternalTexture(CreateRenderTarget(SourceResource->TextureRHI, TEXT("SourceTexture")));
-			FRDGTextureRef ValidityTextureRef = GraphBuilder.RegisterExternalTexture(CreateRenderTarget(ValidityResource->TextureRHI, TEXT("ValidityTexture")));
-
-			FLandscapeEditLayersHeightmapsGenerateNormalsPS::FParameters* PSParams = GraphBuilder.AllocParameters<FLandscapeEditLayersHeightmapsGenerateNormalsPS::FParameters>();
-			PSParams->RenderTargets[0] = FRenderTargetBinding(OutputTextureRef, ERenderTargetLoadAction::ENoAction);
-			PSParams->InTextureSize = FUintVector4(EffectiveTextureSize.X, EffectiveTextureSize.Y, SourceResource->GetSizeX(), SourceResource->GetSizeY());
-			PSParams->InLandscapeGridScale = static_cast<FVector3f>(LandscapeGridScale);
-			PSParams->InComponentSizeQuads = ComponentSizeQuads;
-			PSParams->InNumComponents = FUintVector2(NumComponentsInRenderArea.X, NumComponentsInRenderArea.Y);
-			PSParams->InSourceHeightmapSampler = TStaticSamplerState<SF_Point, AM_Clamp, AM_Clamp>::GetRHI();
-			PSParams->InSourceHeightmap = GraphBuilder.CreateSRV(FRDGTextureSRVDesc::Create(SourceTextureRef));
-			PSParams->InValidityTexture = GraphBuilder.CreateSRV(FRDGTextureSRVDesc::Create(ValidityTextureRef));
+	{
+		FRDGTextureRef OutputTextureRef = GraphBuilder.RegisterExternalTexture(CreateRenderTarget(OutputResource->TextureRHI, TEXT("OutputTexture")));
+		FRDGTextureRef SourceTextureRef = GraphBuilder.RegisterExternalTexture(CreateRenderTarget(SourceResource->TextureRHI, TEXT("SourceTexture")));
+		FRDGTextureRef ValidityTextureRef = GraphBuilder.RegisterExternalTexture(CreateRenderTarget(ValidityResource->TextureRHI, TEXT("ValidityTexture")));
+
+		FLandscapeEditLayersHeightmapsGenerateNormalsPS::FParameters* PSParams = GraphBuilder.AllocParameters<FLandscapeEditLayersHeightmapsGenerateNormalsPS::FParameters>();
+		PSParams->RenderTargets[0] = FRenderTargetBinding(OutputTextureRef, ERenderTargetLoadAction::ENoAction);
+		PSParams->InTextureSize = FUintVector4(EffectiveTextureSize.X, EffectiveTextureSize.Y, SourceResource->GetSizeX(), SourceResource->GetSizeY());
+		PSParams->InLandscapeGridScale = static_cast<FVector3f>(LandscapeGridScale);
+		PSParams->InComponentSizeQuads = ComponentSizeQuads;
+		PSParams->InNumComponents = FUintVector2(NumComponentsInRenderArea.X, NumComponentsInRenderArea.Y);
+		PSParams->InSourceHeightmapSampler = TStaticSamplerState<SF_Point, AM_Clamp, AM_Clamp>::GetRHI();
+		PSParams->InSourceHeightmap = GraphBuilder.CreateSRV(FRDGTextureSRVDesc::Create(SourceTextureRef));
+		PSParams->InValidityTexture = GraphBuilder.CreateSRV(FRDGTextureSRVDesc::Create(ValidityTextureRef));
 			FLandscapeEditLayersHeightmapsGenerateNormalsPS::GenerateNormalsPS(
 				RDG_EVENT_NAME("GenerateNormals -> %s", *OutputResourceName), GraphBuilder, PSParams, EffectiveTextureSize);
 		};
@@ -6498,8 +6498,8 @@ void ULandscapeWeightmapWeightBlendedLayersRenderer::RenderLayer(FRenderParams&
 	ULandscapeScratchRenderTarget* ReadRT = RenderParams.MergeRenderContext->GetBlendRenderTargetRead();
 
 	TArray<FWeightmapPaintLayerInfo> WeightmapPaintLayerInfos;
-	WeightmapPaintLayerInfos.Reserve(RenderParams.RenderGroupTargetLayerInfos.Num());
-	Algo::Transform(RenderParams.RenderGroupTargetLayerInfos, WeightmapPaintLayerInfos, [](ULandscapeLayerInfoObject* InLayerInfo)
+	WeightmapPaintLayerInfos.Reserve(RenderParams.TargetLayerGroupLayerInfos.Num());
+	Algo::Transform(RenderParams.TargetLayerGroupLayerInfos, WeightmapPaintLayerInfos, [](ULandscapeLayerInfoObject* InLayerInfo)
 		{
 			check(InLayerInfo != nullptr); // There should only be valid layer infos at this point
 			FWeightmapPaintLayerInfo WeightmapPaintLayerInfo;
@@ -6515,36 +6515,36 @@ void ULandscapeWeightmapWeightBlendedLayersRenderer::RenderLayer(FRenderParams&
 		});
 
 	auto RDGCommand =
-		[TargetLayerNames = RenderParams.RenderGroupTargetLayerNames
+		[ TargetLayerNames = RenderParams.TargetLayerGroupLayerNames
 		, WeightmapPaintLayerInfos
 		, OutputResource = WriteRT->GetRenderTarget2DArray()->GetResource()
 		, OutputResourceName = WriteRT->GetDebugName()
 		, CurrentEditLayerResource = ReadRT->GetRenderTarget2DArray()->GetResource()
 		, EffectiveTextureSize = RenderParams.MergeRenderContext->GetCurrentRenderBatch()->GetRenderTargetResolution(/*bInWithDuplicateBorders = */false)](FRDGBuilder& GraphBuilder)
-		{
+	{
 			FRDGTextureRef OutputTextureRef = GraphBuilder.RegisterExternalTexture(CreateRenderTarget(OutputResource->TextureRHI, TEXT("OutputTexture")));
-			FRDGTextureRef CurrentEditLayerTextureRef = GraphBuilder.RegisterExternalTexture(CreateRenderTarget(CurrentEditLayerResource->TextureRHI, TEXT("CurrentEditLayerTexture")));
-			FRDGTextureSRVRef CurrentEditLayerTextureSRVRef = GraphBuilder.CreateSRV(FRDGTextureSRVDesc::Create(CurrentEditLayerTextureRef));
+		FRDGTextureRef CurrentEditLayerTextureRef = GraphBuilder.RegisterExternalTexture(CreateRenderTarget(CurrentEditLayerResource->TextureRHI, TEXT("CurrentEditLayerTexture")));
+		FRDGTextureSRVRef CurrentEditLayerTextureSRVRef = GraphBuilder.CreateSRV(FRDGTextureSRVDesc::Create(CurrentEditLayerTextureRef));
 
-			FRDGBufferRef PaintLayerInfosBuffer = CreateStructuredBuffer(GraphBuilder, TEXT("LandscapePaintLayerInfosBuffer"), WeightmapPaintLayerInfos);
-			FRDGBufferSRVRef PaintLayerInfosBufferSRV = GraphBuilder.CreateSRV(FRDGBufferSRVDesc(PaintLayerInfosBuffer));
+		FRDGBufferRef PaintLayerInfosBuffer = CreateStructuredBuffer(GraphBuilder, TEXT("LandscapePaintLayerInfosBuffer"), WeightmapPaintLayerInfos);
+		FRDGBufferSRVRef PaintLayerInfosBufferSRV = GraphBuilder.CreateSRV(FRDGBufferSRVDesc(PaintLayerInfosBuffer));
 
-			const int32 NumTargetLayers = TargetLayerNames.Num();
-			for (int32 TargetLayerIndex = 0; TargetLayerIndex < NumTargetLayers; ++TargetLayerIndex)
-			{
-				FLandscapeEditLayersWeightmapsPerformLegacyWeightBlendingPS::FParameters* PSParams = GraphBuilder.AllocParameters<FLandscapeEditLayersWeightmapsPerformLegacyWeightBlendingPS::FParameters>();
-				PSParams->RenderTargets[0] = FRenderTargetBinding(OutputTextureRef, ERenderTargetLoadAction::ENoAction, /*InMipIndex = */0, /*InArraySlice = */TargetLayerIndex);
-				PSParams->InPaintLayerIndex = TargetLayerIndex;
-				PSParams->InNumPaintLayers = NumTargetLayers;
-				PSParams->InPaintLayerInfos = PaintLayerInfosBufferSRV;
-				PSParams->InCurrentEditLayerWeightmaps = CurrentEditLayerTextureSRVRef;
+		const int32 NumTargetLayers = TargetLayerNames.Num();
+		for (int32 TargetLayerIndex = 0; TargetLayerIndex < NumTargetLayers; ++TargetLayerIndex)
+		{
+			FLandscapeEditLayersWeightmapsPerformLegacyWeightBlendingPS::FParameters* PSParams = GraphBuilder.AllocParameters<FLandscapeEditLayersWeightmapsPerformLegacyWeightBlendingPS::FParameters>();
+			PSParams->RenderTargets[0] = FRenderTargetBinding(OutputTextureRef, ERenderTargetLoadAction::ENoAction, /*InMipIndex = */0, /*InArraySlice = */TargetLayerIndex);
+			PSParams->InPaintLayerIndex = TargetLayerIndex;
+			PSParams->InNumPaintLayers = NumTargetLayers;
+			PSParams->InPaintLayerInfos = PaintLayerInfosBufferSRV;
+			PSParams->InCurrentEditLayerWeightmaps = CurrentEditLayerTextureSRVRef;
 
 				FLandscapeEditLayersWeightmapsPerformLegacyWeightBlendingPS::PerformLegacyWeightBlendingPS(
 					RDG_EVENT_NAME("LegacyWeightBlending(%s) -> %s", *TargetLayerNames[TargetLayerIndex].ToString(), *OutputResourceName),
 					GraphBuilder, PSParams, EffectiveTextureSize);
-			}
+		}
 		};
-	
+
 	// We need to specify the final state of the external textures to prevent the graph builder from transitioning them to SRVMask (even those that end up as SRVMask at the end of this command,
 	//  because they will likely be part of another RDGCommand down the line so we need to maintain an accurate picture of every external texture ever involved in the recorded command so that 
 	//  we can set a proper access when the recorder is flushed (and the FRDGBuilder, executed) :
@@ -6580,9 +6580,9 @@ UE::Landscape::EditLayers::FMergeRenderContext ALandscape::PrepareEditLayersMerg
 	if (!InMergeContext.IsHeightmapMerge())
 	{
 		for (FName TargetLayerName : InMergeRenderParams.WeightmapLayerNames)
-	{
+		{
 			if (!InMergeContext.IsValidTargetLayerName(TargetLayerName))
-		{ 
+			{
 				UE_LOG(LogLandscape, Warning, TEXT("Target layer \"%s\" was requested by the merge but is invalid (missing its layer info object asset). It will be ignored"), *TargetLayerName.ToString());
 			}
 		}
@@ -6633,7 +6633,7 @@ UE::Landscape::EditLayers::FMergeRenderContext ALandscape::PrepareEditLayersMerg
 
 	// Within each render batch, elements can be processed group by group. For heightmap/visibility, there's only one such group. For weightmaps, there's one group per list of weightmaps
 	//  that need to be processed together for weight-blending. Each group is composed of a list of (weightmap layer) names (it's only a debug name in the case of heightmaps)
-	TArray<TBitArray<>> FinalRenderGroups;
+	TArray<TBitArray<>> FinalTargetLayerGroups;
 	if (MergeRenderContext.IsHeightmapMerge())
 	{
 		// All candidates are valid in heightmap merge, the ones that don't affect heightmap have already been filtered from CandidateEditLayerRendererStates
@@ -6641,46 +6641,46 @@ UE::Landscape::EditLayers::FMergeRenderContext ALandscape::PrepareEditLayersMerg
 
 		// Only one group in the case of heightmap: 
 		const FName TargetLayerName = FName("Height");
-		FinalRenderGroups = { TBitArray<>(true, 1) };
+		FinalTargetLayerGroups = { TBitArray<>(true, 1) };
 	}
 	else
 	{
-	// First, let's work out the weightmaps inter-dependencies (i.e. horizontal dependencies) : weight-blending requires some weightmaps to be processed together: 
-		TRACE_CPUPROFILER_EVENT_SCOPE(PrepareRenderGroups);
+		// First, let's work out the weightmaps inter-dependencies (i.e. horizontal dependencies) : weight-blending requires some weightmaps to be processed together: 
+		TRACE_CPUPROFILER_EVENT_SCOPE(PrepareTargetLayerGroups);
 
-		TArray<TArray<TBitArray<>>> RelevantRenderGroupsPerRenderer;
+		TArray<TArray<TBitArray<>>> RelevantTargetLayerGroupsPerRenderer;
 		TArray<TBitArray<>> RelevantTargetLayerBitIndicesPerRenderer;
-		RelevantRenderGroupsPerRenderer.Reserve(CandidateEditLayerRendererStates.Num());
+		RelevantTargetLayerGroupsPerRenderer.Reserve(CandidateEditLayerRendererStates.Num());
 		RelevantTargetLayerBitIndicesPerRenderer.Reserve(CandidateEditLayerRendererStates.Num());
-		// Iterate through all renderers to find out what render group they render (only keep those that are relevant for the current state and request) : 
+		// Iterate through all renderers to find out what target layer group they render (only keep those that are relevant for the current state and request) : 
 		for (const FEditLayerRendererState& EditLayerRendererState : CandidateEditLayerRendererStates)
 		{
 			TBitArray<> RendererStateEnabledTargetLayerBitIndices(EditLayerRendererState.GetActiveTargetWeightmapBitIndices());
 			// Retain only the target layer names that are enabled and valid :
 			RendererStateEnabledTargetLayerBitIndices.CombineWithBitwiseAND(MergeRenderContext.ValidTargetLayerBitIndices, EBitwiseOperatorFlags::MinSize);
 
-			// List of all supported render groups for this renderer :
-			TArray<TBitArray<>> RelevantRenderGroupsForThisRenderer;
-			for (const TBitArray<>& WeightmapRenderGroup : EditLayerRendererState.GetWeightmapRenderGroups())
+			// List of all supported target layer groups for this renderer :
+			TArray<TBitArray<>> RelevantTargetLayerGroupsForThisRenderer;
+			for (const TBitArray<>& TargetLayerGroup : EditLayerRendererState.GetTargetLayerGroups())
 			{
 				// Only retain the target layers that have passed the previous test (valid and enabled):
-				TBitArray<> RenderGroupBitIndices = WeightmapRenderGroup;
-				RenderGroupBitIndices.CombineWithBitwiseAND(RendererStateEnabledTargetLayerBitIndices, EBitwiseOperatorFlags::MinSize);
-				if (RenderGroupBitIndices.Find(true) != INDEX_NONE)
+				TBitArray<> FilteredTargetLayerGroupBitIndices = TargetLayerGroup;
+				FilteredTargetLayerGroupBitIndices.CombineWithBitwiseAND(RendererStateEnabledTargetLayerBitIndices, EBitwiseOperatorFlags::MinSize);
+				if (FilteredTargetLayerGroupBitIndices.Find(true) != INDEX_NONE)
 				{
-					checkf(Algo::NoneOf(RelevantRenderGroupsForThisRenderer, [&RenderGroupBitIndices](const TBitArray<>& InOtherRenderGroupBitIndices) 
+					checkf(Algo::NoneOf(RelevantTargetLayerGroupsForThisRenderer, [&FilteredTargetLayerGroupBitIndices](const TBitArray<>& InOtherGroupBitIndices)
 						{ 
-							return (TBitArray<>::BitwiseAND(InOtherRenderGroupBitIndices, RenderGroupBitIndices, EBitwiseOperatorFlags::MinSize).Find(true) != INDEX_NONE);
+							return (TBitArray<>::BitwiseAND(InOtherGroupBitIndices, FilteredTargetLayerGroupBitIndices, EBitwiseOperatorFlags::MinSize).Find(true) != INDEX_NONE);
 						}), 
-						TEXT("All of the target layers returned by the renderer must belong to 1 render group of this renderer and 1 only"));
-					RelevantRenderGroupsForThisRenderer.Add(RenderGroupBitIndices);
+						TEXT("All of the target layers returned by the renderer must belong to 1 target layer group of this renderer and 1 only"));
+					RelevantTargetLayerGroupsForThisRenderer.Add(FilteredTargetLayerGroupBitIndices);
 				}
 			}
 
-			RelevantRenderGroupsPerRenderer.Add(RelevantRenderGroupsForThisRenderer);
+			RelevantTargetLayerGroupsPerRenderer.Add(RelevantTargetLayerGroupsForThisRenderer);
 			RelevantTargetLayerBitIndicesPerRenderer.Add(RendererStateEnabledTargetLayerBitIndices);
 		}
-		check(CandidateEditLayerRendererStates.Num() == RelevantRenderGroupsPerRenderer.Num());
+		check(CandidateEditLayerRendererStates.Num() == RelevantTargetLayerGroupsPerRenderer.Num());
 		check(CandidateEditLayerRendererStates.Num() == RelevantTargetLayerBitIndicesPerRenderer.Num());
 
 		const int32 CandidateNumRenderers = CandidateEditLayerRendererStates.Num();
@@ -6693,52 +6693,52 @@ UE::Landscape::EditLayers::FMergeRenderContext ALandscape::PrepareEditLayersMerg
 			return MergeRenderContext;
 		}
 
-		FinalRenderGroups.Reserve(MergeRenderContext.ValidTargetLayerBitIndices.CountSetBits());
-		// Start with minimal render groups : one per requested target layer : 
+		FinalTargetLayerGroups.Reserve(MergeRenderContext.ValidTargetLayerBitIndices.CountSetBits());
+		// Start with minimal target layer groups : one per requested target layer : 
 		for (TConstSetBitIterator It(RequestedTargetLayerBitIndices); It; ++It)
 		{
-			TBitArray<> RenderGroup(false, NumAllTargetLayerNames);
-			RenderGroup[It.GetIndex()] = true;
-			FinalRenderGroups.Add(MoveTemp(RenderGroup));
+			TBitArray<> TargetLayerGroup(false, NumAllTargetLayerNames);
+			TargetLayerGroup[It.GetIndex()] = true;
+			FinalTargetLayerGroups.Add(MoveTemp(TargetLayerGroup));
 		}
 
-		// Then iterate in reverse order on renderers to trace the dependency of each of their render groups towards one another and move target layers from one render group to another as we discover new dependencies:
+		// Then iterate in reverse order on renderers to trace the dependency of each of their target layer groups towards one another and move target layers from one group to another as we discover new dependencies:
 		MergeRenderContext.FinalTargetLayerBitIndices = RequestedTargetLayerBitIndices;
 		{
-			TRACE_CPUPROFILER_EVENT_SCOPE(AnalyzeRenderersForRenderGroups);
+			TRACE_CPUPROFILER_EVENT_SCOPE(AnalyzeRenderersForTargetLayerGroups);
 
 			for (int32 CurrentRendererIndex = CandidateNumRenderers - 1; CurrentRendererIndex >= 0; --CurrentRendererIndex)
 			{
-				const TArray<TBitArray<>>& RendererRenderGroups = RelevantRenderGroupsPerRenderer[CurrentRendererIndex];
-				const int32 NumRenderGroups = RendererRenderGroups.Num();
-				for (int32 RenderGroupIndex = 0; RenderGroupIndex < NumRenderGroups; ++RenderGroupIndex)
+				const TArray<TBitArray<>>& RendererTargetLayerGroups = RelevantTargetLayerGroupsPerRenderer[CurrentRendererIndex];
+				const int32 NumTargetLayerGroups = RendererTargetLayerGroups.Num();
+				for (int32 TargetLayerGroupIndex = 0; TargetLayerGroupIndex < NumTargetLayerGroups; ++TargetLayerGroupIndex)
 				{
-					const TBitArray<>& RenderGroup = RendererRenderGroups[RenderGroupIndex];
+					const TBitArray<>& TargetLayerGroup = RendererTargetLayerGroups[TargetLayerGroupIndex];
 
-					TBitArray<> NewMergedRenderGroup(false, NumAllTargetLayerNames);
-					// In all of the final render groups, find the ones that have a layer in common with this render group and merge them all into a single new one
-					FinalRenderGroups.RemoveAllSwap([&NewMergedRenderGroup, &RenderGroup](const TBitArray<>& InFinalRenderGroup)
+					TBitArray<> NewMergedTargetLayerGroup(false, NumAllTargetLayerNames);
+					// In all of the final target layer groups, find the ones that have a layer in common with this target layer group and merge them all into a single new one
+					FinalTargetLayerGroups.RemoveAllSwap([&NewMergedTargetLayerGroup, &TargetLayerGroup](const TBitArray<>& InFinalTargetLayerGroup)
 						{
-							const bool bShouldMerge = (InFinalRenderGroup != RenderGroup)
-								&& (TBitArray<>::BitwiseAND(InFinalRenderGroup, RenderGroup, EBitwiseOperatorFlags::MinSize).Find(true) != INDEX_NONE);
+							const bool bShouldMerge = (InFinalTargetLayerGroup != TargetLayerGroup)
+								&& (TBitArray<>::BitwiseAND(InFinalTargetLayerGroup, TargetLayerGroup, EBitwiseOperatorFlags::MinSize).Find(true) != INDEX_NONE);
 							if (bShouldMerge)
 							{
-								NewMergedRenderGroup.CombineWithBitwiseOR(InFinalRenderGroup, EBitwiseOperatorFlags::MinSize);
+								NewMergedTargetLayerGroup.CombineWithBitwiseOR(InFinalTargetLayerGroup, EBitwiseOperatorFlags::MinSize);
 							}
 							return bShouldMerge;
 						});
-					// Now add it back to the list of final render groups if it's valid, so that it can be merged by the next renderer if required :
-					if (NewMergedRenderGroup.CountSetBits() > 0)
+					// Now add it back to the list of final target layer groups if it's valid, so that it can be merged by the next renderer if required :
+					if (NewMergedTargetLayerGroup.CountSetBits() > 0)
 					{
-						FinalRenderGroups.Add(NewMergedRenderGroup);
-						MergeRenderContext.FinalTargetLayerBitIndices.CombineWithBitwiseOR(NewMergedRenderGroup, EBitwiseOperatorFlags::MinSize);
+						FinalTargetLayerGroups.Add(NewMergedTargetLayerGroup);
+						MergeRenderContext.FinalTargetLayerBitIndices.CombineWithBitwiseOR(NewMergedTargetLayerGroup, EBitwiseOperatorFlags::MinSize);
 					}
 				}
 			}
 		}
 
 		{
-			TRACE_CPUPROFILER_EVENT_SCOPE(FinalizeRenderGroups);
+			TRACE_CPUPROFILER_EVENT_SCOPE(FinalizeTargetLayerGroups);
 
 			// Now that we have the list of all target layers required for rendering the requested ones, we trim the list of edit layer renderers that just don't do anything with those:
 			for (int32 RendererIndex = 0; RendererIndex < CandidateNumRenderers; ++RendererIndex)
@@ -6755,7 +6755,7 @@ UE::Landscape::EditLayers::FMergeRenderContext ALandscape::PrepareEditLayersMerg
 			// When 4 weightmaps are packed in 1 RGBA channel, we cannot guarantee that weightmaps will be fully resolvable when a group is done, since there's no guarantee that a given component's 4 
 			//  allocations will be contained in that group. So we create a single group instead, that contains all weightmap layers. This consumes more memory since we need 3 texture arrays for a batch and 
 			//  the number of weightmaps in a group is the number of slices of the array : 
-			FinalRenderGroups = { MergeRenderContext.FinalTargetLayerBitIndices };
+			FinalTargetLayerGroups = { MergeRenderContext.FinalTargetLayerBitIndices };
 #endif // !SUPPORTS_LANDSCAPE_EDITORONLY_UBER_MATERIAL
 		}
 	}
@@ -7154,7 +7154,7 @@ UE::Landscape::EditLayers::FMergeRenderContext ALandscape::PrepareEditLayersMerg
 			}
 
 			// Now, we have all the info to build our list of successive render steps : process group by group :
-			for (const TBitArray<>& RenderGroup : FinalRenderGroups)
+			for (const TBitArray<>& TargetLayerGroup : FinalTargetLayerGroups)
 			{
 				ERenderFlags LastRenderMode = ERenderFlags::None;
 
@@ -7175,32 +7175,32 @@ UE::Landscape::EditLayers::FMergeRenderContext ALandscape::PrepareEditLayersMerg
 						{
 							if (LastRenderMode != ERenderFlags::RenderMode_Recorded)
 							{
-								MergeRenderBatch.RenderSteps.Emplace(FMergeRenderStep::EType::BeginRenderCommandRecorder, DummyRendererState, RenderGroup, MergeRenderBatch.ComponentsToRender.Array());
+								MergeRenderBatch.RenderSteps.Emplace(FMergeRenderStep::EType::BeginRenderCommandRecorder, DummyRendererState, TargetLayerGroup, MergeRenderBatch.ComponentsToRender.Array());
 							}
 						}
 						else if (LastRenderMode == ERenderFlags::RenderMode_Recorded)
 						{
-							MergeRenderBatch.RenderSteps.Emplace(FMergeRenderStep::EType::EndRenderCommandRecorder, DummyRendererState, RenderGroup, MergeRenderBatch.ComponentsToRender.Array());
+							MergeRenderBatch.RenderSteps.Emplace(FMergeRenderStep::EType::EndRenderCommandRecorder, DummyRendererState, TargetLayerGroup, MergeRenderBatch.ComponentsToRender.Array());
 						}
 						LastRenderMode = CurrentRenderMode;
 
 						// Render the content of this render layer : 
 						FMergeRenderStep::EType RenderStepType = (CurrentRenderMode == ERenderFlags::RenderMode_Immediate)
 							? FMergeRenderStep::EType::RenderLayer_Immediate : FMergeRenderStep::EType::RenderLayer_Recorded;
-						MergeRenderBatch.RenderSteps.Emplace(RenderStepType, EditLayerRendererRenderInfo.RendererState, RenderGroup, ComponentsToRender);
+						MergeRenderBatch.RenderSteps.Emplace(RenderStepType, EditLayerRendererRenderInfo.RendererState, TargetLayerGroup, ComponentsToRender);
 					}
 				}
 
 				// Insert a final render sequence end step if we ended on a render on the render thread :
 				if (LastRenderMode == ERenderFlags::RenderMode_Recorded)
 				{
-					MergeRenderBatch.RenderSteps.Emplace(FMergeRenderStep::EType::EndRenderCommandRecorder, DummyRendererState, RenderGroup, MergeRenderBatch.ComponentsToRender.Array());
+					MergeRenderBatch.RenderSteps.Emplace(FMergeRenderStep::EType::EndRenderCommandRecorder, DummyRendererState, TargetLayerGroup, MergeRenderBatch.ComponentsToRender.Array());
 				}
 
 				// Finally, signal the group as done for this batch
-				MergeRenderBatch.RenderSteps.Emplace(FMergeRenderStep::EType::SignalBatchMergeGroupDone, DummyRendererState, RenderGroup, MergeRenderBatch.ComponentsToRender.Array());
+				MergeRenderBatch.RenderSteps.Emplace(FMergeRenderStep::EType::SignalBatchMergeGroupDone, DummyRendererState, TargetLayerGroup, MergeRenderBatch.ComponentsToRender.Array());
 
-				MergeRenderContext.MaxNeededNumSlices = FMath::Max(MergeRenderContext.MaxNeededNumSlices, RenderGroup.CountSetBits());
+				MergeRenderContext.MaxNeededNumSlices = FMath::Max(MergeRenderContext.MaxNeededNumSlices, TargetLayerGroup.CountSetBits());
 			}
 		}
 
@@ -10547,7 +10547,7 @@ int32 ALandscape::PerformLayersWeightmapsBatchedMerge(FUpdateLayersContentContex
 						{
 							check((AllocationInfo.WeightmapTextureChannel >= 0) && (AllocationInfo.WeightmapTextureChannel < 4));
 							checkf(((WeightmapResolveInfo->ChannelMask & (1 << AllocationInfo.WeightmapTextureChannel)) == 0), TEXT("This channel has already been resolved, it shouldn't happen, it would mean that 2 allocations are using the same channel"));
-							int32 SliceIndex = InParams.RenderGroupTargetLayerNames.IndexOfByKey(AllocationInfo.GetLayerName());
+							int32 SliceIndex = InParams.TargetLayerGroupLayerNames.IndexOfByKey(AllocationInfo.GetLayerName());
 							checkf(SliceIndex != INDEX_NONE, TEXT("Couldn't find %s in the list of weightmaps that have been produced"), *AllocationInfo.GetLayerName().ToString());
 							WeightmapResolveInfo->SourceSliceIndexPerChannel[AllocationInfo.WeightmapTextureChannel] = SliceIndex;
 							WeightmapResolveInfo->SourceSubsectionRectsPerChannel[AllocationInfo.WeightmapTextureChannel] = SourceSubsectionRects;
diff --git a/Engine/Source/Runtime/Landscape/Public/LandscapeBlueprintBrushBase.h b/Engine/Source/Runtime/Landscape/Public/LandscapeBlueprintBrushBase.h
--- a/Engine/Source/Runtime/Landscape/Public/LandscapeBlueprintBrushBase.h
+++ b/Engine/Source/Runtime/Landscape/Public/LandscapeBlueprintBrushBase.h
@@ -112,8 +112,8 @@ public:
 	//~ Begin ILandscapeEditLayerRenderer implementation
 	LANDSCAPE_EDIT_LAYERS_BATCHED_MERGE_EXPERIMENTAL
 	LANDSCAPE_API virtual void GetRendererStateInfo(const UE::Landscape::EditLayers::FMergeContext* InMergeContext,
-		UE::Landscape::EditLayers::FEditLayerTargetTypeState& OutSupportedTargetTypeState, UE::Landscape::EditLayers::FEditLayerTargetTypeState& OutEnabledTargetTypeState, TArray<TBitArray<>>& OutRenderGroups) const override;
-	LANDSCAPE_EDIT_LAYERS_BATCHED_MERGE_EXPERIMENTAL	
+		UE::Landscape::EditLayers::FEditLayerTargetTypeState& OutSupportedTargetTypeState, UE::Landscape::EditLayers::FEditLayerTargetTypeState& OutEnabledTargetTypeState, TArray<TBitArray<>>& OutTargetLayerGroups) const override;
+	LANDSCAPE_EDIT_LAYERS_BATCHED_MERGE_EXPERIMENTAL
 	virtual UE::Landscape::EditLayers::ERenderFlags GetRenderFlags(const UE::Landscape::EditLayers::FMergeContext* InMergeContext) const override { return UE::Landscape::EditLayers::ERenderFlags::RenderMode_Immediate; }
 	LANDSCAPE_EDIT_LAYERS_BATCHED_MERGE_EXPERIMENTAL
 	LANDSCAPE_API virtual TArray<UE::Landscape::EditLayers::FEditLayerRenderItem> GetRenderItems(const UE::Landscape::EditLayers::FMergeContext* InMergeContext) const override;
diff --git a/Engine/Source/Runtime/Landscape/Public/LandscapeEditLayerRenderer.h b/Engine/Source/Runtime/Landscape/Public/LandscapeEditLayerRenderer.h
--- a/Engine/Source/Runtime/Landscape/Public/LandscapeEditLayerRenderer.h
+++ b/Engine/Source/Runtime/Landscape/Public/LandscapeEditLayerRenderer.h
@@ -466,7 +466,7 @@ private:
  *  These states are provided by IEditLayerRendererProvider in order to describe both what the renderer can do and what it currently does by default. e.g. a disabled edit layer supports rendering heightmaps but
  *  its enabled state for the heightmap target type is false. This way, the user can selectively enable it at merge time without altering the entire landscape's state (i.e. just for the purpose of a specific merge render)
  *  A target type must be both supported and enabled on a given edit layer renderer in order for this renderer to render anything. 
- *  It also describes the render groups this renderer needs when rendering its weightmap (i.e. which weightmap needs to be rendered with which weightmaps : e.g. for weight blending)
+ *  It also describes the target layer groups this renderer needs when rendering its weightmap (i.e. which weightmap needs to be rendered with which weightmaps : e.g. for weight blending)
  */
 class FEditLayerRendererState
 {
@@ -588,10 +588,10 @@ public:
 	LANDSCAPE_EDIT_LAYERS_BATCHED_MERGE_EXPERIMENTAL
 	LANDSCAPE_API TBitArray<> GetActiveTargetWeightmapBitIndices() const;
 
-	/** Returns the render groups associated with this renderer. A render group is a set of target layers (weightmaps) that depend on one another in order to produce the output target layers.
+	/** Returns the target layer groups associated with this renderer. A target layer group is a set of target layers (weightmaps) that depend on one another in order to produce the output target layers.
 	 This allows to implement "horizontal blending", where weightmaps can be blended with one another at each step of the landscape edit layers merge algorithm */
 	LANDSCAPE_EDIT_LAYERS_BATCHED_MERGE_EXPERIMENTAL
-	LANDSCAPE_API TArray<TBitArray<>> GetWeightmapRenderGroups() const { return RenderGroups; }
+	LANDSCAPE_API TArray<TBitArray<>> GetTargetLayerGroups() const { return TargetLayerGroups; }
 
 private:
 	void UpdateActiveTargetTypeState();
@@ -615,10 +615,10 @@ private:
 	/** Intersection of SupportedTargetTypeState and EnabledTargetTypeState (defines what is both supported and enabled on this renderer) */
 	FEditLayerTargetTypeState ActiveTargetTypeState;
 
-	/* List of groups of target layers that this renderer requires to be rendered together. All target layers listed in SupportedTargetTypeState must belong to 1 (and 1 only) render group.
-	 *  Each render group is a bit array for which each bit corresponds to an entry in FMergeContext's AllTargetLayerNames
+	/* List of groups of target layers that this renderer requires to be rendered together. All target layers listed in SupportedTargetTypeState must belong to 1 (and 1 only) target layer group.
+	 *  Each target layer group is a bit array for which each bit corresponds to an entry in FMergeContext's AllTargetLayerNames
 	 */
-	TArray<TBitArray<>> RenderGroups;
+	TArray<TBitArray<>> TargetLayerGroups;
 };
 
 
@@ -655,17 +655,17 @@ struct FMergeRenderStep
 		Invalid,
 
 		BeginRenderCommandRecorder, // Initiates a render command recorder that will batch one or several RenderLayer_Recorded operations, runs on the game thread
-		RenderLayer_Recorded, // Performs the rendering of a render group on an edit layer on a given world region (i.e. in a batch), runs on the game thread but this is where render operation (a lambda that will run on the render thread) can be enqueued on the command recorder initiated by BeginRenderCommandRecorder
+		RenderLayer_Recorded, // Performs the rendering of a target layer group on an edit layer on a given world region (i.e. in a batch), runs on the game thread but this is where render operation (a lambda that will run on the render thread) can be enqueued on the command recorder initiated by BeginRenderCommandRecorder
 		EndRenderCommandRecorder, // Ends the render command recorder initiated by the last BeginRenderCommandRecorder, runs on the game thread but this is where the render command executing the recorded operations from (several) RenderLayer_Recorded step(s) will be pushed (on the render thread, this is where the FRDBuilder is created and executed)
 
-		RenderLayer_Immediate, // Perform the rendering of a render group on an edit layer on a given world region (i.e. in a batch), runs on the game thread
-		SignalBatchMergeGroupDone, // Final step when rendering a render group on a given world region (i.e. in a batch) : runs on the game thread and allows to retrieve the result of the merge and do something with it (e.g. resolve the final textures)
+		RenderLayer_Immediate, // Perform the rendering of a target layer group on an edit layer on a given world region (i.e. in a batch), runs on the game thread
+		SignalBatchMergeGroupDone, // Final step when rendering a target layer group on a given world region (i.e. in a batch) : runs on the game thread and allows to retrieve the result of the merge and do something with it (e.g. resolve the final textures)
 	};
 
-	FMergeRenderStep(EType InType, const FEditLayerRendererState& InRendererState, const TBitArray<>& InRenderGroupBitIndices, const TArrayView<ULandscapeComponent*>& InComponentsToRender)
+	FMergeRenderStep(EType InType, const FEditLayerRendererState& InRendererState, const TBitArray<>& InTargetLayerGroupBitIndices, const TArrayView<ULandscapeComponent*>& InComponentsToRender)
 		: Type(InType)
 		, RendererState(InRendererState)
-		, RenderGroupBitIndices(InRenderGroupBitIndices)
+		, TargetLayerGroupBitIndices(InTargetLayerGroupBitIndices)
 		, ComponentsToRender(InComponentsToRender)
 	{}
 
@@ -676,17 +676,20 @@ struct FMergeRenderStep
 	FEditLayerRendererState RendererState;
 	
 	/** List of target layers being involved in this step. Each bit in that bit array corresponds to an entry in FMergeContext's AllTargetLayerNames */
-	TBitArray<> RenderGroupBitIndices;
+	TBitArray<> TargetLayerGroupBitIndices;
 
 	/** List of components involved in this step */
 	TArray<ULandscapeComponent*> ComponentsToRender;
+
+	UE_DEPRECATED(5.6, "Renamed : use TargetLayerGroupBitIndices")
+	TBitArray<> RenderGroupBitIndices;
 };
 
 
 // ----------------------------------------------------------------------------------
 
 /** 
- * Defines an individual render batch when merging the landscape. A batch corresponds to a render group on a world's region, i.e. a set of weightmaps (or the heightmap) to render on a 
+ * Defines an individual render batch when merging the landscape. A batch corresponds to a target layer group on a world's region, i.e. a set of weightmaps (or the heightmap) to render on a 
  *  portion of the world. Each batch is composed of a series of render steps. 
  */
 struct FMergeRenderBatch
@@ -711,7 +714,7 @@ struct FMergeRenderBatch
 	/** List of all components involved in this batch */
 	TSet<ULandscapeComponent*> ComponentsToRender;
 
-	/** List of all target layers being rendered in this batch (i.e. bitwise OR of all of the render steps' RenderGroupBitIndices). Each bit in that bit array corresponds to an entry in FMergeRenderContext's AllTargetLayerNames  */
+	/** List of all target layers being rendered in this batch (i.e. bitwise OR of all of the render steps' TargetLayerGroupBitIndices). Each bit in that bit array corresponds to an entry in FMergeRenderContext's AllTargetLayerNames  */
 	TBitArray<> TargetLayerBitIndices;
 
 	/** List of components involved in this batch and the target layers they're writing to (redundant with ComponentsToRender but we keep the latter for convenience)
@@ -818,30 +821,35 @@ public:
 	{
 		FOnRenderBatchTargetGroupDoneParams(FMergeRenderContext* InMergeRenderContext, 
 			const FMergeRenderBatch& InRenderBatch, 
-			const TArrayView<FName>& InRenderGroupTargetLayerNames,
-			const TArrayView<ULandscapeLayerInfoObject*>& InRenderGroupTargetLayerInfos, 
+			const TArrayView<FName>& InTargetLayerGroupLayerNames,
+			const TArrayView<ULandscapeLayerInfoObject*>& InTargetLayerGroupLayerInfos,
 			const TArrayView<FComponentMergeRenderInfo>& InSortedComponentMergeRenderInfos)
 			: MergeRenderContext(InMergeRenderContext)
 			, RenderBatch(&InRenderBatch)
-			, RenderGroupTargetLayerNames(InRenderGroupTargetLayerNames)
-			, RenderGroupTargetLayerInfos(InRenderGroupTargetLayerInfos)
+			, TargetLayerGroupLayerNames(InTargetLayerGroupLayerNames)
+			, TargetLayerGroupLayerInfos(InTargetLayerGroupLayerInfos)
 			, SortedComponentMergeRenderInfos(InSortedComponentMergeRenderInfos)
 		{}
 
 		/** Render context : this is still active in this step and can be used for doing additional renders in the blend render targets, etc. */
 		FMergeRenderContext* MergeRenderContext = nullptr;
 
-		/** Batch to that was just rendered for this render group */
+		/** Batch to that was just rendered for this target layer group */
 		const FMergeRenderBatch* RenderBatch = nullptr;
 
 		/** List of target layers being involved in this step */
-		TArray<FName> RenderGroupTargetLayerNames;
+		TArray<FName> TargetLayerGroupLayerNames;
 
-		/** List of target layer info objects being involved in this step (same size as RenderGroupTargetLayerNames) */
-		TArray<ULandscapeLayerInfoObject*> RenderGroupTargetLayerInfos;
+		/** List of target layer info objects being involved in this step (same size as TargetLayerGroupLayerNames) */
+		TArray<ULandscapeLayerInfoObject*> TargetLayerGroupLayerInfos;
 
 		/** Additional info about the components that have been processed in this batch render */
 		const TArray<FComponentMergeRenderInfo> SortedComponentMergeRenderInfos;
+
+		UE_DEPRECATED(5.6, "Renamed : use TargetLayerGroupLayerNames")
+		TArray<FName> RenderGroupTargetLayerNames;
+		UE_DEPRECATED(5.6, "Renamed : use TargetLayerGroupLayerInfos")
+		TArray<ULandscapeLayerInfoObject*> RenderGroupTargetLayerInfos;
 	};
 	void Render(TFunction<void(const FOnRenderBatchTargetGroupDoneParams& /*InParams*/, UE::Landscape::FRDGBuilderRecorder& /*RDGBuilderRecorder*/)> OnRenderBatchTargetGroupDone);
 
@@ -915,7 +923,7 @@ private:
 	/** Maximum resolution needed by a given batch in this context (means we won't ever need more than this size of a render target during the whole merge */
 	FIntPoint MaxNeededResolution = FIntPoint(ForceInit);
 
-	/** Maximum number of slices needed by a given batch / render group in this context (means we won't ever need more than this number of slices in a given render target texture array during the whole merge */
+	/** Maximum number of slices needed by a given batch / target layer group in this context (means we won't ever need more than this number of slices in a given render target texture array during the whole merge */
 	int32 MaxNeededNumSlices = 0;
 
 	/** Successive batches of components being processed by this context. Each batch should be self-contained so that we won't ever need to keep more than one in memory (VRAM)
@@ -1168,14 +1176,14 @@ public:
 	struct FRenderParams
 	{
 		FRenderParams(UE::Landscape::EditLayers::FMergeRenderContext* InMergeRenderContext, 
-			const TArrayView<FName>& InRenderGroupTargetLayerNames,
-			const TArrayView<ULandscapeLayerInfoObject*>& InRenderGroupTargetLayerInfos,
+			const TArrayView<FName>& InTargetLayerGroupLayerNames,
+			const TArrayView<ULandscapeLayerInfoObject*>& InTargetLayerGroupLayerInfos,
 			const UE::Landscape::EditLayers::FEditLayerRendererState& InRendererState, 
 			const TArrayView<UE::Landscape::EditLayers::FComponentMergeRenderInfo>& InSortedComponentMergeRenderInfos, 
 			const FTransform& InRenderAreaWorldTransform, const FIntRect& InRenderAreaSectionRect)
 			: MergeRenderContext(InMergeRenderContext)
-			, RenderGroupTargetLayerNames(InRenderGroupTargetLayerNames)
-			, RenderGroupTargetLayerInfos(InRenderGroupTargetLayerInfos)
+			, TargetLayerGroupLayerNames(InTargetLayerGroupLayerNames)
+			, TargetLayerGroupLayerInfos(InTargetLayerGroupLayerInfos)
 			, RendererState(InRendererState)
 			, SortedComponentMergeRenderInfos(InSortedComponentMergeRenderInfos)
 			, RenderAreaWorldTransform(InRenderAreaWorldTransform)
@@ -1186,14 +1194,14 @@ public:
 		UE::Landscape::EditLayers::FMergeRenderContext* MergeRenderContext = nullptr;
 
 		/** List of target layers being involved in this step */
-		TArray<FName> RenderGroupTargetLayerNames;
+		TArray<FName> TargetLayerGroupLayerNames;
 
-		/** List of target layer info objects being involved in this step (same size as RenderGroupTargetLayerNames) */
-		TArray<ULandscapeLayerInfoObject*> RenderGroupTargetLayerInfos;
+		/** List of target layer info objects being involved in this step (same size as TargetLayerGroupLayerNames) */
+		TArray<ULandscapeLayerInfoObject*> TargetLayerGroupLayerInfos;
 
-		/** Full state for the renderer involved in this step. This allows to retrieve the exact state of this renderer (e.g. enabled weightmaps, which can be different than the render group, in that 
+		/** Full state for the renderer involved in this step. This allows to retrieve the exact state of this renderer (e.g. enabled weightmaps, which can be different than the target layer group, in that 
 		 target layers A, B and C might belong to the same group but this renderer actually only has A enabled). This is therefore the renderer's responsibility to check that a given target layer from the 
-		 render group is effectively enabled. */
+		 target layer group is effectively enabled. */
 		UE::Landscape::EditLayers::FEditLayerRendererState RendererState;
 
 		/** List of components (with additional info) to render */
@@ -1205,6 +1213,11 @@ public:
 
 		/** SectionRect (i.e. landscape vertex coordinates, in landscape space) that corresponds to this render area */
 		FIntRect RenderAreaSectionRect;
+
+		UE_DEPRECATED(5.6, "Renamed : use TargetLayerGroupLayerNames")
+		TArray<FName> RenderGroupTargetLayerNames;
+		UE_DEPRECATED(5.6, "Renamed : use TargetLayerGroupLayerInfos")
+		TArray<ULandscapeLayerInfoObject*> RenderGroupTargetLayerInfos;
 	};
 
 	/**
@@ -1222,14 +1235,14 @@ public:
 	 *  A target type must be both supported and enabled in order to have this renderer affect it.
 	 *  The "enabled" state can be changed at will by the user (e.g. to temporarily disable a given edit layer) : see FEditLayerRendererState
 	 * 
-	 * @param OutRenderGroups List of groups of target layers that this renderer requires to be rendered together.
+	 * @param OutTargetLayerGroups List of groups of target layers that this renderer requires to be rendered together.
 	 *  This allows it to perform horizontal blending (i.e.adjust the weights of the targeted weightmaps wrt one another).
-	 *  Depending on the other renderer's needs, the final render groups might contain more layers than was requested by a given renderer. This only means that more layers will be processed together
+	 *  Depending on the other renderer's needs, the final target layer groups might contain more layers than was requested by a given renderer. This only means that more layers will be processed together
 	 *  and if this renderer doesn't act on one of these layers, it will simply do nothing with it in its RenderLayer function
 	 */
 	LANDSCAPE_EDIT_LAYERS_BATCHED_MERGE_EXPERIMENTAL
 	LANDSCAPE_API virtual void GetRendererStateInfo(const UE::Landscape::EditLayers::FMergeContext* InMergeContext,
-		UE::Landscape::EditLayers::FEditLayerTargetTypeState& OutSupportedTargetTypeState, UE::Landscape::EditLayers::FEditLayerTargetTypeState& OutEnabledTargetTypeState, TArray<TBitArray<>>& OutRenderGroups) const
+		UE::Landscape::EditLayers::FEditLayerTargetTypeState& OutSupportedTargetTypeState, UE::Landscape::EditLayers::FEditLayerTargetTypeState& OutEnabledTargetTypeState, TArray<TBitArray<>>& OutTargetLayerGroups) const
 		PURE_VIRTUAL(ILandscapeEditLayerRenderer::GetRendererStateInfo, );
 
 	/**
