Detect re-entrant calls to UpdateLayersContent or other main entry points to add or remove layers during UpdateLayersContent.

diff --git a/Engine/Source/Runtime/Landscape/Classes/Landscape.h b/Engine/Source/Runtime/Landscape/Classes/Landscape.h
--- a/Engine/Source/Runtime/Landscape/Classes/Landscape.h
+++ b/Engine/Source/Runtime/Landscape/Classes/Landscape.h
@@ -752,6 +752,9 @@ private:
 	class FLandscapeTexture2DResource* WeightmapScratchPackLayerTextureResource;
 
 	uint32 LastFlushedLayerUpdateFrame = 0;
+
+	// Counter to detect re-entrance
+	uint32 InLayerUpdateCount = 0;
 #endif
 };
 
diff --git a/Engine/Source/Runtime/Landscape/Private/LandscapeEditLayers.cpp b/Engine/Source/Runtime/Landscape/Private/LandscapeEditLayers.cpp
--- a/Engine/Source/Runtime/Landscape/Private/LandscapeEditLayers.cpp
+++ b/Engine/Source/Runtime/Landscape/Private/LandscapeEditLayers.cpp
@@ -87,6 +87,7 @@ LandscapeEditLayers.cpp: Landscape editing layers mode
 #include "Editor/UnrealEdEngine.h"
 #include "VisualLogger/VisualLogger.h"
 #include "LandscapeTextureStorageProvider.h"
+#include "UObject/Script.h"
 #endif
 
 #define LOCTEXT_NAMESPACE "Landscape"
@@ -530,6 +531,16 @@ namespace UE::Landscape::Private
 	}
 
 	static constexpr uint8 RGBAToRGBASwizzleMask = BuildChannelSwizzleMask(ERGBAChannel::R, ERGBAChannel::G, ERGBAChannel::B, ERGBAChannel::A);
+
+	static bool InBPCallstack()
+	{
+#if DO_BLUEPRINT_GUARD
+		const FBlueprintContextTracker* Tracker = FBlueprintContextTracker::TryGet();
+		return Tracker && Tracker->GetScriptEntryTag() > 0;
+#else
+		return false;
+#endif
+	}
 }
 
 // Vertex format and vertex buffer
@@ -2597,6 +2608,14 @@ bool ALandscape::CreateLayersRenderingResource(bool bUseNormalCapture)
 
 void ALandscape::ToggleCanHaveLayersContent()
 {
+	// Detect any attempt to call this in the middle of UpdateLayersContent.  If called from blueprint, log an error and return early instead of asserting.
+	if (InLayerUpdateCount > 0 && UE::Landscape::Private::InBPCallstack())
+	{
+		UE_LOG(LogLandscapeBP, Error, TEXT("Attempting to make illegal call to ToggleCanHaveLayersContent during UpdateLayersContent."));
+		return;
+	}
+	check(InLayerUpdateCount == 0);
+
 	bCanHaveLayersContent = !bCanHaveLayersContent;
 
 	if (!bCanHaveLayersContent)
@@ -11402,6 +11421,16 @@ void ALandscape::UpdateLayersContent(bool bInWaitForStreaming, bool bInSkipMonit
 {
 	TRACE_CPUPROFILER_EVENT_SCOPE(LandscapeLayers_UpdateLayersContent);
 
+	// Detect any attempt to re-enter.  If called from blueprint, log an error and return early instead of asserting.
+	if (InLayerUpdateCount > 0 && UE::Landscape::Private::InBPCallstack())
+	{
+		UE_LOG(LogLandscapeBP, Error, TEXT("Attempting to make illegal re-entrant call to UpdateLayersContent."));
+		return;
+	}
+
+	check(InLayerUpdateCount == 0);
+	InLayerUpdateCount++;
+
 	bool bHideNotifications = true;
 	ON_SCOPE_EXIT
 	{
@@ -11420,6 +11449,8 @@ void ALandscape::UpdateLayersContent(bool bInWaitForStreaming, bool bInSkipMonit
 		{
 			FLandscapeEditLayerReadback::GarbageCollectTasks();
 		}
+
+		InLayerUpdateCount--;
 	};
 
 	// Note : no early-out allowed before this : even if not actually updating edit layers, we need to poll our resources in order to make sure we register to streaming events when needed: 
@@ -12502,6 +12533,14 @@ bool ALandscape::DeleteLayer(int32 InLayerIndex)
 {
 	ensure(HasLayersContent());
 
+	// Detect any attempt to call this in the middle of UpdateLayersContent.  If called from blueprint, log an error and return early instead of asserting.
+	if (InLayerUpdateCount > 0 && UE::Landscape::Private::InBPCallstack())
+	{
+		UE_LOG(LogLandscapeBP, Error, TEXT("Attempting to make illegal call to DeleteLayer during UpdateLayersContent."));
+		return false;
+	}
+	check(InLayerUpdateCount == 0);
+
 	const FLandscapeLayer* LayerStruct = GetLayerConst(InLayerIndex);
 	if (!LayerStruct)
 	{
@@ -13060,6 +13099,14 @@ FLandscapeLayer* ALandscape::DuplicateLayerAndMoveBrushes(const FLandscapeLayer&
 
 int32 ALandscape::CreateLayer(FName InName, const TSubclassOf<ULandscapeEditLayerBase>& InEditLayerClass)
 {
+	// Detect any attempt to call this in the middle of UpdateLayersContent.  If called from blueprint, log an error and return early instead of asserting.
+	if (InLayerUpdateCount > 0 && UE::Landscape::Private::InBPCallstack())
+	{
+		UE_LOG(LogLandscapeBP, Error, TEXT("Attempting to make illegal call to CreateLayer during UpdateLayersContent."));
+		return INDEX_NONE;
+	}
+	check(InLayerUpdateCount == 0);
+
 	if (IsMaxLayersReached() || !CanHaveLayersContent())
 	{
 		return INDEX_NONE;
