Changed landscape batched merge API to allow using RTTI on ILandscapeEditLayerRenderer

#rnx

diff --git a/Engine/Source/Runtime/Landscape/Private/LandscapeEditLayerRenderer.cpp b/Engine/Source/Runtime/Landscape/Private/LandscapeEditLayerRenderer.cpp
--- a/Engine/Source/Runtime/Landscape/Private/LandscapeEditLayerRenderer.cpp
+++ b/Engine/Source/Runtime/Landscape/Private/LandscapeEditLayerRenderer.cpp
@@ -541,7 +541,7 @@ FEditLayerRendererState::FEditLayerRendererState(const FMergeContext* InMergeCon
 	UpdateActiveTargetTypeState();
 }
 
-FEditLayerRendererState::FEditLayerRendererState(const FMergeContext* InMergeContext, ILandscapeEditLayerRenderer* InRenderer)
+FEditLayerRendererState::FEditLayerRendererState(const FMergeContext* InMergeContext, TScriptInterface<ILandscapeEditLayerRenderer> InRenderer)
 	: MergeContext(InMergeContext)
 	, Renderer(InRenderer)
 	, DebugName(InRenderer->GetEditLayerRendererDebugName())
@@ -1231,7 +1231,7 @@ void FMergeRenderContext::Render(TFunction<void(const FOnRenderBatchTargetGroupD
 			    });
 			    SortedComponentMergeRenderInfos.Sort();
 			}
-
+			
 			switch (RenderStep.Type)
 			{
 			case FMergeRenderStep::EType::BeginRenderCommandRecorder:
@@ -1266,8 +1266,7 @@ void FMergeRenderContext::Render(TFunction<void(const FOnRenderBatchTargetGroupD
 			case FMergeRenderStep::EType::RenderLayer_Recorded:
 				{
 					TRACE_CPUPROFILER_EVENT_SCOPE(RenderLayer);
-
-					ILandscapeEditLayerRenderer* Renderer = RenderStep.RendererState.GetRenderer();
+					TScriptInterface<ILandscapeEditLayerRenderer> Renderer = RenderStep.RendererState.GetRenderer();
 					check((Renderer != nullptr) && EnumHasAnyFlags(Renderer->GetRenderFlags(this), ERenderFlags::RenderMode_Mask));
 
 					// TODO[jonathan.bard] offset the world transform to account for the half-pixel offset?
diff --git a/Engine/Source/Runtime/Landscape/Private/LandscapeEditLayers.cpp b/Engine/Source/Runtime/Landscape/Private/LandscapeEditLayers.cpp
--- a/Engine/Source/Runtime/Landscape/Private/LandscapeEditLayers.cpp
+++ b/Engine/Source/Runtime/Landscape/Private/LandscapeEditLayers.cpp
@@ -5424,7 +5424,7 @@ TArray<UE::Landscape::EditLayers::FEditLayerRendererState> FLandscapeLayer::GetE
 
 	check(EditLayer != nullptr);
 	// The edit layer itself might be a renderer:
-	if (ILandscapeEditLayerRenderer* AsRenderer = Cast<ILandscapeEditLayerRenderer>(EditLayer))
+	if (TScriptInterface<ILandscapeEditLayerRenderer> AsRenderer(EditLayer); AsRenderer != nullptr)
 	{
 		RendererStates.Emplace(InMergeContext, AsRenderer);
 	}
diff --git a/Engine/Source/Runtime/Landscape/Public/LandscapeEditLayerRenderer.h b/Engine/Source/Runtime/Landscape/Public/LandscapeEditLayerRenderer.h
--- a/Engine/Source/Runtime/Landscape/Public/LandscapeEditLayerRenderer.h
+++ b/Engine/Source/Runtime/Landscape/Public/LandscapeEditLayerRenderer.h
@@ -477,11 +477,11 @@ public:
 	LANDSCAPE_EDIT_LAYERS_BATCHED_MERGE_EXPERIMENTAL
 	LANDSCAPE_API explicit FEditLayerRendererState(const FMergeContext* InMergeContext);
 	LANDSCAPE_EDIT_LAYERS_BATCHED_MERGE_EXPERIMENTAL
-	LANDSCAPE_API explicit FEditLayerRendererState(const FMergeContext* InMergeContext, ILandscapeEditLayerRenderer* InRenderer);
+	LANDSCAPE_API explicit FEditLayerRendererState(const FMergeContext* InMergeContext, TScriptInterface<ILandscapeEditLayerRenderer> InRenderer);
 
 	/** Returns the edit layer renderer which this state relates to */
 	LANDSCAPE_EDIT_LAYERS_BATCHED_MERGE_EXPERIMENTAL
-	LANDSCAPE_API ILandscapeEditLayerRenderer* GetRenderer() const { return Renderer; }
+	LANDSCAPE_API TScriptInterface<ILandscapeEditLayerRenderer> GetRenderer() const { return Renderer; }
 
 	/** Returns a mask of all target types / weightmaps supported by this renderer */
 	LANDSCAPE_EDIT_LAYERS_BATCHED_MERGE_EXPERIMENTAL
@@ -601,7 +601,7 @@ private:
 	const FMergeContext* MergeContext = nullptr;
 
 	/** Renderer associated with this state */
-	ILandscapeEditLayerRenderer* Renderer = nullptr;
+	TScriptInterface<ILandscapeEditLayerRenderer> Renderer;
 
 	/** For debug purposes : this is the same as Renderer->GetEditLayerRendererDebugName() but having a member makes it easier to debug in the watch window */
 	FString DebugName;
