Always use compressed landscape textures (which fixes an issue with patching textures on some platforms).
New optimized landscape decompress code approximately 3-6 times faster.
On some low-end platforms we set the compression mip limit to a very high number so the compression path stores all of the data uncompressed, in order to save decompression time.

diff --git a/Engine/Source/Runtime/Engine/Classes/Engine/TextureAllMipDataProviderFactory.h b/Engine/Source/Runtime/Engine/Classes/Engine/TextureAllMipDataProviderFactory.h
--- a/Engine/Source/Runtime/Engine/Classes/Engine/TextureAllMipDataProviderFactory.h
+++ b/Engine/Source/Runtime/Engine/Classes/Engine/TextureAllMipDataProviderFactory.h
@@ -52,4 +52,9 @@ public:
 	  */
 	ENGINE_API virtual FStreamableRenderResourceState GetResourcePostInitState(const UTexture* Owner, bool bAllowStreaming)
 		PURE_VIRTUAL(UTextureAllMipDataProviderFactory::GetResourcePostInitState, return FStreamableRenderResourceState(););
+
+	/**
+	  * Return whether the mip data provided by this provider is in "offline processed" format (i.e. pre-tiled specifically for the target platform)
+	  */
+	ENGINE_API virtual bool IsDataOfflineProcessed(const UTexture* Owner) { return false; }
 };
diff --git a/Engine/Source/Runtime/Engine/Private/Rendering/StreamableTextureResource.cpp b/Engine/Source/Runtime/Engine/Private/Rendering/StreamableTextureResource.cpp
--- a/Engine/Source/Runtime/Engine/Private/Rendering/StreamableTextureResource.cpp
+++ b/Engine/Source/Runtime/Engine/Private/Rendering/StreamableTextureResource.cpp
@@ -12,7 +12,7 @@
 #include "ProfilingDebugging/ScopedDebugInfo.h"
 #include "Stats/StatsTrace.h"
 #include "RHIUtilities.h"
-#include "Engine/TextureMipDataProviderFactory.h"
+#include "Engine/TextureAllMipDataProviderFactory.h"
 
 #if STATS
 int64 GUITextureMemory = 0;
@@ -112,8 +112,11 @@ FStreamableTextureResource::FStreamableTextureResource(UTexture* InOwner, const
 	SizeY = Mip0.SizeY;
 	SizeZ = Mip0.SizeZ;
 
+	UTextureAllMipDataProviderFactory* AllMipProvider = InOwner->GetAssetUserData<UTextureAllMipDataProviderFactory>();
+	const bool bNotOfflineProcessed = (AllMipProvider != nullptr) ? (!AllMipProvider->IsDataOfflineProcessed(InOwner)) : (InOwner->bNotOfflineProcessed);
+
 	MipFadeSetting = (LODGroup == TEXTUREGROUP_Lightmap || LODGroup == TEXTUREGROUP_Shadowmap) ? MipFade_Slow : MipFade_Normal;
-	CreationFlags = (InOwner->SRGB ? TexCreate_SRGB : TexCreate_None)  | (InOwner->bNotOfflineProcessed ? TexCreate_None : TexCreate_OfflineProcessed) | TexCreate_ShaderResource | (InOwner->bNoTiling ? TexCreate_NoTiling : TexCreate_None);
+	CreationFlags = (InOwner->SRGB ? TexCreate_SRGB : TexCreate_None)  | (bNotOfflineProcessed ? TexCreate_None : TexCreate_OfflineProcessed) | TexCreate_ShaderResource | (InOwner->bNoTiling ? TexCreate_NoTiling : TexCreate_None);
 
 	if (InPostInitState.MaxNumLODs > 1)
 	{
diff --git a/Engine/Source/Runtime/Engine/Private/TextureDerivedData.cpp b/Engine/Source/Runtime/Engine/Private/TextureDerivedData.cpp
--- a/Engine/Source/Runtime/Engine/Private/TextureDerivedData.cpp
+++ b/Engine/Source/Runtime/Engine/Private/TextureDerivedData.cpp
@@ -3057,9 +3057,19 @@ static void SerializePlatformData(
 
 		// Serialize PackedData, OptData
 		Ar << PlatformData->PackedData;
-		if (PlatformData->GetHasOptData())
+
+		// The OptData describes the cooked mips; if !bSerializeMipData then clear it,
+		// as it's describing data we aren't serializing.
+		if (bSerializeMipData)
 		{
-			Ar << PlatformData->OptData;
+			if (PlatformData->GetHasOptData())
+			{
+				Ar << PlatformData->OptData;
+			}
+		}
+		else if (Ar.IsLoading()) // !bSerializeMipData and loading - make sure to clear opt data
+		{
+			PlatformData->SetOptData(FOptTexturePlatformData{});
 		}
 
 		// Serialize PixelFormat
@@ -3188,9 +3198,18 @@ static void SerializePlatformData(
 		Ar << PixelFormatString;
 	}
 
-	if (PlatformData->GetHasOptData())
+	// The OptData describes the cooked mips; if !bSerializeMipData then clear it,
+	// as it's describing data we aren't serializing.
+	if (bSerializeMipData)
+	{
+		if (PlatformData->GetHasOptData())
+		{
+			Ar << PlatformData->OptData;
+		}
+	}
+	else if (Ar.IsLoading()) // !bSerializeMipData and loading - make sure to clear opt data
 	{
-		Ar << PlatformData->OptData;
+		PlatformData->SetOptData(FOptTexturePlatformData{});
 	}
 
 	if (PlatformData->GetHasCpuCopy())
diff --git a/Engine/Source/Runtime/Landscape/Classes/LandscapeComponent.h b/Engine/Source/Runtime/Landscape/Classes/LandscapeComponent.h
--- a/Engine/Source/Runtime/Landscape/Classes/LandscapeComponent.h
+++ b/Engine/Source/Runtime/Landscape/Classes/LandscapeComponent.h
@@ -907,7 +907,7 @@ public:
 	  * [EDITOR ONLY] It will compress heightmaps to a ULandscapeTextureStorageProviderFactory when bUseCompression is true.  Ignored in non-editor builds.
 	  *    NOTE: It will NOT remove an existing one if false -- applying compression is a one-way street, and should only be used in cooked content.
 	  */
-	ULandscapeHeightmapTextureEdgeFixup* InstallOrUpdateTextureUserDatas(bool bUseEdgeFixup, bool bUseCompression = false, bool bUpdateSnapshotNow = true);
+	ULandscapeHeightmapTextureEdgeFixup* InstallOrUpdateTextureUserDatas(bool bUseEdgeFixup, bool bUseCompression = false, bool bUpdateSnapshotNow = true, int32 HeightmapCompressionMipThreshold = 32);
 
 #if WITH_EDITOR
 	/** Deletes a layer from this component if it does not contain data, calling DeleteLayerAllocation. */
diff --git a/Engine/Source/Runtime/Landscape/Classes/LandscapeTextureStorageProvider.h b/Engine/Source/Runtime/Landscape/Classes/LandscapeTextureStorageProvider.h
--- a/Engine/Source/Runtime/Landscape/Classes/LandscapeTextureStorageProvider.h
+++ b/Engine/Source/Runtime/Landscape/Classes/LandscapeTextureStorageProvider.h
@@ -188,7 +188,8 @@ public:
 	TObjectPtr<ULandscapeHeightmapTextureEdgeFixup> EdgeFixup;
 
 	#if WITH_EDITORONLY_DATA
-	static ULandscapeTextureStorageProviderFactory* ApplyTo(UTexture2D* InTargetTexture, const FVector& InLandsapeGridScale);
+	static ULandscapeTextureStorageProviderFactory* ApplyTo(UTexture2D* InTargetTexture, const FVector& InLandsapeGridScale, int32 InHeightmapCompressionMipThreshold);
+	void UpdateCompressedDataFromSource(UTexture2D* InTargetTexture, const FVector& InLandsapeGridScale, int32 InHeightmapCompressionMipThreshold);
 	#endif // WITH_EDITORONLY_DATA
 
 	void SetupEdgeFixup(ULandscapeHeightmapTextureEdgeFixup* EdgeFixup);
@@ -199,6 +200,7 @@ public:
 
 	// decompress the mip data. Currently operates in-place (though this could be changed in the future)
 	void DecompressMip(uint8* SourceData, int64 SourceDataBytes, uint8* DestData, int64 DestDataBytes, int32 MipIndex);
+	void DecompressMipOriginalUnoptimized(uint8* SourceData, int64 SourceDataBytes, uint8* DestData, int64 DestDataBytes, int32 MipIndex);
 
 	virtual void Serialize(FArchive& Ar) override;
 	static void AddReferencedObjects(UObject* InThis, FReferenceCollector& Collector);
diff --git a/Engine/Source/Runtime/Landscape/Private/Landscape.cpp b/Engine/Source/Runtime/Landscape/Private/Landscape.cpp
--- a/Engine/Source/Runtime/Landscape/Private/Landscape.cpp
+++ b/Engine/Source/Runtime/Landscape/Private/Landscape.cpp
@@ -228,16 +228,25 @@ static FAutoConsoleVariable CVarAllowGrassStripping(
 	true,
 	TEXT("Enables the conditional stripping of grass data during cook.  Disabling this means the bStripGrassWhenCooked* will be ignored."));
 
-int32 GLandscapeHeightmapCompressionMode = 0;
+int32 GLandscapeHeightmapCompressionMode = 1;
 static FAutoConsoleVariableRef CVarLandscapeHeightmapCompressionMode(
 	TEXT("landscape.HeightmapCompressionMode"),
 	GLandscapeHeightmapCompressionMode,
-	TEXT("Defines whether compression is applied to landscapes.\n")
+	TEXT("Defines whether compression is applied to landscapes.  Can be defined per platform.\n")
 	TEXT(" 0: use the per-landscape setting bUseCompressedHeightmapStorage (default)\n")
 	TEXT(" 1: force enable heightmap compression on all landscapes\n")
 	TEXT(" -1: force disable heightmap compression on all landscapes\n"),
 	ECVF_Preview | ECVF_ReadOnly);
 
+int32 GLandscapeHeightmapCompressionMipThreshold = 32;
+static FAutoConsoleVariableRef CVarLandscapeHeightmapCompressionMipThreshold(
+	TEXT("landscape.HeightmapCompressionMipThreshold"),
+	GLandscapeHeightmapCompressionMipThreshold,
+	TEXT("Sets the minimum size for which heightmap mips are stored in a compressed layout.  Can be defined per platform.\n")
+	TEXT("Below this size, mips are stored in an uncompressed layout.\n")
+	TEXT("Default threshold is 32, though some slower platforms may have a higher default threshold out of the box.\n"),
+	ECVF_Preview | ECVF_ReadOnly);
+
 bool GLandscapePrioritizeDirtyRVTPages = true;
 static FAutoConsoleVariableRef CVarPrioritizeDirtyRVTPages(
 	TEXT("landscape.PrioritizeDirtyRVTPages"),
@@ -2106,7 +2115,7 @@ static void OnStaticMeshLODDistanceScaleChanged()
 
 FAutoConsoleVariableSink OnStaticMeshLODDistanceScaleChangedSink(FConsoleCommandDelegate::CreateStatic(&OnStaticMeshLODDistanceScaleChanged));
 
-ULandscapeHeightmapTextureEdgeFixup* ULandscapeComponent::InstallOrUpdateTextureUserDatas(bool bUseEdgeFixup, bool bUseCompression, bool bUpdateSnapshotNow)
+ULandscapeHeightmapTextureEdgeFixup* ULandscapeComponent::InstallOrUpdateTextureUserDatas(bool bUseEdgeFixup, bool bUseCompression, bool bUpdateSnapshotNow, int32 HeightmapCompressionMipThreshold)
 {
 	if (HeightmapTexture == nullptr)
 	{
@@ -2163,7 +2172,12 @@ ULandscapeHeightmapTextureEdgeFixup* ULandscapeComponent::InstallOrUpdateTexture
 		// Install or Update the Storage Factory [editor only]
 		if (StorageFactory == nullptr)
 		{
-			StorageFactory = ULandscapeTextureStorageProviderFactory::ApplyTo(HeightmapTexture, LandscapeGridScale);
+			StorageFactory = ULandscapeTextureStorageProviderFactory::ApplyTo(HeightmapTexture, LandscapeGridScale, HeightmapCompressionMipThreshold);
+		}
+		else
+		{
+			// since different platforms may change compression settings / thresholds, we should update the compression each time
+			StorageFactory->UpdateCompressedDataFromSource(HeightmapTexture, LandscapeGridScale, HeightmapCompressionMipThreshold);
 		}
 #endif // WITH_EDITOR
 		StorageFactory->SetupEdgeFixup(EdgeFixup);
@@ -3854,13 +3868,23 @@ void ALandscapeProxy::PreSave(FObjectPreSaveContext ObjectSaveContext)
 			CVarLandscapeHeightmapCompressionMode->GetValue(HeightmapCompressionMode);
 		}
 
+		int32 HeightmapCompressionMipThreshold = 32;
+		if (IConsoleVariable* PlatformCVar = CVarLandscapeHeightmapCompressionMipThreshold->GetPlatformValueVariable(IniPlatformName).Get())
+		{
+			PlatformCVar->GetValue(HeightmapCompressionMipThreshold);
+		}
+		else
+		{
+			CVarLandscapeHeightmapCompressionMipThreshold->GetValue(HeightmapCompressionMipThreshold);
+		}
+
 		const bool bShouldCompressHeightmap = (HeightmapCompressionMode == 0) ? bUseCompressedHeightmapStorage : (HeightmapCompressionMode > 0);
 		const bool bIsStreamingProxy = this->IsA<ALandscapeStreamingProxy>();
 		const bool bUseEdgeFixups = bIsStreamingProxy;
 		const bool bUpdateSnapshotNow = true;
 		for (ULandscapeComponent* LandscapeComponent : LandscapeComponents)
 		{
-			LandscapeComponent->InstallOrUpdateTextureUserDatas(bUseEdgeFixups, bShouldCompressHeightmap, bUpdateSnapshotNow);
+			LandscapeComponent->InstallOrUpdateTextureUserDatas(bUseEdgeFixups, bShouldCompressHeightmap, bUpdateSnapshotNow, HeightmapCompressionMipThreshold);
 		}
 	}
 #endif // WITH_EDITOR
diff --git a/Engine/Source/Runtime/Landscape/Private/LandscapeTextureStorageProvider.cpp b/Engine/Source/Runtime/Landscape/Private/LandscapeTextureStorageProvider.cpp
--- a/Engine/Source/Runtime/Landscape/Private/LandscapeTextureStorageProvider.cpp
+++ b/Engine/Source/Runtime/Landscape/Private/LandscapeTextureStorageProvider.cpp
@@ -388,13 +388,24 @@ bool ULandscapeTextureStorageProviderFactory::GetInitialMipData(int32 FirstMipTo
 		if (BulkDataSize > 0)
 		{
 			void* SourceData = nullptr;
-			bool bDiscardInternalCopy = true;
+			const bool bDiscardInternalCopy = true;
 			Mip.BulkData.GetCopy(&SourceData, bDiscardInternalCopy);
 			check(SourceData);
-			uint8* DestData = static_cast<uint8*>(FMemory::Malloc(DestBytes));
-			DecompressMip((uint8*) SourceData, BulkDataSize, DestData, DestBytes, MipIndex);
-			OutMipData[MipIndex - FirstMipToLoad] = DestData;
-			FMemory::Free(SourceData);
+
+			if (Mip.bCompressed)
+			{
+				// decompress the mip to a new buffer, then free the original buffer
+				uint8* DestData = static_cast<uint8*>(FMemory::Malloc(DestBytes));
+				DecompressMip((uint8*)SourceData, BulkDataSize, DestData, DestBytes, MipIndex);
+				OutMipData[MipIndex - FirstMipToLoad] = DestData;
+				FMemory::Free(SourceData);
+			}
+			else
+			{
+				// mip is uncompressed, it should already be the correct size, and we can just use the source data buffer directly
+				check(BulkDataSize == DestBytes);
+				OutMipData[MipIndex - FirstMipToLoad] = SourceData;
+			}
 
 			if (OutMipSize.Num() > 0)
 			{
@@ -437,72 +448,81 @@ bool ULandscapeTextureStorageProviderFactory::GetInitialMipData(int32 FirstMipTo
 
 
 #if WITH_EDITORONLY_DATA
-ULandscapeTextureStorageProviderFactory* ULandscapeTextureStorageProviderFactory::ApplyTo(UTexture2D* TargetTexture, const FVector& InLandsapeGridScale)
+ULandscapeTextureStorageProviderFactory* ULandscapeTextureStorageProviderFactory::ApplyTo(UTexture2D* InTargetTexture, const FVector& InLandsapeGridScale, int32 InHeightmapCompressionMipThreshold)
 {
-	check(TargetTexture);
-	check(TargetTexture->Source.IsValid())
-
-	check(TargetTexture->Source.GetFormat() == TSF_BGRA8);
-	EPixelFormat Format = PF_B8G8R8A8;
-
-	int32 Width = TargetTexture->Source.GetSizeX();
-	int32 Height = TargetTexture->Source.GetSizeY();
-	int32 MipCount = TargetTexture->Source.GetNumMips();
+	check(InTargetTexture);
+	check(InTargetTexture->Source.IsValid())
 
-	uint32 SrcBpp = GPixelFormats[Format].BlockBytes;
-	uint32 SrcPitch = Width * SrcBpp;
+	check(InTargetTexture->Source.GetFormat() == TSF_BGRA8);
 
 	// try to get an existing factory
-	ULandscapeTextureStorageProviderFactory* Factory= TargetTexture->GetAssetUserData<ULandscapeTextureStorageProviderFactory>();
+	ULandscapeTextureStorageProviderFactory* Factory= InTargetTexture->GetAssetUserData<ULandscapeTextureStorageProviderFactory>();
 	if (Factory == nullptr)
 	{
 		// create a new one
-		Factory = NewObject<ULandscapeTextureStorageProviderFactory>(TargetTexture);
-		Factory->Texture = TargetTexture;
-		TargetTexture->AddAssetUserData(Factory);
+		Factory = NewObject<ULandscapeTextureStorageProviderFactory>(InTargetTexture);
+		Factory->Texture = InTargetTexture;
+		InTargetTexture->AddAssetUserData(Factory);
 	}
 
+	Factory->UpdateCompressedDataFromSource(InTargetTexture, InLandsapeGridScale, InHeightmapCompressionMipThreshold);
+
+	return Factory;
+}
+
+void ULandscapeTextureStorageProviderFactory::UpdateCompressedDataFromSource(UTexture2D* InTargetTexture, const FVector& InLandsapeGridScale, int32 InHeightmapCompressionMipThreshold)
+{
+	check(this->Texture == InTargetTexture);
+
+	EPixelFormat Format = PF_B8G8R8A8;
+
+	int32 Width = InTargetTexture->Source.GetSizeX();
+	int32 Height = InTargetTexture->Source.GetSizeY();
+	int32 MipCount = InTargetTexture->Source.GetNumMips();
+
+	uint32 SrcBpp = GPixelFormats[Format].BlockBytes;
+	uint32 SrcPitch = Width * SrcBpp;
+
 	// need this to properly calculate normals
-	Factory->LandscapeGridScale = InLandsapeGridScale;
+	this->LandscapeGridScale = InLandsapeGridScale;
 
 	// calculate number of non-streaming mips
 	// TODO [chris.tchou] : we could make this calculation platform specific, like Texture2D does.
 	// We would have to calculate it during serialization, when we know the target platform.
 	{
-		int32 NumNonStreamingMips = 1;
+		int32 NumberOfNonStreamingMips = 1;
 
 		// TODO [chris.tchou] : we could ensure Mip Tails are not streamed, as it's more overhead to upload.
 		// we would have to query TextureCompressorModule for platform specific info.
 		// Ignoring the mip tail should still work, just less optimal as it does more work at runtime to blit into the mip tail.
 		int32 NumMipsInTail = 0;
 
-		NumNonStreamingMips = FMath::Max(NumNonStreamingMips, NumMipsInTail);
-		NumNonStreamingMips = FMath::Max(NumNonStreamingMips, UTexture2D::GetStaticMinTextureResidentMipCount());
-		NumNonStreamingMips = FMath::Min(NumNonStreamingMips, MipCount);
-		Factory->NumNonStreamingMips = NumNonStreamingMips;
+		NumberOfNonStreamingMips = FMath::Max(NumberOfNonStreamingMips, NumMipsInTail);
+		NumberOfNonStreamingMips = FMath::Max(NumberOfNonStreamingMips, UTexture2D::GetStaticMinTextureResidentMipCount());
+		NumberOfNonStreamingMips = FMath::Min(NumberOfNonStreamingMips, MipCount);
+		this->NumNonStreamingMips = NumberOfNonStreamingMips;
 	}
 
 	// calculate number of non-optional mips
 	{
 		// for now, landscape texture storage does not have any optional mips
-		Factory->NumNonOptionalMips = MipCount;
+		this->NumNonOptionalMips = MipCount;
 	}
 
-	Factory->Mips.Empty();
+	this->Mips.Empty();
 	int32 MipWidth = Width;
 	int32 MipHeight = Height;
 	for (int32 MipIndex = 0; MipIndex < MipCount; MipIndex++)
 	{
-		FLandscapeTexture2DMipMap* Mip = new(Factory->Mips) FLandscapeTexture2DMipMap();
+		FLandscapeTexture2DMipMap* Mip = new(this->Mips) FLandscapeTexture2DMipMap();
 		Mip->SizeX = MipWidth;
 		Mip->SizeY = MipHeight;
-		
+
 		TArray64<uint8> MipData;
-		TargetTexture->Source.GetMipData(MipData, MipIndex);
+		InTargetTexture->Source.GetMipData(MipData, MipIndex);
 
-		// store small mips uncompressed
-		constexpr int32 UncompressedMipSizeThreshold = 8;
-		if ((MipWidth <= UncompressedMipSizeThreshold) || (MipHeight <= UncompressedMipSizeThreshold))
+		// Store mips below the threshold size uncompressed
+		if ((MipWidth < InHeightmapCompressionMipThreshold) || (MipHeight < InHeightmapCompressionMipThreshold))
 		{
 			Mip->bCompressed = false;
 			CopyMipToBulkData(MipIndex, MipWidth, MipHeight, MipData.GetData(), MipData.Num(), Mip->BulkData);
@@ -516,8 +536,6 @@ ULandscapeTextureStorageProviderFactory* ULandscapeTextureStorageProviderFactory
 		MipWidth = FMath::Max(MipWidth / 2, 1);
 		MipHeight = FMath::Max(MipHeight / 2, 1);
 	}
-
-	return Factory;
 }
 #endif // WITH_EDITORONLY_DATA
 
@@ -628,11 +646,6 @@ int32 FLandscapeTextureStorageMipProvider::GetMips(const FTextureUpdateContext&
 			break;
 		}
 
-		// If Data size is specified, check compatibility for safety		// TODO: size doesn't have to match when compressed...
-		if (DestMip.DataSize && (uint64)SourceMip->BulkData.GetBulkDataSize() > DestMip.DataSize)
-		{
-			break;
-		}
 
 		// Increment the sync counter.  This causes the system to not advance to the next tick, until RescheduleCallback() is called (by AsyncFileCallBack when counter reaches zero)
 		// If a request completes immediately, then it will call the RescheduleCallback,
@@ -640,16 +653,46 @@ int32 FLandscapeTextureStorageMipProvider::GetMips(const FTextureUpdateContext&
 		SyncOptions.Counter->Increment();
 
 		int64 StreamDataSize = SourceMip->BulkData.GetBulkDataSize();
-		uint8* StreamData = static_cast<uint8*>(FMemory::Malloc(StreamDataSize));
-		TRACE_IOSTORE_METADATA_SCOPE_TAG("Landscape");
-		IORequests[StartingMipIndex].BulkDataIORequest.Reset(
-			SourceMip->BulkData.CreateStreamingRequest(
-				0,
-				StreamDataSize,
-				(EAsyncIOPriorityAndFlags) FMath::Clamp<int32>(AIOP_Low + (bPrioritizedIORequest ? 1 : 0), AIOP_Low, AIOP_High) | AIOP_FLAG_DONTCACHE,
-				&AsyncFileCallBack,
-				StreamData)
-		);
+
+		if (SourceMip->bCompressed)
+		{
+			// allocate a buffer to receive the streamed data
+			uint8* StreamData = static_cast<uint8*>(FMemory::Malloc(StreamDataSize));
+
+			TRACE_IOSTORE_METADATA_SCOPE_TAG("Landscape");
+			IORequests[StartingMipIndex].BulkDataIORequest.Reset(
+				SourceMip->BulkData.CreateStreamingRequest(
+					0,
+					StreamDataSize,
+					(EAsyncIOPriorityAndFlags)FMath::Clamp<int32>(AIOP_Low + (bPrioritizedIORequest ? 1 : 0), AIOP_Low, AIOP_High) | AIOP_FLAG_DONTCACHE,
+					&AsyncFileCallBack,
+					StreamData)
+			);
+		}
+		else
+		{
+			// If DataSize is specified (optional, may be zero), then check that the size matches expectations
+			if ((DestMip.DataSize != 0) && (DestMip.DataSize != StreamDataSize))
+			{
+				// LogPlayLevel: Error: UAT: [2024.11.08 - 17.09.47:903] [10] LogLandscape : Error : Unexpected data size for landscape mip 0 : expected 0 bytes(512 x 512), has 1048576 bytes(512 x 512 compressed: 0)
+				UE_LOG(LogLandscape, Error, TEXT("Unexpected data size for landscape mip %d : expected %d bytes (%d x %d), has %d bytes (%d x %d compressed: %d)"),
+					StartingMipIndex,
+					DestMip.DataSize, DestMip.SizeX, DestMip.SizeY,
+					StreamDataSize, SourceMip->SizeX, SourceMip->SizeY, SourceMip->bCompressed
+				);
+
+				check(StreamDataSize == DestMip.DataSize);
+			}
+			TRACE_IOSTORE_METADATA_SCOPE_TAG("Landscape");
+			IORequests[StartingMipIndex].BulkDataIORequest.Reset(
+				SourceMip->BulkData.CreateStreamingRequest(
+					0,
+					StreamDataSize,
+					(EAsyncIOPriorityAndFlags)FMath::Clamp<int32>(AIOP_Low + (bPrioritizedIORequest ? 1 : 0), AIOP_Low, AIOP_High) | AIOP_FLAG_DONTCACHE,
+					&AsyncFileCallBack,
+					(uint8*) DestMip.DestData)	// when not compressed, we can stream directly into the dest mip memory
+			);
+		}
 
 		// remember the dest mip data buffer (we can't fill it out now, must wait until streaming is complete)
 		IORequests[StartingMipIndex].DestMipData = static_cast<uint8*>(DestMip.DestData);
@@ -684,10 +727,21 @@ bool FLandscapeTextureStorageMipProvider::PollMips(const FTextureUpdateSyncOptio
 		for (int MipIndex = FirstRequestedMipIndex; MipIndex < CurrentFirstLODIdx; MipIndex++)
 		{
 			const FLandscapeTexture2DMipMap* SourceMip = Factory->GetMip(MipIndex);
-			uint8* SourceData = IORequests[MipIndex].BulkDataIORequest->GetReadResults();
-			int64 DestDataBytes = SourceMip->SizeX * SourceMip->SizeY * 4;
-			uint8* DestData = IORequests[MipIndex].DestMipData;
-			Factory->DecompressMip(SourceData, SourceMip->BulkData.GetBulkDataSize(), DestData, DestDataBytes, MipIndex);
+
+			if (SourceMip->bCompressed)
+			{
+				uint8* SourceData = IORequests[MipIndex].BulkDataIORequest->GetReadResults();
+				int64 DestDataBytes = SourceMip->SizeX * SourceMip->SizeY * 4;
+				uint8* DestData = IORequests[MipIndex].DestMipData;
+				Factory->DecompressMip(SourceData, SourceMip->BulkData.GetBulkDataSize(), DestData, DestDataBytes, MipIndex);
+				FMemory::Free(SourceData);
+			}
+			else
+			{
+				// uncompressed streams directly into the dst mip data buffer, so nothing to do here (other than a sanity check)
+				uint8* SourceData = IORequests[MipIndex].BulkDataIORequest->GetReadResults();
+				check(IORequests[MipIndex].DestMipData == SourceData);
+			}
 		}
 		
 		if (ShouldPatchStreamingMipEdges())
@@ -758,7 +812,7 @@ void ULandscapeTextureStorageProviderFactory::CopyMipToBulkData(int32 MipIndex,
 	DestBulkData.Lock(LOCK_READ_WRITE);
 
 	int32 TotalPixels = MipSizeX * MipSizeY;
-	check(SourceDataBytes == TotalPixels * 4);
+	checkf(SourceDataBytes == TotalPixels * 4, TEXT("SourceDataBytes: %d TotalPixels: %d MipIndex: %d MipSizeX: %d MipSizeY: %d"), SourceDataBytes, TotalPixels, MipIndex, MipSizeX, MipSizeY);
 
 	int32 DestBytes = SourceDataBytes;
 	uint8* DestData = DestBulkData.Realloc(DestBytes);
@@ -842,14 +896,82 @@ void ULandscapeTextureStorageProviderFactory::CompressMipToBulkData(int32 MipInd
 }
 
 // Compute the normal of the triangle formed by the 3 points (in winding order).
-static FVector ComputeTriangleNormal(const FVector& InPoint0, const FVector& InPoint1, const FVector& InPoint2)
+inline FVector ComputeTriangleNormal(const FVector& InPoint0, const FVector& InPoint1, const FVector& InPoint2)
 {
 	FVector Normal = (InPoint0 - InPoint1).Cross(InPoint1 - InPoint2);
 	Normal.Normalize();
 	return Normal;
 }
 
-static void SampleWorldPositionAtOffset(FVector& OutPoint, const uint8* MipData, int32 X, int32 Y, int32 MipSizeX, const FVector& InLandscapeGridScale)
+// This function is unused, but explains how we get from ComputeTriangleNormal above to the optimized version below.
+// When computing normals on a height grid, you can simplify the math, and only care about delta height in the +X and +Y directions.
+// Then we can take advantage of the zeros in DX and DY to simplify the cross product:
+inline FVector3f ComputeGridNormalFromDeltaHeights(float DHDX, float DHDY, int32 MipScale, const FVector& LandscapeGridScale)
+{
+	// by placing the origin at the center vertex, and ensuring one vector is along +X and the other along +Y, a lot of math is removed:
+	// FVector3f Center(0.0f, 0.0f, 0.0f);
+	FVector3f DX(MipScale * LandscapeGridScale.X, 0.0, DHDX * (float)LandscapeGridScale.Z);
+	FVector3f DY(0.0, MipScale * LandscapeGridScale.Y, DHDY * (float)LandscapeGridScale.Z);
+
+	FVector3f Normal; // = (DX-Center).Cross(DY-Center);
+	{
+		Normal.X = /*DX.Y * DY.Z*/ - DX.Z * DY.Y;		// DHDX * (-LGS.Z * LGS.Y * MipScale)		<< note values in parens are constant in the inner loop
+		Normal.Y = /*DX.Z * DY.X*/ - DX.X * DY.Z;		// DHDY * (-LGS.Z * LGS.X * MipScale)
+		Normal.Z = DX.X * DY.Y /* - DX.Y * DY.X*/;		// (LGS.X * LGS.Y * MipScale * MipScale)	<< fully constant in the inner loop
+	}
+	Normal.Normalize();
+	return Normal;
+}
+
+inline FVector2f CalculatePremultU16(int32 MipIndex, const FVector& LandscapeGridScale)
+{
+	// We optimize the cross product calculation in the inner loop by precalculating the DHDX and DHDY multipliers.
+	int32 MipScale = 1 << MipIndex;
+
+	// Note that we're also doing an optimization trick by scaling the resulting vector such that CrossProductResult.Z == 1.0.
+	// Since we pass the result through Normalize(), that scale factor doesn't matter -- but it's faster to calculate that way.
+	// LANDSCAPE_ZSCALE comes from the fact that we are operating on the integer height values, and this converts the integer heights to landscape space
+	// (and then the LandscapeGridScale converts that to world space)
+	float ScaleFactor = -LANDSCAPE_ZSCALE / (LandscapeGridScale.X * LandscapeGridScale.Y * MipScale);
+	FVector2f PremultU16;
+	PremultU16.X = LandscapeGridScale.Z * LandscapeGridScale.Y * ScaleFactor;
+	PremultU16.Y = LandscapeGridScale.Z * LandscapeGridScale.X * ScaleFactor;
+	return PremultU16;
+}
+
+// This takes it a few steps further : we've minimized the math here by premultiplying everything related to LGS, MipScale, and LandscapeScale into PremultLGS
+// We've also scaled up the results so that Normal.Z == 1, which reduces the math used by the Normalize.
+inline FVector3f ComputeGridNormalFromDeltaHeightsPremultU16(int32 DHDX, int32 DHDY, const FVector2f& PremultU16)
+{
+	FVector3f Normal; // = DX.Cross(DY);
+	{
+		// we've calculated PremultU16 to ensure Normal.Z is 1.0 (see CalculatePremultU16), which saves some math in Normalize()
+		Normal.X = DHDX * PremultU16.X;
+		Normal.Y = DHDY * PremultU16.Y;
+		// Normal.Z = 1.0f;
+	}
+	// Normal.Normalize(); optimized below
+	{
+		const float SquareSum = Normal.X * Normal.X + Normal.Y * Normal.Y + 1.0f;
+		if (SquareSum > UE_SMALL_NUMBER)
+		{
+			// sqrt estimate should be more than sufficient for 8 bit results.
+			const float Scale = FMath::InvSqrtEst(SquareSum);
+			Normal.X *= Scale;
+			Normal.Y *= Scale;
+			Normal.Z = Scale;	// take advantage of knowing Normal.Z == 1.0
+		}
+		else
+		{
+			Normal.X = 0.0f;
+			Normal.Y = 0.0f;
+			Normal.Z = 1.0f;
+		}
+	}
+	return Normal;
+}
+
+inline void SampleWorldPositionAtOffset(FVector& OutPoint, const uint8* MipData, int32 X, int32 Y, int32 MipSizeX, const FVector& InLandscapeGridScale)
 {
 	int32 OffsetBytes = (Y * MipSizeX + X) * 4;
 	uint16 HeightData = MipData[OffsetBytes + 2] * 256 + MipData[OffsetBytes + 1];
@@ -861,118 +983,204 @@ static void SampleWorldPositionAtOffset(FVector& OutPoint, const uint8* MipData,
 		LandscapeDataAccess::GetLocalHeight(HeightData) * InLandscapeGridScale.Z);
 }
 
+inline uint16 DecodeHeightU16(const FColor* Pixel)
+{
+	uint16 HeightData = Pixel->R * 256 + Pixel->G;
+	return HeightData;
+}
+
+inline void FastNormalize(FVector3f& V)
+{
+	const float SquareSum = V.X * V.X + V.Y * V.Y + V.Z * V.Z;
+	if (SquareSum > UE_SMALL_NUMBER)
+	{
+		const float Scale = FMath::InvSqrtEst(SquareSum);
+		V *= Scale;
+	}
+	else
+	{
+		V.X = 0.0f;
+		V.Y = 0.0f;
+		V.Z = 1.0f;
+	}
+}
+
+// The triangle topology is the following (where C = center, T = top, B = bottom, L = left, R = right and Nx the normals we need to interpolate):
+// .  ------ . --------.
+// |         | \       |
+//    \                |
+// |         |   \     |
+//   P0'\ P1'| N0'     |	<< normals calculated for the previous line
+// |         |     \   |
+//        \            |
+// |         |       \ |
+// . - - - - TL ------ TT
+// |         | \       |
+//    \      |  \      |
+// |         |   \     |
+//   P0 \ P1 | N0 \ N1 |
+// |         |     \   |
+//        \  |      \  |
+// |         |       \ |
+// . - - - - LL ------ CC   << current pixel being processed
+//
+// we calculate normals while we decompress, as a single pass gives better cache coherency.
+// while iterating each interior pixel left to right, top to bottom, we:
+// 1) Decode Height at CC (current pixel)
+// 2) Write Height at CC
+// 3) Compute N0/N1 using heights at CC/TT/TL/LL (all previously decoded)
+// 4) Complete Normal calculation for TL == (P0' + P1' + N0') + P1 + N0 + N1
+// 5) Write Normal for TL
+// 6) Store Partial Normal for LL in PrevLine cache -- stores P0 + P1 + N0
+
 void ULandscapeTextureStorageProviderFactory::DecompressMip(uint8* SourceData, int64 SourceDataBytes, uint8* DestData, int64 DestDataBytes, int32 MipIndex)
 {
 	TRACE_CPUPROFILER_EVENT_SCOPE(ULandscapeTextureStorageProviderFactory::DecompressMip);
-
+	
 	check(SourceData && DestData);
 
 	FLandscapeTexture2DMipMap& Mip = Mips[MipIndex];
 
-	if (!Mip.bCompressed)
-	{
-		// mip is uncompressed, just copy it
-		check(SourceDataBytes == DestDataBytes);
-		memcpy(DestData, SourceData, DestDataBytes);
-		return;
-	}
-
-	check(Mip.bCompressed);
+	check(Mip.bCompressed);		// uncompressed should be handled outside of this function
 
-	int32 TotalPixels = Mip.SizeX * Mip.SizeY;
-	check(SourceDataBytes == (TotalPixels + (Mip.SizeX + Mip.SizeY) * 2 - 4) * 2);
+	int32 Width = Mip.SizeX;
+	int32 Height = Mip.SizeY;
+	int32 TotalPixels = Width * Height;
+	int32 BorderPixels = (Width + Height) * 2 - 4;
+	check(SourceDataBytes == (TotalPixels + BorderPixels) * 2); // 2 bytes (height) for each pixel, plus 2 bytes (normal x/y) for each border pixel
 	check(DestDataBytes == TotalPixels * 4);
 
-	// Undo Delta Encode of Heights
-	uint16 LastHeight = 32768;
-	for (int32 PixelIndex = 0; PixelIndex < TotalPixels; PixelIndex++)
-	{
-		int32 SourceOffset = PixelIndex * 2;
-		uint16 DeltaHeight = SourceData[SourceOffset + 0] * 256 + SourceData[SourceOffset + 1];
+	// save some multiplying by premultiplying the grid scales, mip scale and ZScale
+	FVector2f PremultU16 = CalculatePremultU16(MipIndex, LandscapeGridScale);
 
-		// undo delta
-		LastHeight += DeltaHeight;
+	// current center pixel height
+	// (also used to delta decode the heights - initial value must match the initial value used during encoding)
+	uint16 CC = 32768;
 
-		// texture data is stored as BGRA, or [normal x, height low bits, height high bits, normal y]
-		int32 DestOffset = PixelIndex * 4;
-		DestData[DestOffset + 0] = 128;
-		DestData[DestOffset + 1] = LastHeight & 0xff;
-		DestData[DestOffset + 2] = LastHeight >> 8;
-		DestData[DestOffset + 3] = 128;
-	}
+	// partial normal results recorded for the previous line
+	TArray<FVector3f, TInlineAllocator<512>> PrevLinePartialNormals;
+	PrevLinePartialNormals.SetNumZeroed(Width);
 
-	// recompute normals in the interior
+	// iterate each line
+	for (int32 Y = 0; Y < Height; Y++)
 	{
-		// we skip computing the edges, as they will be overwritten later (and this way we don't have to handle samples that go off the edge)
-		for (int32 Y = 1; Y < Mip.SizeY - 1; Y++)
+		int32 LineOffsetInPixels = Y * Width;
+		uint8* Src = &SourceData[LineOffsetInPixels * 2];
+		FColor* Dst = (FColor*)&DestData[LineOffsetInPixels * 4];
+
+		if (Y == 0)
 		{
-			for (int32 X = 1; X < Mip.SizeX - 1; X++)
+			// just decode heights for the first line (normals don't matter they will be stomped below)
+			for (int32 X = 0; X < Width; X++)
 			{
-				// based on shader code in LandscapeLayersPS.usf
-				FVector TL, TT, CC, LL, RR, BR, BB;
-
-				SampleWorldPositionAtOffset(TL, DestData, X - 1, Y - 1, Mip.SizeX, LandscapeGridScale);
-				SampleWorldPositionAtOffset(TT, DestData, X + 0, Y - 1, Mip.SizeX, LandscapeGridScale);
-				SampleWorldPositionAtOffset(CC, DestData, X + 0, Y + 0, Mip.SizeX, LandscapeGridScale);
-				SampleWorldPositionAtOffset(LL, DestData, X - 1, Y + 0, Mip.SizeX, LandscapeGridScale);
-				SampleWorldPositionAtOffset(RR, DestData, X + 1, Y + 0, Mip.SizeX, LandscapeGridScale);
-				SampleWorldPositionAtOffset(BR, DestData, X + 1, Y + 1, Mip.SizeX, LandscapeGridScale);
-				SampleWorldPositionAtOffset(BB, DestData, X + 0, Y + 1, Mip.SizeX, LandscapeGridScale);
-
-				FVector N0 = ComputeTriangleNormal(CC, LL, TL);
-				FVector N1 = ComputeTriangleNormal(TL, TT, CC);
-				FVector N2 = ComputeTriangleNormal(LL, CC, BB);
-				FVector N3 = ComputeTriangleNormal(RR, CC, TT);
-				FVector N4 = ComputeTriangleNormal(BR, BB, CC);
-				FVector N5 = ComputeTriangleNormal(CC, RR, BR);
-
-				FVector FinalNormal = (N0 + N1 + N2 + N3 + N4 + N5);
-				FinalNormal.Normalize();
-
-				// rescale normal.xy to [0,255] range, and write out as bytes
-				int32 OffsetBytes = (Y * Mip.SizeX + X) * 4;
-				DestData[OffsetBytes + 0] = static_cast<uint8>(FMath::Clamp(((FinalNormal.X + 1.0) * 0.5) * 255.0, 0.0, 255.0));
-				DestData[OffsetBytes + 3] = static_cast<uint8>(FMath::Clamp(((FinalNormal.Y + 1.0) * 0.5) * 255.0, 0.0, 255.0));
+				uint16 DeltaHeight = Src[0] * 256 + Src[1];
+				CC += DeltaHeight;
+				*Dst = FColor(CC >> 8, CC & 0xff, 128, 128);
+				Src += 2;
+				Dst++;
 			}
 		}
-	}
+		else
+		{
+			// compute initial values (first pixel)
+			FVector3f P1(ForceInitToZero), P01(ForceInitToZero);	// previous quad N1 and (N0+N1) normals
+			uint16 TT;												// previous quad TT height
+			{
+				uint16 DeltaHeight = Src[0] * 256 + Src[1];
+				CC += DeltaHeight;
+				*Dst = FColor(CC >> 8, CC & 0xff, 128, 128);
 
-	// write out normals along the edge (delta encoded clockwise starting from top left)
-	int32 SourceOffset = TotalPixels * 2;
-	uint8 LastNormalX = 128;
-	uint8 LastNormalY = 128;
+				// load TT for first pixel (becomes TL for second pixel)
+				TT = DecodeHeightU16(Dst + 0 - Width);
 
-	auto DecodeNormal = [&LastNormalX, &LastNormalY, SourceData, DestData, MipSizeX = Mip.SizeX, &SourceOffset](int32 X, int32 Y)
-	{
-		int32 DestOffset = (Y * MipSizeX + X) * 4;
-		LastNormalX += SourceData[SourceOffset + 0];
-		LastNormalY += SourceData[SourceOffset + 1];
-		DestData[DestOffset + 0] = LastNormalX;
-		DestData[DestOffset + 3] = LastNormalY;
-		SourceOffset += 2;
-	};
+				Src += 2;
+				Dst++;
+			}
 
-	for (int32 X = 0; X < Mip.SizeX; X++)				// [0 ... MipSizeX-1], 0
-	{
-		DecodeNormal(X, 0);
+			// rest of the pixels in the line
+			for (int32 X = 1; X < Width; X++)
+			{
+				// re-use previous pixel TT and CC as this pixel TL and LL
+				uint16 TL = TT;
+				uint16 LL = CC;
+
+				// 1) Decode Height at CC
+				uint16 DeltaHeight = Src[0] * 256 + Src[1];
+				CC += DeltaHeight;
+
+				// load TT
+				TT = DecodeHeightU16(Dst + 0 - Width);
+
+				// 2) Write Height at CC (normals get written during processing of the next line)
+				*Dst = FColor(CC >> 8, CC & 0xff, 128, 128);
+
+				// 3) Compute local normals N0/N1 for the current quad (CC/TT/TL/LL)
+				FVector3f N0 = ComputeGridNormalFromDeltaHeightsPremultU16(CC - LL, LL - TL, PremultU16);
+				FVector3f N1 = ComputeGridNormalFromDeltaHeightsPremultU16(TT - TL, CC - TT, PremultU16);
+				FVector3f N01 = N0 + N1;
+
+				// 4) Complete Normal calculation for TL - this takes the partial result from the previous line and fills in the rest
+				FVector3f TL_Normal = PrevLinePartialNormals[X - 1] + P1 + N01;
+				FastNormalize(TL_Normal);
+				
+				// 5) Write Normal for TL
+				Dst[-Width - 1].B = static_cast<uint8>(FMath::Clamp((TL_Normal.X * 127.5f + 127.5f), 0.0f, 255.0f));
+				Dst[-Width - 1].A = static_cast<uint8>(FMath::Clamp((TL_Normal.Y * 127.5f + 127.5f), 0.0f, 255.0f));
+
+				// 6) Store Partial Normal for LL in PrevLinePartialNormals (P0 + P1 + N0) - the rest will be filled in when processing the next line
+				FVector3f LL_PartialNormal = P01 + N0;
+				PrevLinePartialNormals[X - 1] = LL_PartialNormal;
+
+				// pass normals to next pixel
+				P1 = N1;
+				P01 = N01;
+
+				Src += 2;
+				Dst++;
+			}
+		}
 	}
 
-	for (int32 Y = 1; Y < Mip.SizeY; Y++)				// MipSizeX-1, [1 ... MipSizeY-1]
+	// write out normals along the edge (delta encoded clockwise starting from top left)
 	{
-		DecodeNormal(Mip.SizeX - 1, Y);
-	}
+		TRACE_CPUPROFILER_EVENT_SCOPE(EdgeNormalFixup);
 
-	for (int32 X = Mip.SizeX - 2; X >= 0; X--)				// [MipSizeX-2 ... 0], MipSizeY-1
-	{
-		DecodeNormal(X, Mip.SizeY - 1);
-	}
+		uint8* Src = &SourceData[TotalPixels * 2];
+		uint8 LastNormalX = 128;
+		uint8 LastNormalY = 128;
 
-	for (int32 Y = Mip.SizeY - 2; Y >= 1; Y--)				// 0, [MipSizeY-2 ... 1]
-	{
-		DecodeNormal(0, Y);
-	}
+		auto DecodeNormal = [&LastNormalX, &LastNormalY, DestData, Width, &Src](int32 X, int32 Y)
+			{
+				int32 DestOffset = (Y * Width + X) * 4;
+				LastNormalX += Src[0];
+				LastNormalY += Src[1];
+				DestData[DestOffset + 0] = LastNormalX;
+				DestData[DestOffset + 3] = LastNormalY;
+				Src += 2;
+			};
+
+		for (int32 X = 0; X < Width; X++)		// [0 ... Width-1], 0
+		{
+			DecodeNormal(X, 0);
+		}
+
+		for (int32 Y = 1; Y < Height; Y++)		// Width-1, [1 ... Height-1]
+		{
+			DecodeNormal(Width - 1, Y);
+		}
 
-	check(SourceOffset == SourceDataBytes);
+		for (int32 X = Width - 2; X >= 0; X--)	// [Width-2 ... 0], Height-1
+		{
+			DecodeNormal(X, Height - 1);
+		}
+
+		for (int32 Y = Height - 2; Y >= 1; Y--)	// 0, [Height-2 ... 1]
+		{
+			DecodeNormal(0, Y);
+		}
+
+		check(Src == &SourceData[SourceDataBytes]);
+	}
 }
 
 #undef PROVIDER_DEBUG_LOG
