Landscape : fixed undo on new landscapes (i.e. when they haven't been saved yet). The code was incorrectly detecting, upon registration of the landscape proxy (for those that hadn't gone through PostLoad), that it went from non-edit layers to edit layers mode and, as a result, was copying the final heightmap/weightmap data to the default (first) edit layer

diff --git a/Engine/Source/Runtime/Landscape/Classes/Landscape.h b/Engine/Source/Runtime/Landscape/Classes/Landscape.h
--- a/Engine/Source/Runtime/Landscape/Classes/Landscape.h
+++ b/Engine/Source/Runtime/Landscape/Classes/Landscape.h
@@ -296,7 +296,6 @@ public:
 	virtual bool CanChangeIsSpatiallyLoadedFlag() const override { return false; }
 #endif // WITH_EDITOR
 
-	virtual void PostInitProperties() override;
 	virtual void PostLoad() override;
 	virtual void BeginDestroy() override;
 	virtual void FinishDestroy() override;
diff --git a/Engine/Source/Runtime/Landscape/Classes/LandscapeProxy.h b/Engine/Source/Runtime/Landscape/Classes/LandscapeProxy.h
--- a/Engine/Source/Runtime/Landscape/Classes/LandscapeProxy.h
+++ b/Engine/Source/Runtime/Landscape/Classes/LandscapeProxy.h
@@ -960,7 +960,7 @@ PRAGMA_ENABLE_DEPRECATION_WARNINGS
 	/** Because bHasLayersContent can be stomped by the FixupSharedData copying layers over from the parent LandscapeActor,
 	  * we preserve the initial state of bHasLayersContent until we get to registration and can apply upgrade steps
 	  **/
-	bool bHadLayersContentAtPostLoadTime = false;
+	TOptional<bool> HadLayersContentAtPostLoadTime;
 
 	/** Enable compressed heightmap texture storage. */
 	UPROPERTY(EditAnywhere, Category = Landscape, AdvancedDisplay, meta = (LandscapeInherited))
diff --git a/Engine/Source/Runtime/Landscape/Private/Landscape.cpp b/Engine/Source/Runtime/Landscape/Private/Landscape.cpp
--- a/Engine/Source/Runtime/Landscape/Private/Landscape.cpp
+++ b/Engine/Source/Runtime/Landscape/Private/Landscape.cpp
@@ -2921,12 +2921,6 @@ void ULandscapeInfo::Serialize(FArchive& Ar)
 	}
 }
 
-void ALandscape::PostInitProperties()
-{
-	Super::PostInitProperties();
-}
-
-
 void ALandscape::PostLoad()
 {
 	TRACE_CPUPROFILER_EVENT_SCOPE(ALandscape::PostLoad);
@@ -4203,7 +4197,8 @@ void ALandscapeProxy::PostLoad()
 #endif // WITH_EDITOR
 
 	// save the load time state of layers content before doing anything else (specifically FixupSharedData can stomp it)
-	bHadLayersContentAtPostLoadTime = bHasLayersContent;
+	check(!HadLayersContentAtPostLoadTime.IsSet());
+	HadLayersContentAtPostLoadTime = bHasLayersContent;
 
 	Super::PostLoad();
 
@@ -5704,13 +5699,16 @@ void ULandscapeInfo::RegisterLandscapeActorWithProxyInternal(ALandscapeProxy* Pr
 
 		Proxy->UpdateCachedHasLayersContent(true);
 
-		bool bHasLayersContentBefore = Proxy->bHadLayersContentAtPostLoadTime;
-		check(Proxy->WeightmapFixupVersion == Proxy->CurrentVersion);
-		const bool bNeedOldDataMigration = !bHasLayersContentBefore && CanHaveLayersContent();
-		if (bNeedOldDataMigration && LandscapeActor->HasLayersContent())
+		if (Proxy->HadLayersContentAtPostLoadTime.IsSet())
 		{
-			LandscapeActor->CopyOldDataToDefaultLayer(Proxy);
-			Proxy->bHadLayersContentAtPostLoadTime = false;
+			bool bHasLayersContentBefore = Proxy->HadLayersContentAtPostLoadTime.GetValue();
+			check(Proxy->WeightmapFixupVersion == Proxy->CurrentVersion);
+			const bool bNeedOldDataMigration = !bHasLayersContentBefore && CanHaveLayersContent();
+			if (bNeedOldDataMigration && LandscapeActor->HasLayersContent())
+			{
+				LandscapeActor->CopyOldDataToDefaultLayer(Proxy);
+				Proxy->HadLayersContentAtPostLoadTime = false;
+			}
 		}
 	}
 #endif // WITH_EDITOR
