Defensive fix for landscape (wrongly) assuming a root component is available (same approach as everywhere else in the code)

diff --git a/Engine/Source/Runtime/Landscape/Private/Landscape.cpp b/Engine/Source/Runtime/Landscape/Private/Landscape.cpp
--- a/Engine/Source/Runtime/Landscape/Private/Landscape.cpp
+++ b/Engine/Source/Runtime/Landscape/Private/Landscape.cpp
@@ -5366,8 +5366,9 @@ void ULandscapeInfo::FixupProxiesTransform(bool bDirty)
 {
 	ALandscape* Landscape = LandscapeActor.Get();
 
-	if (Landscape == nullptr ||
-		Landscape->GetRootComponent()->IsRegistered() == false)
+	if ((Landscape == nullptr)
+		|| (Landscape->GetRootComponent() == nullptr)
+		|| !Landscape->GetRootComponent()->IsRegistered())
 	{
 		return;
 	}
