Fixed case where, in WP, the landscape heightmaps' FTextureSource could have an initialized mip0 (non zero) but zeroed-out mip1 and up, which would break the assumption that the texture hash (which is only based on mip0) is representative of all mips.

Instead, we now make sure that when using edit layers, all mips are initialized to zero, so that the first GPU readback appropriately updates the heightmap

diff --git a/Engine/Source/Runtime/Landscape/Private/LandscapeEdit.cpp b/Engine/Source/Runtime/Landscape/Private/LandscapeEdit.cpp
--- a/Engine/Source/Runtime/Landscape/Private/LandscapeEdit.cpp
+++ b/Engine/Source/Runtime/Landscape/Private/LandscapeEdit.cpp
@@ -6002,6 +6002,8 @@ void ALandscape::SplitHeightmap(ULandscapeComponent* Comp, ALandscapeProxy* Targ
 
 	ALandscapeProxy* SrcProxy = Comp->GetLandscapeProxy();
 	ALandscapeProxy* DstProxy = TargetProxy ? TargetProxy : SrcProxy;
+	check(SrcProxy->HasLayersContent() == DstProxy->CanHaveLayersContent());
+	const bool bWithLayersContent = SrcProxy->HasLayersContent() && DstProxy->CanHaveLayersContent();
 	SrcProxy->Modify();
 	DstProxy->Modify();
 
@@ -6036,7 +6038,17 @@ void ALandscape::SplitHeightmap(ULandscapeComponent* Comp, ALandscapeProxy* Targ
 		Comp->SetHeightmap(NewHeightmapTexture);
 
 		check(Comp->GetHeightmap(false) == Comp->GetHeightmap(true));
-		LandscapeEdit.SetHeightData(Comp->GetSectionBase().X, Comp->GetSectionBase().Y, Comp->GetSectionBase().X + Comp->ComponentSizeQuads, Comp->GetSectionBase().Y + Comp->ComponentSizeQuads, (uint16*)HeightData.GetData(), 0, false, (uint16*)NormalData.GetData());
+		// Don't bother copying the data if we're using edit layers, because that will be done in the next UpdateLayersContent. 
+		//  It would also have a weird, pervert effect, because when edit layers are in use, the "InGenerateMips = true" argument passed to the following call to SetHeightData (default value) is actually
+		//  ignored. This leads to the heightmap's FTextureSource having non-zero values on mip0 and zeroes on all the other mip levels, which is wrong because then FTextureSource's higher mips won't match what 
+		//  the GPU texture's higher mips will contain after completion of UpdateLayersContent. You would expect that it would actually get fixed after the GPU readback, but this is not actually happening, because the 
+		//  computed hash (see ULandscapeTextureHash::CalculateTextureHash64 in the following lines) only takes into account mip0! So when the readback occurs, it will consider that the FTextureSource and 
+		//  the GPU texture are in sync (even though it's not the case since their mips 1 and above are different) and will *not* update FTextureSource's higher mips accordingly!
+		if (!bWithLayersContent)
+		{
+			LandscapeEdit.SetHeightData(Comp->GetSectionBase().X, Comp->GetSectionBase().Y, Comp->GetSectionBase().X + Comp->ComponentSizeQuads, Comp->GetSectionBase().Y + Comp->ComponentSizeQuads, 
+				(uint16*)HeightData.GetData(), 0, false, (uint16*)NormalData.GetData());
+		}
 	}
 
 	OldHeightmapTexture->PostEditChange();
@@ -6092,10 +6104,11 @@ void ALandscape::SplitHeightmap(ULandscapeComponent* Comp, ALandscapeProxy* Targ
 	GDisableAutomaticTextureMaterialUpdateDependencies = false;
 
 #if WITH_EDITORONLY_DATA
-	check(Comp->GetLandscapeProxy()->HasLayersContent() == DstProxy->CanHaveLayersContent());
-	if (Comp->GetLandscapeProxy()->HasLayersContent() && DstProxy->CanHaveLayersContent())
+	if (bWithLayersContent)
 	{
 		FLandscapeEditLayerReadback* NewCPUReadback = new FLandscapeEditLayerReadback();
+		const uint64 Hash = ULandscapeTextureHash::CalculateTextureHash64(NewHeightmapTexture, ELandscapeTextureType::Heightmap);
+		NewCPUReadback->SetHash(Hash);
 		DstProxy->HeightmapsCPUReadback.Add(NewHeightmapTexture, NewCPUReadback);
 
 		// Free OldHeightmapTexture's CPUReadBackResource if not used by any component
