RHI: Change params to work with/without new profiler.

#rnx

diff --git a/Engine/Source/Runtime/Landscape/Public/LandscapeUtils.h b/Engine/Source/Runtime/Landscape/Public/LandscapeUtils.h
--- a/Engine/Source/Runtime/Landscape/Public/LandscapeUtils.h
+++ b/Engine/Source/Runtime/Landscape/Public/LandscapeUtils.h
@@ -197,7 +197,7 @@ public:
 					std::apply([RDGEvent, &GraphBuilder](auto&&... InnerArgs)
 						{
 							// Allocate the TOptional<TRDGEventScopeGuard> now in order to insert the tag : 
-							RDGEvent->Emplace(GraphBuilder, ERDGScopeFlags::None, FRHIBreadcrumbData(__FILE__, __LINE__, TStatId(), NAME_None), std::forward<TArgs>(InnerArgs)...);
+							RDGEvent->Emplace(GraphBuilder, ERDGScopeFlags::None, FRHIBreadcrumbData(__FILE__, __LINE__), std::forward<TArgs>(InnerArgs)...);
 						}, Values);
 				});
 		}
