Changed the tooltip of UseCompressedHeightmapStorage in landscape proxy actors to make it explicit that this is only useful in the case where the per-platform CVar landscape.HeightmapCompressionMode is not 0
+ Removed deprecated code from 5.4

#rnx

diff --git a/Engine/Plugins/Experimental/Water/Source/Runtime/Private/WaterBodyExclusionVolume.cpp b/Engine/Plugins/Experimental/Water/Source/Runtime/Private/WaterBodyExclusionVolume.cpp
--- a/Engine/Plugins/Experimental/Water/Source/Runtime/Private/WaterBodyExclusionVolume.cpp
+++ b/Engine/Plugins/Experimental/Water/Source/Runtime/Private/WaterBodyExclusionVolume.cpp
@@ -89,13 +89,6 @@ void AWaterBodyExclusionVolume::UpdateOverlappingWaterBodies(const FWaterExclusi
 	}
 }
 
-void AWaterBodyExclusionVolume::UpdateOverlappingWaterBodies()
-{
-	FWaterExclusionVolumeChangedParams Params;
-	Params.bUserTriggered = false;
-	UpdateOverlappingWaterBodies(Params);
-}
-
 #if WITH_EDITOR
 void AWaterBodyExclusionVolume::UpdateActorIcon()
 {
diff --git a/Engine/Plugins/Experimental/Water/Source/Runtime/Public/WaterBodyExclusionVolume.h b/Engine/Plugins/Experimental/Water/Source/Runtime/Public/WaterBodyExclusionVolume.h
--- a/Engine/Plugins/Experimental/Water/Source/Runtime/Public/WaterBodyExclusionVolume.h
+++ b/Engine/Plugins/Experimental/Water/Source/Runtime/Public/WaterBodyExclusionVolume.h
@@ -46,8 +46,6 @@ class WATER_API AWaterBodyExclusionVolume : public APhysicsVolume
 public:
 	void UpdateOverlappingWaterBodies(const FWaterExclusionVolumeChangedParams& Params);
 
-	UE_DEPRECATED(5.4, "UpdateOverlappingWaterBodies not taking a params struct is deprecated. The function now has the potential of dirtying water bodies when Params.bUserTriggered = true!")
-	void UpdateOverlappingWaterBodies();
 #if WITH_EDITOR
 	void UpdateActorIcon();
 #endif // WITH_EDITOR
diff --git a/Engine/Plugins/Experimental/Water/Source/Runtime/Public/WaterMeshComponent.h b/Engine/Plugins/Experimental/Water/Source/Runtime/Public/WaterMeshComponent.h
--- a/Engine/Plugins/Experimental/Water/Source/Runtime/Public/WaterMeshComponent.h
+++ b/Engine/Plugins/Experimental/Water/Source/Runtime/Public/WaterMeshComponent.h
@@ -92,8 +92,6 @@ public:
 	UFUNCTION(BlueprintPure, Category = Rendering)
 	bool IsEnabled() const { return bIsEnabled; }
 
-	UE_DEPRECATED(5.4, "The ExtentInTiles is now derived from the water zone extent and the tile size.")
-	void SetExtentInTiles(FIntPoint NewExtentInTiles) {}
 private:
 	//~ Begin USceneComponent Interface
 	virtual FBoxSphereBounds CalcBounds(const FTransform& LocalToWorld) const override;
diff --git a/Engine/Source/Runtime/Foliage/Private/InstancedGrass.cpp b/Engine/Source/Runtime/Foliage/Private/InstancedGrass.cpp
--- a/Engine/Source/Runtime/Foliage/Private/InstancedGrass.cpp
+++ b/Engine/Source/Runtime/Foliage/Private/InstancedGrass.cpp
@@ -41,11 +41,6 @@ void UGrassInstancedStaticMeshComponent::BuildTreeAnyThread(
 	OutSortedInstances = MoveTemp(Builder.Result->SortedInstances);
 }
 
-void UGrassInstancedStaticMeshComponent::AcceptPrebuiltTree(TArray<FInstancedStaticMeshInstanceData>& InInstanceData, TArray<FClusterNode>& InClusterTree, int32 InOcclusionLayerNumNodes, int32 InNumBuiltRenderInstances)
-{
-	check(false);
-}
-
 void UGrassInstancedStaticMeshComponent::AcceptPrebuiltTree(TArray<FClusterNode>& InClusterTree, int32 InOcclusionLayerNumNodes, int32 InNumBuiltRenderInstances, FStaticMeshInstanceData* InSharedInstanceBufferData)
 {
 	checkSlow(IsInGameThread());
diff --git a/Engine/Source/Runtime/Foliage/Public/GrassInstancedStaticMeshComponent.h b/Engine/Source/Runtime/Foliage/Public/GrassInstancedStaticMeshComponent.h
--- a/Engine/Source/Runtime/Foliage/Public/GrassInstancedStaticMeshComponent.h
+++ b/Engine/Source/Runtime/Foliage/Public/GrassInstancedStaticMeshComponent.h
@@ -14,9 +14,6 @@ class UGrassInstancedStaticMeshComponent : public UHierarchicalInstancedStaticMe
 public:
 	static FOLIAGE_API void BuildTreeAnyThread(TArray<FMatrix>& InstanceTransforms, TArray<float>& InstanceCustomDataFloats, int32 NumCustomDataFloats, const FBox& MeshBox, TArray<FClusterNode>& OutClusterTree, TArray<int32>& OutSortedInstances, TArray<int32>& OutInstanceReorderTable, int32& OutOcclusionLayerNum, int32 MaxInstancesPerLeaf, bool InGenerateInstanceScalingRange);
 	
-	UE_DEPRECATED(5.4, "Use the new version of AcceptPrebuiltTree below.")
-	FOLIAGE_API void AcceptPrebuiltTree(TArray<FInstancedStaticMeshInstanceData>& InInstanceData, TArray<FClusterNode>& InClusterTree, int32 InOcclusionLayerNumNodes, int32 InNumBuiltRenderInstances);
-
 	FOLIAGE_API void AcceptPrebuiltTree(TArray<FClusterNode>& InClusterTree, int32 InOcclusionLayerNumNodes, int32 InNumBuiltRenderInstances, FStaticMeshInstanceData* InSharedInstanceBufferData);
 
 	FOLIAGE_API bool SupportsWorldPositionOffsetVelocity() const override;
diff --git a/Engine/Source/Runtime/Landscape/Classes/LandscapeHeightfieldCollisionComponent.h b/Engine/Source/Runtime/Landscape/Classes/LandscapeHeightfieldCollisionComponent.h
--- a/Engine/Source/Runtime/Landscape/Classes/LandscapeHeightfieldCollisionComponent.h
+++ b/Engine/Source/Runtime/Landscape/Classes/LandscapeHeightfieldCollisionComponent.h
@@ -107,17 +107,8 @@ public:
 		Chaos::FHeightFieldPtr HeightfieldGeometry;
 	    Chaos::FHeightFieldPtr HeightfieldSimpleGeometry;
 
-		UE_DEPRECATED(5.4, "Please use HeightfieldGeometry instead")
-		TUniquePtr<Chaos::FHeightField> Heightfield;
-		
-		UE_DEPRECATED(5.4, "Please use HeightfieldSimpleGeometry instead")
-		TUniquePtr<Chaos::FHeightField> HeightfieldSimple;
-
 #if WITH_EDITORONLY_DATA
 		Chaos::FHeightFieldPtr EditorHeightfieldGeometry;
-
-		UE_DEPRECATED(5.4, "Please use HeightfieldSimpleGeometry instead")
-		TUniquePtr<Chaos::FHeightField> EditorHeightfield;
 #endif // WITH_EDITORONLY_DATA
 
 		FHeightfieldGeometryRef(FGuid& InGuid);
diff --git a/Engine/Source/Runtime/Landscape/Classes/LandscapeMeshCollisionComponent.h b/Engine/Source/Runtime/Landscape/Classes/LandscapeMeshCollisionComponent.h
--- a/Engine/Source/Runtime/Landscape/Classes/LandscapeMeshCollisionComponent.h
+++ b/Engine/Source/Runtime/Landscape/Classes/LandscapeMeshCollisionComponent.h
@@ -45,17 +45,9 @@ public:
 
 		TArray<Chaos::FMaterialHandle> UsedChaosMaterials;
 		Chaos::FTriangleMeshImplicitObjectPtr TrimeshGeometry;
-
-		// When deleted, remove PRAGMA_DISABLE_DEPRECATION_WARNINGS / PRAGMA_ENABLE_DEPRECATION_WARNINGS from ~FTriMeshGeometryRef()
-		UE_DEPRECATED(5.4, "Please use TrimeshGeometry instead")
-		TUniquePtr<Chaos::FTriangleMeshImplicitObject> Trimesh;
 		
 #if WITH_EDITORONLY_DATA
 		Chaos::FTriangleMeshImplicitObjectPtr EditorTrimeshGeometry;
-		
-		// When deleted, remove PRAGMA_DISABLE_DEPRECATION_WARNINGS / PRAGMA_ENABLE_DEPRECATION_WARNINGS from ~FTriMeshGeometryRef()
-		UE_DEPRECATED(5.4, "Please use EditorTrimeshGeometry instead")
-		TUniquePtr<Chaos::FTriangleMeshImplicitObject> EditorTrimesh;
 #endif // WITH_EDITORONLY_DATA
 
 		FTriMeshGeometryRef();
diff --git a/Engine/Source/Runtime/Landscape/Classes/LandscapeProxy.h b/Engine/Source/Runtime/Landscape/Classes/LandscapeProxy.h
--- a/Engine/Source/Runtime/Landscape/Classes/LandscapeProxy.h
+++ b/Engine/Source/Runtime/Landscape/Classes/LandscapeProxy.h
@@ -501,10 +501,6 @@ public:
 	TEnumAsByte<ELandscapeLODFalloff::Type> LODFalloff_DEPRECATED;
 #endif // WITH_EDITORONLY_DATA
 
-	UE_DEPRECATED(5.4, "Unused property : will be removed in a future version")
-	UPROPERTY()
-	float ComponentScreenSizeToUseSubSections;
-
 	/** This is the starting screen size used to calculate the distribution. You can increase the value if you want less LOD0 component, and you use very large landscape component. */
 	UPROPERTY(EditAnywhere, Category = "LOD Distribution", meta = (EditCondition = "!bUseScalableLODSettings", DisplayName = "LOD 0 Screen Size", ClampMin = "0.1", ClampMax = "10.0", UIMin = "0.1", UIMax = "10.0", LandscapeInherited))
 	float LOD0ScreenSize = 0.5f;
@@ -962,7 +958,7 @@ PRAGMA_ENABLE_DEPRECATION_WARNINGS
 	  **/
 	TOptional<bool> HadLayersContentAtPostLoadTime;
 
-	/** Enable compressed heightmap texture storage. */
+	/** Enable compressed heightmap texture storage. This setting can be overridden by the per-platform CVar landscape.HeightmapCompressionMode */
 	UPROPERTY(EditAnywhere, Category = Landscape, AdvancedDisplay, meta = (LandscapeInherited))
 	bool bUseCompressedHeightmapStorage = false;
 
@@ -999,7 +995,6 @@ PRAGMA_ENABLE_DEPRECATION_WARNINGS
 	UFUNCTION(BlueprintCallable, Category = "Rendering", meta=(DeprecatedFunction, DeprecationMessage = "This value can't be changed anymore, you should edit the property LODDistributionSetting of the Landscape"))
 	virtual void ChangeLODDistanceFactor(float InLODDistanceFactor);
 
-	/** Change ComponentScreenSizeToUseSubSections value on the render proxy.*/
 	UFUNCTION(BlueprintCallable, Category = "Rendering", meta = (DeprecatedFunction, DeprecationMessage = "This value can't be changed anymore and has been ineffective for several versions now. Please stop using it"))
 	virtual void ChangeComponentScreenSizeToUseSubSections(float InComponentScreenSizeToUseSubSections);
 
@@ -1199,10 +1194,6 @@ PRAGMA_ENABLE_DEPRECATION_WARNINGS
 	LANDSCAPE_API void UpdateRenderingMethod();
 
 #if WITH_EDITOR
-	/** Render grass maps for the specified components */
-	UE_DEPRECATED(5.4, "This version of RenderGrassMaps is deprecated.  Use BuildGrassMaps() instead.")
-	void RenderGrassMaps(const TArray<ULandscapeComponent*>& InLandscapeComponents, const TArray<ULandscapeGrassType*>& InGrassTypes) {}
-
 	/**
 	 * Update the landscape physical material render tasks
 	 * @param bInShouldMarkDirty indicates whether the proxy should be marked as dirty if anything was built
diff --git a/Engine/Source/Runtime/Landscape/Private/Landscape.cpp b/Engine/Source/Runtime/Landscape/Private/Landscape.cpp
--- a/Engine/Source/Runtime/Landscape/Private/Landscape.cpp
+++ b/Engine/Source/Runtime/Landscape/Private/Landscape.cpp
@@ -233,8 +233,8 @@ static FAutoConsoleVariableRef CVarLandscapeHeightmapCompressionMode(
 	TEXT("landscape.HeightmapCompressionMode"),
 	GLandscapeHeightmapCompressionMode,
 	TEXT("Defines whether compression is applied to landscapes.  Can be defined per platform.\n")
-	TEXT(" 0: use the per-landscape setting bUseCompressedHeightmapStorage (default)\n")
-	TEXT(" 1: force enable heightmap compression on all landscapes\n")
+	TEXT(" 0: use the per-landscape setting bUseCompressedHeightmapStorage\n")
+	TEXT(" 1: force enable heightmap compression on all landscapes, regardless of the per-landscape setting bUseCompressedHeightmapStorage (default)\n")
 	TEXT(" -1: force disable heightmap compression on all landscapes\n"),
 	ECVF_Preview | ECVF_ReadOnly);
 
diff --git a/Engine/Source/Runtime/Landscape/Private/LandscapeCollision.cpp b/Engine/Source/Runtime/Landscape/Private/LandscapeCollision.cpp
--- a/Engine/Source/Runtime/Landscape/Private/LandscapeCollision.cpp
+++ b/Engine/Source/Runtime/Landscape/Private/LandscapeCollision.cpp
@@ -142,14 +142,12 @@ ULandscapeHeightfieldCollisionComponent::FHeightfieldGeometryRef::FHeightfieldGe
 {
 }
 
-PRAGMA_DISABLE_DEPRECATION_WARNINGS
 ULandscapeHeightfieldCollisionComponent::FHeightfieldGeometryRef::~FHeightfieldGeometryRef()
 {
 	// Remove ourselves from the shared map.
 	UE_SCOPELOCK_SHARED_HEIGHTFIELD_REFS();
 	GSharedHeightfieldRefs.Remove(Guid);
 }
-PRAGMA_ENABLE_DEPRECATION_WARNINGS
 
 void ULandscapeHeightfieldCollisionComponent::FHeightfieldGeometryRef::GetResourceSizeEx(FResourceSizeEx& CumulativeResourceSize)
 {
@@ -185,13 +183,11 @@ ULandscapeMeshCollisionComponent::FTriMeshGeometryRef::FTriMeshGeometryRef(FGuid
 	: Guid(InGuid)
 {}
 
-PRAGMA_DISABLE_DEPRECATION_WARNINGS
 ULandscapeMeshCollisionComponent::FTriMeshGeometryRef::~FTriMeshGeometryRef()
 {
 	// Remove ourselves from the shared map.
 	GSharedMeshRefs.Remove(Guid);
 }
-PRAGMA_ENABLE_DEPRECATION_WARNINGS
 
 void ULandscapeMeshCollisionComponent::FTriMeshGeometryRef::GetResourceSizeEx(FResourceSizeEx& CumulativeResourceSize)
 {
diff --git a/Engine/Source/Runtime/Landscape/Private/LandscapeRender.cpp b/Engine/Source/Runtime/Landscape/Private/LandscapeRender.cpp
--- a/Engine/Source/Runtime/Landscape/Private/LandscapeRender.cpp
+++ b/Engine/Source/Runtime/Landscape/Private/LandscapeRender.cpp
@@ -2603,12 +2603,6 @@ void FLandscapeComponentSceneProxy::DrawStaticElements(FStaticPrimitiveDrawInter
 	check(StaticBatchParamArray.Num() <= TotalBatchCount);
 }
 
-// Deprecated : do not use
-int8 FLandscapeComponentSceneProxy::GetLODFromScreenSize(float InScreenSizeSquared, float InViewLODScale) const
-{
-	return static_cast<int8>(FLandscapeRenderSystem::ComputeLODFromScreenSize(LODSettings, InScreenSizeSquared / (InViewLODScale * InViewLODScale)));
-}
-
 namespace
 {
 	FLinearColor GetColorForLod(int32 CurrentLOD, int32 ForcedLOD, bool DisplayCombinedBatch)
@@ -4448,11 +4442,6 @@ void ALandscapeProxy::ChangeLODDistanceFactor(float InLODDistanceFactor)
 	// Deprecated
 }
 
-void FLandscapeComponentSceneProxy::ChangeComponentScreenSizeToUseSubSections_RenderThread(float InComponentScreenSizeToUseSubSections)
-{
-	// Deprecated
-}
-
 bool FLandscapeComponentSceneProxy::HeightfieldHasPendingStreaming() const
 {
 	bool bHeightmapTextureStreaming = false;
diff --git a/Engine/Source/Runtime/Landscape/Public/LandscapeBlueprintBrushBase.h b/Engine/Source/Runtime/Landscape/Public/LandscapeBlueprintBrushBase.h
--- a/Engine/Source/Runtime/Landscape/Public/LandscapeBlueprintBrushBase.h
+++ b/Engine/Source/Runtime/Landscape/Public/LandscapeBlueprintBrushBase.h
@@ -18,14 +18,6 @@ struct FLandscapeBrushParameters
 	GENERATED_BODY()
 
 	FLandscapeBrushParameters() = default;
-
-	UE_DEPRECATED(5.4, "Use the other constructor")
-	FLandscapeBrushParameters(const ELandscapeToolTargetType& InLayerType, TObjectPtr<UTextureRenderTarget2D> InCombinedResult, const FName& InWeightmapLayerName = FName())
-		: CombinedResult(InCombinedResult)
-		, LayerType(InLayerType)
-		, WeightmapLayerName(InWeightmapLayerName)
-	{}
-
 	FLandscapeBrushParameters(bool bInIsHeightmapMerge, const FTransform& InRenderAreaWorldTransform, const FIntPoint& InRenderAreaSize, UTextureRenderTarget2D* InCombinedResult, const FName& InWeightmapLayerName = FName());
 
 	UPROPERTY(Category = "Settings", EditAnywhere, BlueprintReadWrite)
diff --git a/Engine/Source/Runtime/Landscape/Public/LandscapeRender.h b/Engine/Source/Runtime/Landscape/Public/LandscapeRender.h
--- a/Engine/Source/Runtime/Landscape/Public/LandscapeRender.h
+++ b/Engine/Source/Runtime/Landscape/Public/LandscapeRender.h
@@ -903,9 +903,6 @@ protected:
 protected:
 	LANDSCAPE_API virtual ~FLandscapeComponentSceneProxy();
 	
-	UE_DEPRECATED(5.4, "Removed function")
-	LANDSCAPE_API int8 GetLODFromScreenSize(float InScreenSizeSquared, float InViewLODScale) const;
-
 	LANDSCAPE_API bool GetMeshElementForFixedGrid(int32 InLodIndex, FMaterialRenderProxy* InMaterialInterface, FMeshBatch& OutMeshBatch, TArray<FLandscapeBatchElementParams>& OutStaticBatchParamArray) const;
 	template<class ArrayType> bool GetStaticMeshElement(int32 LODIndex, bool bForToolMesh, FMeshBatch& MeshBatch, ArrayType& OutStaticBatchParamArray) const;
 
@@ -936,9 +933,6 @@ public:
 	const FMeshBatch& GetGrassMeshBatch() const { return GrassMeshBatch; }
 
 	// FLandcapeSceneProxy
-	UE_DEPRECATED(5.4, "This function has been deprecated and is ineffective")
-	LANDSCAPE_API void ChangeComponentScreenSizeToUseSubSections_RenderThread(float InComponentScreenSizeToUseSubSections);
-
 	LANDSCAPE_API virtual bool HeightfieldHasPendingStreaming() const override;
 
 	LANDSCAPE_API virtual void GetHeightfieldRepresentation(UTexture2D*& OutHeightmapTexture, UTexture2D*& OutVisibilityTexture, FHeightfieldComponentDescription& OutDescription) const override;
