Speculative fix for content worker crash where a landscape proxy is missing its root component. FLandscapeGroup will now fail registering the landscape if that occurs, which is fine in a cooking situation as this is only for rendering purposes.

[RN]

diff --git a/Engine/Source/Runtime/Landscape/Private/LandscapeGroup.cpp b/Engine/Source/Runtime/Landscape/Private/LandscapeGroup.cpp
--- a/Engine/Source/Runtime/Landscape/Private/LandscapeGroup.cpp
+++ b/Engine/Source/Runtime/Landscape/Private/LandscapeGroup.cpp
@@ -83,9 +83,16 @@ void FLandscapeGroup::RegisterComponent(ULandscapeComponent* Component)
 		return;
 	}
 
-	// TODO [chris.tchou] : Remove this once we remove shared heightmaps, so we can enable edge fixup on non-WP landscapes.
-	ALandscapeStreamingProxy* StreamingProxy = Cast<ALandscapeStreamingProxy>(Component->GetOwner());
-	if (!StreamingProxy || !StreamingProxy->bIsLandscapeActorRegisteredWithLandscapeInfo)
+	ALandscapeProxy* LandscapeProxy = Component->GetLandscapeProxy();
+	check(LandscapeProxy != nullptr);
+
+	// TODO [jonathan.bard] : Remove this once we remove shared heightmaps, so we can enable edge fixup on non-WP landscapes.
+	if (!LandscapeProxy->IsA<ALandscapeStreamingProxy>())
+	{
+		return;
+	}
+	
+	if (!LandscapeProxy->bIsLandscapeActorRegisteredWithLandscapeInfo)
 	{
 		// the landscape group only registers components on streaming proxies (we don't handle non-WP cases)
 		// and we can't register until the landscape actor has been registered and shared settings are fixed
@@ -95,6 +102,13 @@ void FLandscapeGroup::RegisterComponent(ULandscapeComponent* Component)
 		return;
 	}
 
+	// Early-out if the proxy hasn't got a proper root component. To our knowledge, this only ever happens when cooking where landscape groups are useless : 
+	if (USceneComponent* RootComponent = LandscapeProxy->GetRootComponent(); (RootComponent == nullptr) || !RootComponent->IsRegistered())
+	{
+		UE_LOG(LogLandscape, Log, TEXT("Skipped registration of landscape component to its proxy %s : the proxy's root component isn't properly set or registered"), *LandscapeProxy->GetFullName())
+		return;
+	}
+
 	UTexture2D* Heightmap = Component->GetHeightmap();
 	check(Heightmap);
 
