Deprecated landscape's per-proxy setting for enabling/disabling heightmap compression. It's now only controlled by the per-platform CVar landscape.HeightmapCompressionMode

diff --git a/Engine/Source/Runtime/Landscape/Classes/LandscapeProxy.h b/Engine/Source/Runtime/Landscape/Classes/LandscapeProxy.h
--- a/Engine/Source/Runtime/Landscape/Classes/LandscapeProxy.h
+++ b/Engine/Source/Runtime/Landscape/Classes/LandscapeProxy.h
@@ -958,9 +958,9 @@ PRAGMA_ENABLE_DEPRECATION_WARNINGS
 	  **/
 	TOptional<bool> HadLayersContentAtPostLoadTime;
 
-	/** Enable compressed heightmap texture storage. This setting can be overridden by the per-platform CVar landscape.HeightmapCompressionMode */
-	UPROPERTY(EditAnywhere, Category = Landscape, AdvancedDisplay, meta = (LandscapeInherited))
-	bool bUseCompressedHeightmapStorage = false;
+	UE_DEPRECATED(5.6, "This has been replaced by the per-platform CVar landscape.HeightmapCompressionMode")
+	UPROPERTY()
+	bool bUseCompressedHeightmapStorage_DEPRECATED = false;
 
 	/** Strip Physics/collision components when cooked for client */
 	UPROPERTY(EditAnywhere, Category = Landscape, AdvancedDisplay, meta = (LandscapeOverridable))
diff --git a/Engine/Source/Runtime/Landscape/Private/Landscape.cpp b/Engine/Source/Runtime/Landscape/Private/Landscape.cpp
--- a/Engine/Source/Runtime/Landscape/Private/Landscape.cpp
+++ b/Engine/Source/Runtime/Landscape/Private/Landscape.cpp
@@ -233,9 +233,8 @@ static FAutoConsoleVariableRef CVarLandscapeHeightmapCompressionMode(
 	TEXT("landscape.HeightmapCompressionMode"),
 	GLandscapeHeightmapCompressionMode,
 	TEXT("Defines whether compression is applied to landscapes.  Can be defined per platform.\n")
-	TEXT(" 0: use the per-landscape setting bUseCompressedHeightmapStorage\n")
-	TEXT(" 1: force enable heightmap compression on all landscapes, regardless of the per-landscape setting bUseCompressedHeightmapStorage (default)\n")
-	TEXT(" -1: force disable heightmap compression on all landscapes\n"),
+	TEXT(" 0: force disable heightmap compression on all landscapes\n")
+	TEXT(" 1: force enable heightmap compression on all landscapes (default)\n"),
 	ECVF_Preview | ECVF_ReadOnly);
 
 int32 GLandscapeHeightmapCompressionMipThreshold = 32;
@@ -3872,7 +3871,7 @@ void ALandscapeProxy::PreSave(FObjectPreSaveContext ObjectSaveContext)
 			CVarLandscapeHeightmapCompressionMipThreshold->GetValue(HeightmapCompressionMipThreshold);
 		}
 
-		const bool bShouldCompressHeightmap = (HeightmapCompressionMode == 0) ? bUseCompressedHeightmapStorage : (HeightmapCompressionMode > 0);
+		const bool bShouldCompressHeightmap = (HeightmapCompressionMode > 0);
 		const bool bIsStreamingProxy = this->IsA<ALandscapeStreamingProxy>();
 		const bool bUseEdgeFixups = bIsStreamingProxy;
 		const bool bUpdateSnapshotNow = true;
