Add controls for warning spam about landscape.EditLayersLocalMerge.MaxResolutionPerRenderBatch.
Rename and deprecate landscape.SupressMapCheckWarnings.Nanite to landscape.Nanite.SilenceMapCheckWarnings for consistency.

diff --git a/Engine/Source/Runtime/Landscape/Classes/Landscape.h b/Engine/Source/Runtime/Landscape/Classes/Landscape.h
--- a/Engine/Source/Runtime/Landscape/Classes/Landscape.h
+++ b/Engine/Source/Runtime/Landscape/Classes/Landscape.h
@@ -718,6 +718,8 @@ private:
 	UPROPERTY(Transient)
 	bool bSplineLayerUpdateRequested;
 
+	bool bWarnedLayerMergeResolution = false;
+
 	ELandscapeEditLayersMergeMode CurrentEditLayersMergeMode = ELandscapeEditLayersMergeMode::GlobalMerge;
 
 	struct FWaitingForResourcesNotificationHelper
diff --git a/Engine/Source/Runtime/Landscape/Private/Landscape.cpp b/Engine/Source/Runtime/Landscape/Private/Landscape.cpp
--- a/Engine/Source/Runtime/Landscape/Private/Landscape.cpp
+++ b/Engine/Source/Runtime/Landscape/Private/Landscape.cpp
@@ -213,10 +213,11 @@ static FAutoConsoleVariable CVarSilenceSharedPropertyDeprecationFixup(
 	true,
 	TEXT("Silently performs the fixup of discrepancies in shared properties when handling data modified before the enforcement introduction."));
 
-static FAutoConsoleVariable CVarLandscapeSupressMapCheckWarnings_Nanite(
-	TEXT("landscape.SupressMapCheckWarnings.Nanite"),
+static FAutoConsoleVariable CVarLandscapeSilenceMapCheckWarnings_Nanite(
+	TEXT("landscape.Nanite.SilenceMapCheckWarnings"),
 	false,
 	TEXT("Issue MapCheck Info messages instead of warnings if Nanite Data is out of date"));
+static FAutoConsoleVariableDeprecated CVarLandscapeSuppressMapCheckWarnings_Nanite_Deprecated(TEXT("landscape.SupressMapCheckWarnings.Nanite"), TEXT("landscape.Nanite.SilenceMapCheckWarnings"), TEXT("5.6"));
 
 FAutoConsoleVariable CVarStripLayerTextureMipsOnLoad(
 	TEXT("landscape.StripLayerMipsOnLoad"),
@@ -4424,7 +4425,7 @@ PRAGMA_ENABLE_DEPRECATION_WARNINGS
 	{
 		auto CreateMapCheckMessage = [](FMessageLog& MessageLog)
 		{
-			if (CVarLandscapeSupressMapCheckWarnings_Nanite->GetBool())
+			if (CVarLandscapeSilenceMapCheckWarnings_Nanite->GetBool())
 			{
 				return MessageLog.Info();
 			}
diff --git a/Engine/Source/Runtime/Landscape/Private/LandscapeEditLayers.cpp b/Engine/Source/Runtime/Landscape/Private/LandscapeEditLayers.cpp
--- a/Engine/Source/Runtime/Landscape/Private/LandscapeEditLayers.cpp
+++ b/Engine/Source/Runtime/Landscape/Private/LandscapeEditLayers.cpp
@@ -254,6 +254,10 @@ TAutoConsoleVariable<FString> CVarLandscapeBatchedMergeVisualLogShowComponentDep
 	"",
 	TEXT("Allows to visualize all the dependency graph for the component specified : use the \"X= Y=\" format to specify the component for which to show the dependencies"));
 
+TAutoConsoleVariable<bool> CVarSilenceMergeBatchResolutionWarning(
+	TEXT("landscape.BatchedMerge.SilenceResolutionWarning"),
+	false,
+	TEXT("When true, don't warn about about exceeding batch merge resolution from landscape.EditLayersLocalMerge.MaxResolutionPerRenderBatch"));
 
 void OnLandscapeEditLayersLocalMergeChanged(IConsoleVariable* CVar)
 {
@@ -6136,7 +6140,7 @@ namespace UE::Landscape::EditLayers::Private
 	 *   - When a component is being rendered by a given batch, all components needed for rendering this given component are present in the batch
 	 *   - All components end up being rendered in at least one of the batches
 	 */
-	TArray<FRenderBatchInfo> DivideIntoBatches(const TBitArray<>& InFinalComponentsToRenderInfoBitIndices, const TArray<FComponentToRenderInfo>& InAllComponentsToRenderInfos)
+	TArray<FRenderBatchInfo> DivideIntoBatches(const TBitArray<>& InFinalComponentsToRenderInfoBitIndices, const TArray<FComponentToRenderInfo>& InAllComponentsToRenderInfos, bool &bInOutWarnedResolution)
 	{
 		TRACE_CPUPROFILER_EVENT_SCOPE(DivideIntoBatches);
 
@@ -6162,8 +6166,21 @@ namespace UE::Landscape::EditLayers::Private
 		FIntPoint MaxBatchResolution = MinBatchResolution.ComponentMax(FIntPoint(DesiredMaxBatchResolution));
 		if ((MinBatchResolution.X > DesiredMaxBatchResolution) || (MinBatchResolution.Y > DesiredMaxBatchResolution))
 		{
-			UE_LOG(LogLandscape, Warning, TEXT("Landscape edit layers merge requires a minimum batch size of resolution %ix%i, which is higher than the current desired maximum batch resolution %ix%i. Consider adjusting the maximum batch resolution (landscape.EditLayersLocalMerge.MaxResolutionPerRenderBatch) or make sure the landscape edit layers renderers in use require smaller work area."),
-				MinBatchResolution.X, MinBatchResolution.Y, DesiredMaxBatchResolution, DesiredMaxBatchResolution);
+			if (!bInOutWarnedResolution)
+			{
+				bInOutWarnedResolution = true;
+				if (CVarSilenceMergeBatchResolutionWarning.GetValueOnGameThread())
+				{
+					// Reduce level to Display and slightly different message
+					UE_LOG(LogLandscape, Display, TEXT("Landscape edit layers merge requires a minimum batch size of resolution %ix%i, which is higher than the current desired maximum batch resolution %ix%i. Consider adjusting the maximum batch resolution (landscape.EditLayersLocalMerge.MaxResolutionPerRenderBatch) or make sure the landscape edit layers renderers in use require smaller work area."),
+						MinBatchResolution.X, MinBatchResolution.Y, DesiredMaxBatchResolution, DesiredMaxBatchResolution);
+				}
+				else
+				{
+					UE_LOG(LogLandscape, Warning, TEXT("Landscape edit layers merge requires a minimum batch size of resolution %ix%i, which is higher than the current desired maximum batch resolution %ix%i. Consider adjusting the maximum batch resolution (landscape.EditLayersLocalMerge.MaxResolutionPerRenderBatch) or make sure the landscape edit layers renderers in use require smaller work area. (Quiet this warning by setting landscape.BatchedMerge.SilenceResolutionWarning to 1)."),
+						MinBatchResolution.X, MinBatchResolution.Y, DesiredMaxBatchResolution, DesiredMaxBatchResolution);
+				}
+			}
 			MaxBatchResolution = MinBatchResolution;
 		}
 
@@ -7117,7 +7134,7 @@ UE::Landscape::EditLayers::FMergeRenderContext ALandscape::PrepareEditLayersMerg
 #endif // ENABLE_VISUAL_LOG
 
 	// Now divide the work into batches as large as possible (but fitting in the desired max batch resolution, if possible): 
-	TArray<FRenderBatchInfo> AllBatchInfos = DivideIntoBatches(FinalComponentsToRenderInfoBitIndices, AllComponentsToRenderInfos);
+	TArray<FRenderBatchInfo> AllBatchInfos = DivideIntoBatches(FinalComponentsToRenderInfoBitIndices, AllComponentsToRenderInfos, bWarnedLayerMergeResolution);
 	{
 		TRACE_CPUPROFILER_EVENT_SCOPE(PrepareBatches);
 		MergeRenderContext.RenderBatches.Reserve(AllBatchInfos.Num());
