Change default landscape edit layer method to BatchMerge.  landscape.EditLayersLocalMerge.Enable default value changes to 2.

diff --git a/Engine/Source/Runtime/Landscape/Private/LandscapeEditLayers.cpp b/Engine/Source/Runtime/Landscape/Private/LandscapeEditLayers.cpp
--- a/Engine/Source/Runtime/Landscape/Private/LandscapeEditLayers.cpp
+++ b/Engine/Source/Runtime/Landscape/Private/LandscapeEditLayers.cpp
@@ -281,7 +281,7 @@ void OnLandscapeEditLayersLocalMergeChanged(IConsoleVariable* CVar)
 	}
 }
 
-int32 LandscapeEditLayersLocalMerge = 0;
+int32 LandscapeEditLayersLocalMerge = 2;
 static FAutoConsoleVariableRef CVarLandscapeEditLayersLocalMerge(
 	TEXT("landscape.EditLayersLocalMerge.Enable"),
 	LandscapeEditLayersLocalMerge,
