Fix for crash in LandscapeGroup when PIE
Seems like it can attempt to capture or patch an edge while the main landscape actor is not registered.
Fixed so that it delays the operation while there is no landscape actor available.

diff --git a/Engine/Source/Runtime/Landscape/Private/LandscapeGroup.cpp b/Engine/Source/Runtime/Landscape/Private/LandscapeGroup.cpp
--- a/Engine/Source/Runtime/Landscape/Private/LandscapeGroup.cpp
+++ b/Engine/Source/Runtime/Landscape/Private/LandscapeGroup.cpp
@@ -590,10 +590,9 @@ void FLandscapeGroup::TickEdgeFixup(ULandscapeSubsystem* LandscapeSubsystem, boo
 
 #if WITH_EDITORONLY_DATA
 			ALandscape* Landscape = EdgeFixup->ActiveComponent->GetLandscapeActor();
-			check(Landscape);
-			if (!Landscape->bGrassUpdateEnabled)
+			if (!Landscape || !Landscape->bGrassUpdateEnabled)
 			{
-				GROUP_DEBUG_LOG_DETAIL(TEXT("-- Snap Component %p (%d,%d) -------- PAUSED (editing)"), EdgeFixup->ActiveComponent.Get(), EdgeFixup->GroupCoord.X, EdgeFixup->GroupCoord.Y);
+				GROUP_DEBUG_LOG_DETAIL(TEXT("-- Snap Component %p (%d,%d) -------- PAUSED (editing or landscape unregistered)"), EdgeFixup->ActiveComponent.Get(), EdgeFixup->GroupCoord.X, EdgeFixup->GroupCoord.Y);
 				continue;
 			}
 #endif // WITH_EDITORONLY_DATA
@@ -634,10 +633,9 @@ void FLandscapeGroup::TickEdgeFixup(ULandscapeSubsystem* LandscapeSubsystem, boo
 
 #if WITH_EDITORONLY_DATA
 			ALandscape* Landscape = EdgeFixup->ActiveComponent->GetLandscapeActor();
-			check(Landscape);
-			if (!Landscape->bGrassUpdateEnabled)
+			if (!Landscape || !Landscape->bGrassUpdateEnabled)
 			{
-				GROUP_DEBUG_LOG_DETAIL(TEXT("-- Patch Component %p (%d,%d) -------- PAUSED (editing)"), EdgeFixup->ActiveComponent.Get(), EdgeFixup->GroupCoord.X, EdgeFixup->GroupCoord.Y);
+				GROUP_DEBUG_LOG_DETAIL(TEXT("-- Patch Component %p (%d,%d) -------- PAUSED (editing or landscape unregistered)"), EdgeFixup->ActiveComponent.Get(), EdgeFixup->GroupCoord.X, EdgeFixup->GroupCoord.Y);
 				continue;
 			}
 #endif // WITH_EDITORONLY_DATA
