UE: Fix nullrhi crash if a landscape patch is needed

diff --git a/Engine/Source/Runtime/Landscape/Private/LandscapeEdgeFixup.cpp b/Engine/Source/Runtime/Landscape/Private/LandscapeEdgeFixup.cpp
--- a/Engine/Source/Runtime/Landscape/Private/LandscapeEdgeFixup.cpp
+++ b/Engine/Source/Runtime/Landscape/Private/LandscapeEdgeFixup.cpp
@@ -1098,14 +1098,18 @@ void ULandscapeHeightmapTextureEdgeFixup::PatchTextureEdge_Internal(UE::Landscap
 
 	using namespace UE::Landscape;
 
+	FTextureResource* HeightTextureResource = HeightmapTexture->GetResource();
+	if (!HeightTextureResource)
+	{
+		return;
+	}
+
 #if WITH_EDITOR
 	// this should be caught earlier -- we can't patch default textures
 	check(!HeightmapTexture->IsDefaultTexture());
 #endif // WITH_EDITOR
 	check(!IsDiagonalCorner(EdgeIndex));
 
-	FTextureResource* HeightTextureResource = HeightmapTexture->GetResource();
-
 	const int32 NumMips = HeightmapTexture->GetNumMips();
 	const int32 ResidentMips = HeightmapTexture->GetNumResidentMips();
 
@@ -1188,14 +1192,18 @@ void ULandscapeHeightmapTextureEdgeFixup::PatchTextureCorner_Internal(UE::Landsc
 
 	using namespace UE::Landscape;
 
+	FTextureResource* HeightTextureResource = HeightmapTexture->GetResource();
+	if (!HeightTextureResource)
+	{
+		return;
+	}
+
 #if WITH_EDITOR
 	// this should be caught earlier -- we can't patch default textures
 	check(!HeightmapTexture->IsDefaultTexture());
 #endif // WITH_EDITOR
 	check(IsDiagonalCorner(CornerIndex));
 
-	FTextureResource* HeightTextureResource = HeightmapTexture->GetResource();
-
 	const int32 NumMips = HeightmapTexture->GetNumMips();
 	const int32 ResidentMips = HeightmapTexture->GetNumResidentMips();
 
