Fix assert with BatchMerge when a weightmap layer is only written by a BP brush.  The layer may be missing from the LandscapeInfo Layers list.

diff --git a/Engine/Source/Runtime/Landscape/Private/LandscapeEditLayerRenderer.cpp b/Engine/Source/Runtime/Landscape/Private/LandscapeEditLayerRenderer.cpp
--- a/Engine/Source/Runtime/Landscape/Private/LandscapeEditLayerRenderer.cpp
+++ b/Engine/Source/Runtime/Landscape/Private/LandscapeEditLayerRenderer.cpp
@@ -1554,10 +1554,11 @@ bool FComponentMergeRenderInfo::operator<(const FComponentMergeRenderInfo& InOth
 
 // ----------------------------------------------------------------------------------
 
-FMergeRenderParams::FMergeRenderParams(TArray<ULandscapeComponent*> InComponentsToMerge, const TArrayView<FEditLayerRendererState>& InEditLayerRendererStates, const TSet<FName>& InWeightmapLayerNames)
+FMergeRenderParams::FMergeRenderParams(TArray<ULandscapeComponent*> InComponentsToMerge, const TArrayView<FEditLayerRendererState>& InEditLayerRendererStates, const TSet<FName>& InWeightmapLayerNames, bool bInRequestAllLayers)
 	: ComponentsToMerge(InComponentsToMerge)
 	, EditLayerRendererStates(InEditLayerRendererStates)
 	, WeightmapLayerNames(InWeightmapLayerNames)
+	, bRequestAllLayers(bInRequestAllLayers)
 {
 }
 
diff --git a/Engine/Source/Runtime/Landscape/Private/LandscapeEditLayers.cpp b/Engine/Source/Runtime/Landscape/Private/LandscapeEditLayers.cpp
--- a/Engine/Source/Runtime/Landscape/Private/LandscapeEditLayers.cpp
+++ b/Engine/Source/Runtime/Landscape/Private/LandscapeEditLayers.cpp
@@ -5919,14 +5919,16 @@ void ULandscapeEditLayerPersistent::RenderLayer(FRenderParams& RenderParams, UE:
 			// Prepare a list of all paint layers to merge and do it all in 1 command :
 			TArray<FMergePaintLayerParams> MergePaintLayersParams;
 			MergePaintLayersParams.Reserve(NumTargetLayersInGroup);
-			for (const FName& TargetLayerName : RenderParams.TargetLayerGroupLayerNames)
+
+			check(NumTargetLayersInGroup == RenderParams.TargetLayerGroupLayerInfos.Num());
+			for (int32 i = 0; i < NumTargetLayersInGroup; ++i)
 			{
+				const FName& TargetLayerName = RenderParams.TargetLayerGroupLayerNames[i];
+				const ULandscapeLayerInfoObject* LayerInfoObj = RenderParams.TargetLayerGroupLayerInfos[i];
+				check(LayerInfoObj != nullptr);
+
 				FMergePaintLayerParams& MergePaintLayerParams = MergePaintLayersParams.Emplace_GetRef(TargetLayerName);
 
-				const FLandscapeInfoLayerSettings* LayerSettings = Info->Layers.FindByPredicate([TargetLayerName](const FLandscapeInfoLayerSettings& InLayerSettings) { return (InLayerSettings.LayerName == TargetLayerName); });
-				check(LayerSettings != nullptr);
-				ULandscapeLayerInfoObject* LayerInfoObj = LayerSettings->LayerInfoObj;
-				check(LayerInfoObj != nullptr);
 				bool bIsSubtractiveBlend = false;
 				if (const bool* bSubstractiveInLayer = LayerStruct->WeightmapLayerAllocationBlend.Find(LayerInfoObj); (bSubstractiveInLayer != nullptr) && *bSubstractiveInLayer)
 				{
@@ -6720,7 +6722,16 @@ UE::Landscape::EditLayers::FMergeRenderContext ALandscape::PrepareEditLayersMerg
 		check(CandidateEditLayerRendererStates.Num() == RelevantTargetLayerBitIndicesPerRenderer.Num());
 
 		const int32 CandidateNumRenderers = CandidateEditLayerRendererStates.Num();
-		TBitArray<> RequestedTargetLayerBitIndices = MergeRenderContext.ConvertTargetLayerNamesToBitIndices(InMergeRenderParams.WeightmapLayerNames.Array());
+
+		TBitArray<> RequestedTargetLayerBitIndices;
+		if (InMergeRenderParams.bRequestAllLayers)
+		{
+			RequestedTargetLayerBitIndices = MergeRenderContext.ValidTargetLayerBitIndices;
+		}
+		else
+		{
+			RequestedTargetLayerBitIndices = MergeRenderContext.ConvertTargetLayerNamesToBitIndices(InMergeRenderParams.WeightmapLayerNames.Array());
+		}
 		// No need to retain the invalid target layers : 
 		RequestedTargetLayerBitIndices.CombineWithBitwiseAND(MergeRenderContext.ValidTargetLayerBitIndices, EBitwiseOperatorFlags::MinSize);
 		// Early-out when there's nothing to do :
@@ -10432,7 +10443,9 @@ int32 ALandscape::PerformLayersWeightmapsBatchedMerge(FUpdateLayersContentContex
 
 	// Prepare the merge : 
 	TSet<ULandscapeComponent*> LandscapeComponentsWeightmapsToRenderBefore(InUpdateLayersContentContext.LandscapeComponentsWeightmapsToRender);
-	FMergeRenderParams MergeRenderParams(InUpdateLayersContentContext.LandscapeComponentsWeightmapsToRender, RendererStates, RequestedWeightmapLayerNames);
+	// Not yet ready to selectively render layers, so request all layers
+	const bool bRequestAllLayers = true;
+	FMergeRenderParams MergeRenderParams(InUpdateLayersContentContext.LandscapeComponentsWeightmapsToRender, RendererStates, RequestedWeightmapLayerNames, bRequestAllLayers);
 	FMergeRenderContext MergeRenderContext = PrepareEditLayersMergeRenderContext(MergeContext, MergeRenderParams);
 	if (!MergeRenderContext.IsValid())
 	{
diff --git a/Engine/Source/Runtime/Landscape/Public/LandscapeEditLayerRenderer.h b/Engine/Source/Runtime/Landscape/Public/LandscapeEditLayerRenderer.h
--- a/Engine/Source/Runtime/Landscape/Public/LandscapeEditLayerRenderer.h
+++ b/Engine/Source/Runtime/Landscape/Public/LandscapeEditLayerRenderer.h
@@ -631,7 +631,7 @@ private:
  */
 struct FMergeRenderParams
 {
-	FMergeRenderParams(TArray<ULandscapeComponent*> InComponentsToMerge, const TArrayView<FEditLayerRendererState>& InEditLayerRendererStates, const TSet<FName>& InWeightmapLayerNames = {});
+	FMergeRenderParams(TArray<ULandscapeComponent*> InComponentsToMerge, const TArrayView<FEditLayerRendererState>& InEditLayerRendererStates, const TSet<FName>& InWeightmapLayerNames = {}, bool bInRequestAllLayers = false);
 
 	/** List of components that need merging */
 	TArray<ULandscapeComponent*> ComponentsToMerge;
@@ -641,6 +641,9 @@ struct FMergeRenderParams
 	
 	/** List of weightmap layers being requested. */
 	TSet<FName> WeightmapLayerNames;
+
+	/** Ignore the computed list of requested layers and instead request all valid layers */
+	bool bRequestAllLayers = false;
 };
 
 
