[AutoRTFM] Fix an AutoRTFM crash when cleaning up landscape collision data.

This can be reached in Verse by triggering an end-of-game, which immediately unregisters actors. It can also be reached by triggering cinematics (as seen in the attached Jira).

diff --git a/Engine/Source/Runtime/Landscape/Private/LandscapeCollision.cpp b/Engine/Source/Runtime/Landscape/Private/LandscapeCollision.cpp
--- a/Engine/Source/Runtime/Landscape/Private/LandscapeCollision.cpp
+++ b/Engine/Source/Runtime/Landscape/Private/LandscapeCollision.cpp
@@ -1,6 +1,7 @@
 // Copyright Epic Games, Inc. All Rights Reserved.
 
 #include "CoreMinimal.h"
+#include "AutoRTFM/AutoRTFM.h"
 #include "GenericPlatform/GenericPlatformStackWalk.h"
 #include "Materials/Material.h"
 #include "Misc/Guid.h"
@@ -2208,16 +2209,24 @@ void ULandscapeHeightfieldCollisionComponent::DeferredDestroyCollision(const TRe
 	// dirty proxy list on the physics thread with a geometry that has been destructed by the lifetime
 	// extender falling out of scope.
 	// To avoid this we dispatch an empty callable with the unique geometries which runs after the
-	// proxy queue will have been cleared, avoiding a use-after-free.
+	// proxy queue will have been cleared, avoiding a use-after-free. 
+	// We also avoid enqueueing any off-thread work until the AutoRTFM transaction has completed.
 	// #TODO auto ref counted user objects for Chaos.
-	PhysScene->GetSolver()->EnqueueCommandImmediate(
-		[ComplexHeightfield = MoveTemp(HeightfieldRefLifetimeExtender->HeightfieldGeometry)
-		, SimpleHeightfield = MoveTemp(HeightfieldRefLifetimeExtender->HeightfieldSimpleGeometry)
+	AutoRTFM::OnCommit([this
+	                    , PhysScene
+					    , ComplexHeightfield = MoveTemp(HeightfieldRefLifetimeExtender->HeightfieldGeometry)
+					    , SimpleHeightfield = MoveTemp(HeightfieldRefLifetimeExtender->HeightfieldSimpleGeometry)
 #if WITH_EDITORONLY_DATA
-		, EditorHeightfield = MoveTemp(HeightfieldRefLifetimeExtender->EditorHeightfieldGeometry)
+					    , EditorHeightfield = MoveTemp(HeightfieldRefLifetimeExtender->EditorHeightfieldGeometry)
 #endif
-		]
-	() mutable
+                       ]() mutable
+	{
+		PhysScene->GetSolver()->EnqueueCommandImmediate([ComplexHeightfield = MoveTemp(ComplexHeightfield)
+														 , SimpleHeightfield = MoveTemp(SimpleHeightfield)
+#if WITH_EDITORONLY_DATA
+														 , EditorHeightfield = MoveTemp(EditorHeightfield)
+#endif 
+														]() mutable
 		{
 			ComplexHeightfield = nullptr;
 			SimpleHeightfield = nullptr;
@@ -2225,6 +2234,7 @@ void ULandscapeHeightfieldCollisionComponent::DeferredDestroyCollision(const TRe
 			EditorHeightfield = nullptr;
 #endif
 		});
+	});
 }
 
 #if WITH_EDITORONLY_DATA
