Landscape - Simplified handling of the visibility of landscape scene proxies during level AddToWorld/RemoveFromWorld
* Currently, as the add to world sequence is timesliced, we force hide scene proxies until we can make them all visible at the same time.
* Landscape proxies delay the registration of their sections until that point, otherwise it could lead to incorrect LOD mapping between the currently visible LS proxy HLOD and the neighbor proxies.
* Override the new OnForceHiddenChanged() to register or unregister the landscape sections

diff --git a/Engine/Source/Runtime/Engine/Private/PrimitiveSceneProxy.cpp b/Engine/Source/Runtime/Engine/Private/PrimitiveSceneProxy.cpp
--- a/Engine/Source/Runtime/Engine/Private/PrimitiveSceneProxy.cpp
+++ b/Engine/Source/Runtime/Engine/Private/PrimitiveSceneProxy.cpp
@@ -741,11 +741,10 @@ FPrimitiveSceneProxy::FPrimitiveSceneProxy(const FPrimitiveSceneProxyDesc& InPro
 	}
 }
 
-bool FPrimitiveSceneProxy::OnLevelAddedToWorld_RenderThread()
+void FPrimitiveSceneProxy::OnLevelAddedToWorld_RenderThread()
 {
 	check(bShouldNotifyOnWorldAddRemove);
 	SetForceHidden(false);
-	return false;
 }
 
 void FPrimitiveSceneProxy::OnLevelRemovedFromWorld_RenderThread()
diff --git a/Engine/Source/Runtime/Engine/Public/PrimitiveSceneProxy.h b/Engine/Source/Runtime/Engine/Public/PrimitiveSceneProxy.h
--- a/Engine/Source/Runtime/Engine/Public/PrimitiveSceneProxy.h
+++ b/Engine/Source/Runtime/Engine/Public/PrimitiveSceneProxy.h
@@ -549,22 +549,6 @@ public:
 	 */
 	virtual void OnTransformChanged(FRHICommandListBase& RHICmdList) { }
 
-	/**
-	 * Called to notify the proxy that the level has been fully added to
-	 * the world and the primitive will now be rendered.
-	 * Only called if bShouldNotifyOnWorldAddRemove is set to true.
-	 * 
-	 * @return return true if the primitive should be added to the scene
-	 */
-	ENGINE_API virtual bool OnLevelAddedToWorld_RenderThread();
-
-	/**
-	 * Called to notify the proxy that the level has been fully removed from
-	 * the world and the primitive will not be rendered.
-	 * Only called if bShouldNotifyOnWorldAddRemove is set to true.
-	 */
-	ENGINE_API virtual void OnLevelRemovedFromWorld_RenderThread();
-
 	/**
 	* @return true if the proxy can be culled when occluded by other primitives
 	*/
@@ -1178,7 +1162,30 @@ protected:
 		OwnerName = InOwnerName;
 	}
 
-	void SetForceHidden(bool bForceHiddenIn) {bForceHidden = bForceHiddenIn;}
+	/**
+	 * Called to notify the proxy that the level has been fully added to
+	 * the world and the primitive will now be rendered.
+	 * Only called if bShouldNotifyOnWorldAddRemove is set to true.
+	 */
+	ENGINE_API void OnLevelAddedToWorld_RenderThread();
+
+	/**
+	 * Called to notify the proxy that the level has been fully removed from
+	 * the world and the primitive will not be rendered.
+	 * Only called if bShouldNotifyOnWorldAddRemove is set to true.
+	 */
+	ENGINE_API void OnLevelRemovedFromWorld_RenderThread();
+
+	void SetForceHidden(bool bForceHiddenIn)
+	{
+		if (bForceHidden != bForceHiddenIn)
+		{
+			bForceHidden = bForceHiddenIn;
+			OnForceHiddenChanged();
+		}
+	}
+
+	ENGINE_API virtual void OnForceHiddenChanged() {}
 
 	/** 
 	 * Call during setup to set flags to indicate GPU-Scene support for the proxy if GPU-Scene is enabled & supported for the current feature level.
diff --git a/Engine/Source/Runtime/Landscape/Private/LandscapeRender.cpp b/Engine/Source/Runtime/Landscape/Private/LandscapeRender.cpp
--- a/Engine/Source/Runtime/Landscape/Private/LandscapeRender.cpp
+++ b/Engine/Source/Runtime/Landscape/Private/LandscapeRender.cpp
@@ -1333,39 +1333,6 @@ FLandscapeRenderSystem* FLandscapeSceneViewExtension::GetLandscapeRenderSystem(c
 	return (Found != nullptr) ? *Found : nullptr;
 }
 
-//
-// FLandscapeVisibilityHelper
-//
-void FLandscapeVisibilityHelper::Init(UPrimitiveComponent* LandscapeComponent, FPrimitiveSceneProxy* ProxyIn)
-{
-	// Flag components to render only after level will be fully added to the world
-	ULevel* ComponentLevel = LandscapeComponent->GetComponentLevel();
-	bRequiresVisibleLevelToRender = (ComponentLevel && ComponentLevel->bRequireFullVisibilityToRender);
-	bIsComponentLevelVisible = (!ComponentLevel || ComponentLevel->bIsVisible);
-}
-
-bool FLandscapeVisibilityHelper::OnAddedToWorld()
-{
-	if (bIsComponentLevelVisible)
-	{
-		return false;
-	}
-
-	bIsComponentLevelVisible = true;
-	return true;
-}
-
-bool FLandscapeVisibilityHelper::OnRemoveFromWorld()
-{
-	if (!bIsComponentLevelVisible)
-	{
-		return false;
-	}
-
-	bIsComponentLevelVisible = false;
-	return true;
-}
-
 #if WITH_EDITOR
 ULandscapeEditResourcesSubsystem* FLandscapeComponentSceneProxy::LandscapeEditResourcesSubsystem = nullptr;
 #endif // WITH_EDITOR
@@ -1422,18 +1389,6 @@ FLandscapeComponentSceneProxy::FLandscapeComponentSceneProxy(ULandscapeComponent
 	// Disable GDRTI since there's seems to be a race condition in FLandscapeRenderSystem::GetSectionLODValue(...) otherwise
 	bSupportsParallelGDRTI = false;
 
-	VisibilityHelper.Init(InComponent, this);
-
-	if (!VisibilityHelper.ShouldBeVisible())
-	{
-		SetForceHidden(true);
-	}
-
-	if (VisibilityHelper.RequiresVisibleLevelToRender())
-	{
-		bShouldNotifyOnWorldAddRemove = true;
-	}
-
 	bNaniteActive = InComponent->IsNaniteActive();
 	bUsesLandscapeCulling = (!bNaniteActive && Culling::UseCulling(GetScene().GetShaderPlatform()));
 
@@ -1745,7 +1700,7 @@ void FLandscapeComponentSceneProxy::CreateRenderThreadResources(FRHICommandListB
 
 	FLandscapeRenderSystem::CreateResources(RHICmdList, this);
 
-	if (VisibilityHelper.ShouldBeVisible())
+	if (!IsForceHidden())
 	{
 		FLandscapeRenderSystem::RegisterSection(this);
 	}
@@ -1982,23 +1937,15 @@ void FLandscapeComponentSceneProxy::DestroyRenderThreadResources()
 	Culling::UnregisterLandscape(LandscapeKey);
 }
 
-bool FLandscapeComponentSceneProxy::OnLevelAddedToWorld_RenderThread()
+void FLandscapeComponentSceneProxy::OnForceHiddenChanged()
 {
-	if (VisibilityHelper.OnAddedToWorld())
+	if (IsForceHidden())
 	{
-		SetForceHidden(false);
-		FLandscapeRenderSystem::RegisterSection(this);
-		return true;
+		FLandscapeRenderSystem::UnregisterSection(this);
 	}
-	return false;
-}
-
-void FLandscapeComponentSceneProxy::OnLevelRemovedFromWorld_RenderThread()
-{
-	if (VisibilityHelper.OnRemoveFromWorld())
+	else
 	{
-		SetForceHidden(true);
-		FLandscapeRenderSystem::UnregisterSection(this);
+		FLandscapeRenderSystem::RegisterSection(this);
 	}
 }
 
@@ -4764,13 +4711,6 @@ FLandscapeMeshProxySceneProxy::FLandscapeMeshProxySceneProxy(UStaticMeshComponen
 {
 	check(InLODGroupKey == 0 || ComponentResolution != 0);	// if using LODGroupKey, we require HLODs to be rebuilt to capture the component resolution and section transform data (position / orientation)
 
-	VisibilityHelper.Init(InComponent, this);
-
-	if (VisibilityHelper.RequiresVisibleLevelToRender())
-	{
-		bShouldNotifyOnWorldAddRemove = true;
-	}
-
 	ProxySectionsInfos.Empty(InProxySectionsBases.Num());
 	
 	check(InProxySectionsBases.Num() == InProxySectionsCentersLocalSpace.Num() ||
@@ -4815,30 +4755,21 @@ void FLandscapeMeshProxySceneProxy::CreateRenderThreadResources(FRHICommandListB
 		FLandscapeRenderSystem::CreateResources(RHICmdList, Info.Get());
 	}
 
-	if (VisibilityHelper.ShouldBeVisible())
+	if (!IsForceHidden())
 	{
 		RegisterSections();
 	}
 }
 
-bool FLandscapeMeshProxySceneProxy::OnLevelAddedToWorld_RenderThread()
+void FLandscapeMeshProxySceneProxy::OnForceHiddenChanged()
 {
-	if (VisibilityHelper.OnAddedToWorld())
+	if (IsForceHidden())
 	{
-		SetForceHidden(false);
-		RegisterSections();
-		return true;
+		UnregisterSections();
 	}
-
-	return false;
-}
-
-void FLandscapeMeshProxySceneProxy::OnLevelRemovedFromWorld_RenderThread()
-{
-	if (VisibilityHelper.OnRemoveFromWorld())
+	else
 	{
-		SetForceHidden(true);
-		UnregisterSections();
+		RegisterSections();
 	}
 }
 
diff --git a/Engine/Source/Runtime/Landscape/Public/LandscapeRender.h b/Engine/Source/Runtime/Landscape/Public/LandscapeRender.h
--- a/Engine/Source/Runtime/Landscape/Public/LandscapeRender.h
+++ b/Engine/Source/Runtime/Landscape/Public/LandscapeRender.h
@@ -647,23 +647,6 @@ private:
 	int32 NumViewsWithShowCollisionAcc = 0; // Accumulate the number of views with collision
 };
 
-
-//
-// FLandscapeVisibilityHelper
-//
-class FLandscapeVisibilityHelper
-{
-public:
-	void Init(UPrimitiveComponent* LandscapeComponent, FPrimitiveSceneProxy* ProxyIn);
-	bool OnAddedToWorld();
-	bool OnRemoveFromWorld();
-	bool ShouldBeVisible() const { return !bRequiresVisibleLevelToRender || bIsComponentLevelVisible; }
-	bool RequiresVisibleLevelToRender() const { return bRequiresVisibleLevelToRender; }
-private:
-	bool bRequiresVisibleLevelToRender = false;
-	bool bIsComponentLevelVisible = false;
-};
-
 //
 // FLandscapeDebugOptions
 //
@@ -695,15 +678,14 @@ public:
 	FLandscapeMeshProxySceneProxy(UStaticMeshComponent* InComponent, const FGuid& InLandscapeGuid, const TArray<FIntPoint>& InProxySectionsBases, const TArray<FVector>& InProxySectionsCentersLocalSpace, const FVector& InComponentXVector, const FVector& InComponentYVector, const FTransform& LocalToWorld, int32 ComponentResolution, int8 InProxyLOD, uint32 InLODGroupKey, uint32 LandscapeKey);
 	virtual void CreateRenderThreadResources(FRHICommandListBase& RHICmdList) override;
 	virtual void DestroyRenderThreadResources() override;
-	virtual bool OnLevelAddedToWorld_RenderThread() override;
-	virtual void OnLevelRemovedFromWorld_RenderThread() override;
+
+protected:
+	virtual void OnForceHiddenChanged() override;
 
 private:
 	void RegisterSections();
 	void UnregisterSections();
 
-	FLandscapeVisibilityHelper VisibilityHelper;
-
 	TArray<TUniquePtr<FLandscapeSectionInfo>> ProxySectionsInfos;
 };
 
@@ -894,8 +876,6 @@ protected:
 	// Cached versions of these
 	FMatrix WorldToLocal;
 
-	FLandscapeVisibilityHelper VisibilityHelper;
-
 #if !UE_BUILD_SHIPPING
 	FName DebugName;
 #endif // !UE_BUILD_SHIPPING
@@ -906,6 +886,8 @@ protected:
 	LANDSCAPE_API bool GetMeshElementForFixedGrid(int32 InLodIndex, FMaterialRenderProxy* InMaterialInterface, FMeshBatch& OutMeshBatch, TArray<FLandscapeBatchElementParams>& OutStaticBatchParamArray) const;
 	template<class ArrayType> bool GetStaticMeshElement(int32 LODIndex, bool bForToolMesh, FMeshBatch& MeshBatch, ArrayType& OutStaticBatchParamArray) const;
 
+	virtual void OnForceHiddenChanged() override;
+
 public:
 	// constructor
 	LANDSCAPE_API FLandscapeComponentSceneProxy(ULandscapeComponent* InComponent);
@@ -921,9 +903,7 @@ public:
 	LANDSCAPE_API virtual void OnTransformChanged(FRHICommandListBase& RHICmdList) override;
 	LANDSCAPE_API virtual void CreateRenderThreadResources(FRHICommandListBase& RHICmdList) override;
 	LANDSCAPE_API virtual void DestroyRenderThreadResources() override;
-	LANDSCAPE_API virtual bool OnLevelAddedToWorld_RenderThread() override;
-	LANDSCAPE_API virtual void OnLevelRemovedFromWorld_RenderThread() override;
-	
+
 	friend class ULandscapeComponent;
 	friend class FLandscapeVertexFactoryVertexShaderParameters;
 	friend class FLandscapeXYOffsetVertexFactoryVertexShaderParameters;
diff --git a/Engine/Source/Runtime/Renderer/Private/RendererScene.cpp b/Engine/Source/Runtime/Renderer/Private/RendererScene.cpp
--- a/Engine/Source/Runtime/Renderer/Private/RendererScene.cpp
+++ b/Engine/Source/Runtime/Renderer/Private/RendererScene.cpp
@@ -6526,12 +6526,7 @@ void FScene::Update(FRDGBuilder& GraphBuilder, const FUpdateParameters& Paramete
 			{
 				for (FPrimitiveSceneInfo* Primitive : *LevelPrimitives)
 				{
-					// If the primitive proxy returns true, it needs it's static meshes added to the scene
-					if (Primitive->Proxy->OnLevelAddedToWorld_RenderThread() && QueueAddStaticMeshes(Primitive))
-					{
-						Primitive->RemoveStaticMeshes();
-					}
-
+					Primitive->Proxy->OnLevelAddedToWorld_RenderThread();
 					InvalidateNaniteGPUScene(Primitive);
 				}
 			}
