Fixed crash when renaming a landscape target layer with batched merge active

#rnx

diff --git a/Engine/Source/Runtime/Landscape/Private/LandscapeEditLayers.cpp b/Engine/Source/Runtime/Landscape/Private/LandscapeEditLayers.cpp
--- a/Engine/Source/Runtime/Landscape/Private/LandscapeEditLayers.cpp
+++ b/Engine/Source/Runtime/Landscape/Private/LandscapeEditLayers.cpp
@@ -5541,7 +5541,7 @@ TArray<UE::Landscape::EditLayers::FEditLayerRenderItem> ULandscapeEditLayerPersi
 						{
 							OutputTargetTypeState.AddTargetTypeMask(ELandscapeToolTargetTypeFlags::Weightmap);
 						}
-						OutputTargetTypeState.AddWeightmapChecked(LayerAllocationInfo.GetLayerName());
+						OutputTargetTypeState.AddWeightmap(LayerAllocationInfo.GetLayerName());
 					}
 				}
 
