Fixed race condition crash in landscape when async physics state creation is enabled : Chaos::FPhysicalMaterialManager's OnMaterialUpdated can then be called from any thread

#rnx

diff --git a/Engine/Source/Runtime/Landscape/Private/LandscapeModule.cpp b/Engine/Source/Runtime/Landscape/Private/LandscapeModule.cpp
--- a/Engine/Source/Runtime/Landscape/Private/LandscapeModule.cpp
+++ b/Engine/Source/Runtime/Landscape/Private/LandscapeModule.cpp
@@ -1,6 +1,7 @@
 // Copyright Epic Games, Inc. All Rights Reserved.
 
 #include "LandscapeModule.h"
+#include "Containers/Ticker.h"
 #include "Serialization/CustomVersion.h"
 #include "Modules/ModuleManager.h"
 #include "UObject/UObjectHash.h"
@@ -288,12 +289,15 @@ void FLandscapeModule::ShutdownModule()
 #if ALLOW_LANDSCAPE_COLLISION_COMPONENT_DEBUG_DISPLAY
 void FLandscapeModule::OnUpdatePhysicalMaterial(Chaos::FMaterialHandle InHandle)
 {
-	// Just rebuild all collision components' render states when the Chaos material changes, this is just for debug purposes anyway : 
-	for (ULandscapeHeightfieldCollisionComponent* LandscapeHeightfieldCollisionComponent 
-		: TObjectRange<ULandscapeHeightfieldCollisionComponent>(RF_ClassDefaultObject | RF_ArchetypeObject, /*bIncludeDerivedClasses = */true, EInternalObjectFlags::Garbage))
-	{
-		LandscapeHeightfieldCollisionComponent->MarkRenderStateDirty();
-	}
+	ExecuteOnGameThread(UE_SOURCE_LOCATION, []
+		{
+			// Just rebuild all collision components' render states when the Chaos material changes, this is just for debug purposes anyway : 
+			for (ULandscapeHeightfieldCollisionComponent* LandscapeHeightfieldCollisionComponent 
+				: TObjectRange<ULandscapeHeightfieldCollisionComponent>(RF_ClassDefaultObject | RF_ArchetypeObject, /*bIncludeDerivedClasses = */true, EInternalObjectFlags::Garbage))
+			{
+				LandscapeHeightfieldCollisionComponent->MarkRenderStateDirty();
+			}
+		});
 }
 #endif // ALLOW_LANDSCAPE_COLLISION_COMPONENT_DEBUG_DISPLAY
 
