Fix for crashes in LandscapeTextureStreamingManager when passed null Textures.

It seems to specifically be hit (rarely) when iterating the landscape weightmap textures.  Here we handle the null weightmap textures at the outer level before calling the function.

diff --git a/Engine/Source/Runtime/Landscape/Private/LandscapeEditLayers.cpp b/Engine/Source/Runtime/Landscape/Private/LandscapeEditLayers.cpp
--- a/Engine/Source/Runtime/Landscape/Private/LandscapeEditLayers.cpp
+++ b/Engine/Source/Runtime/Landscape/Private/LandscapeEditLayers.cpp
@@ -4102,7 +4102,7 @@ bool ALandscape::PrepareTextureResources(bool bInWaitForStreaming)
 		for (ULandscapeComponent* Component : Proxy->LandscapeComponents)
 		{
 			UTexture2D* ComponentHeightmap = Component->GetHeightmap();
-
+			check(ComponentHeightmap);
 			{
 				bool bIsTextureReady = TextureStreamingManager->RequestTextureFullyStreamedInForever(ComponentHeightmap, bInWaitForStreaming);
 				if (!bIsTextureReady)
@@ -4123,6 +4123,11 @@ bool ALandscape::PrepareTextureResources(bool bInWaitForStreaming)
 
 			for (UTexture2D* ComponentWeightmap : Component->GetWeightmapTextures())
 			{
+				if (!ComponentWeightmap)
+				{
+					continue; // ignore missing weightmaps
+				}
+
 				bool bIsTextureReady = TextureStreamingManager->RequestTextureFullyStreamedInForever(ComponentWeightmap, bInWaitForStreaming);
 				// If the texture is not ready, start tracking its state changes to be notified when it's fully streamed in : 
 				if (!bIsTextureReady)
@@ -4276,11 +4281,15 @@ bool ALandscape::PrepareLayersTextureResources(const TArray<FLandscapeLayer>& In
 			{
 				if (FLandscapeLayerComponentData* ComponentLayerData = Component->GetLayerData(Layer.Guid))
 				{
+					check(ComponentLayerData->HeightmapData.Texture);
 					bIsReady &= TextureStreamingManager->RequestTextureFullyStreamedInForever(ComponentLayerData->HeightmapData.Texture, bInWaitForStreaming);
 
 					for (UTexture2D* LayerWeightmap : ComponentLayerData->WeightmapData.Textures)
 					{
-						bIsReady &= TextureStreamingManager->RequestTextureFullyStreamedInForever(LayerWeightmap, bInWaitForStreaming);
+						if (LayerWeightmap) // ignore missing weightmaps
+						{
+							bIsReady &= TextureStreamingManager->RequestTextureFullyStreamedInForever(LayerWeightmap, bInWaitForStreaming);
+						}
 					}
 				}
 			}
@@ -8900,6 +8909,7 @@ void ALandscape::ReallocateLayersWeightmaps(FUpdateLayersContentContext& InUpdat
 	FLandscapeTextureStreamingManager* TextureStreamingManager = GetWorld()->GetSubsystem<ULandscapeSubsystem>()->GetTextureStreamingManager();
 	for (UTexture* Texture : NewCreatedTextures)
 	{
+		check(Texture);
 		TextureStreamingManager->RequestTextureFullyStreamedInForever(Texture, /* bWaitForStreaming= */ true);
 	}
 
diff --git a/Engine/Source/Runtime/Landscape/Private/LandscapeTextureStreamingManager.cpp b/Engine/Source/Runtime/Landscape/Private/LandscapeTextureStreamingManager.cpp
--- a/Engine/Source/Runtime/Landscape/Private/LandscapeTextureStreamingManager.cpp
+++ b/Engine/Source/Runtime/Landscape/Private/LandscapeTextureStreamingManager.cpp
@@ -63,6 +63,7 @@ TArray<FLandscapeTextureStreamingManager*> FLandscapeTextureStreamingManager::Al
 bool FLandscapeTextureStreamingManager::RequestTextureFullyStreamedIn(UTexture* Texture, bool bWaitForStreaming)
 {
 	using namespace UE::Landscape;
+	check(Texture != nullptr);
 	TWeakObjectPtr<UTexture> TexturePtr = Texture;
 	FTextureState& State = TextureStates.FindOrAdd(TexturePtr);
 
@@ -90,6 +91,7 @@ bool FLandscapeTextureStreamingManager::RequestTextureFullyStreamedIn(UTexture*
 bool FLandscapeTextureStreamingManager::RequestTextureFullyStreamedInForever(UTexture* Texture, bool bWaitForStreaming)
 {
 	using namespace UE::Landscape;
+	check(Texture != nullptr);
 	TWeakObjectPtr<UTexture> TexturePtr = Texture;
 	FTextureState& State = TextureStates.FindOrAdd(TexturePtr);
 	State.bForever = true;
