Fixed crash introduced in UE5.5 (likely happening following an undo, but also 100% repro when spawning a landscape actor via the BP Spawn Actor function)

diff --git a/Engine/Source/Runtime/Landscape/Private/LandscapeSubsystem.cpp b/Engine/Source/Runtime/Landscape/Private/LandscapeSubsystem.cpp
--- a/Engine/Source/Runtime/Landscape/Private/LandscapeSubsystem.cpp
+++ b/Engine/Source/Runtime/Landscape/Private/LandscapeSubsystem.cpp
@@ -735,9 +735,10 @@ void ULandscapeSubsystem::Tick(float DeltaTime)
 	for (TObjectPtr<ALandscape> ActorPtr : LandscapeActors)
 	{
 		ALandscape* Landscape = ActorPtr.Get();
+		if (ULandscapeInfo* LandscapeInfo = Landscape->GetLandscapeInfo())
 		{
 			const bool bLandscapeIsUpToDate = Landscape->IsUpToDate();
-			const bool bLandscapeSupportsEditing = Landscape->GetLandscapeInfo()->SupportsLandscapeEditing();
+			const bool bLandscapeSupportsEditing = LandscapeInfo->SupportsLandscapeEditing();
 			const bool bLandscapeUpdateAllowed = bLandscapeSupportsEditing && (Landscape->GetWorld()->GetFeatureLevel() >= ERHIFeatureLevel::SM5);
 
 			// if either of these things are true, then we wait for them to complete before running ANY grass map updates..
