Fixed crash when landscape proxies (somehow) end up with an invalid root component. It's unclear how this can happen but the landscape collision heightfield component failed to survive this as it would contain a Chaos heightfield with an invalid heighfield geometry.
On top of this, we now try to re-assign a root component if we find that it's missing from the proxy on load

diff --git a/Engine/Source/Runtime/Landscape/Private/Landscape.cpp b/Engine/Source/Runtime/Landscape/Private/Landscape.cpp
--- a/Engine/Source/Runtime/Landscape/Private/Landscape.cpp
+++ b/Engine/Source/Runtime/Landscape/Private/Landscape.cpp
@@ -4186,7 +4186,25 @@ void ALandscapeProxy::PostLoad()
 {
 	TRACE_CPUPROFILER_EVENT_SCOPE(ALandscapeProxy::PostLoad);
 
+
 #if WITH_EDITOR
+	// Not sure that this can ever happen without someone deliberately changing the root component but a landscape without a root component is 
+	//  worthless and will lead to pain and crash, so attempt to fix it up on load here : 
+	if (GetRootComponent() == nullptr)
+	{
+		TInlineComponentArray<USceneComponent*> SceneComponents;
+		GetComponents<USceneComponent>(SceneComponents, /*bIncludeFromChildActors = */false);
+		UClass* SceneComponentClass = USceneComponent::StaticClass();
+		if (USceneComponent** RootComponentCandidate = Algo::FindByPredicate(SceneComponents, [SceneComponentClass](USceneComponent* InComponent) { return (InComponent->GetClass() == SceneComponentClass); }))
+		{
+			SetRootComponent(*RootComponentCandidate);
+		}
+		else
+		{
+			UE_LOG(LogLandscape, Error, TEXT("Unable to retrieve a root component for landscape proxy %s. The landscape will not work properly."), *GetFullName());
+		}
+	}
+
 	// Fix up bHasLayersContent if needed : there was a point where there was a missing call (when adding new components) that was leading to  it not being 
 	//  properly updated even though there was some layer content. We need to fix it now, so that we don't attempt to migrate data from non-edit layers to 
 	//  edit layers later on, since that would stomp the edit layer data : 
diff --git a/Engine/Source/Runtime/Landscape/Private/LandscapeCollision.cpp b/Engine/Source/Runtime/Landscape/Private/LandscapeCollision.cpp
--- a/Engine/Source/Runtime/Landscape/Private/LandscapeCollision.cpp
+++ b/Engine/Source/Runtime/Landscape/Private/LandscapeCollision.cpp
@@ -935,6 +935,12 @@ void ULandscapeHeightfieldCollisionComponent::CreateCollisionObject()
 				{
 					RegisterMaterials(HeightfieldRef->UsedChaosMaterials, CookedPhysicalMaterials);
 				}
+				else
+				{
+					// A heightfield with invalid content is as good as an invalid heightfield and ULandscapeHeightfieldCollisionComponent::CreateCollisionObject won't survive it anyway 
+					//  so we're better off invalidating it here entirely : 
+					HeightfieldRef = nullptr;
+				}
 
 				// Return unconditionally even if GenerateCollisionObjects failed.  CookCollisionData succeeds or fails in the same conditions, so trying that
 				// as a fallback doesn't gain anything.
