[Landscape] Parallelize retrieval of heights from landscape collision component.

#rnx

diff --git a/Engine/Source/Runtime/Landscape/Private/LandscapeCollision.cpp b/Engine/Source/Runtime/Landscape/Private/LandscapeCollision.cpp
--- a/Engine/Source/Runtime/Landscape/Private/LandscapeCollision.cpp
+++ b/Engine/Source/Runtime/Landscape/Private/LandscapeCollision.cpp
@@ -3301,7 +3301,7 @@ bool ULandscapeHeightfieldCollisionComponent::FillHeightTile(TArrayView<float> H
 	const float ZScale = static_cast<float>(WorldTransform.GetScale3D().Z * LANDSCAPE_ZSCALE);
 
 	// Write all values to output array
-	for (int32 y = 0; y < Accessor.NumY; ++y)
+	ParallelFor(Accessor.NumY, [&Heights, &Accessor, &WorldTransform, Offset, Stride, ZScale](int32 y)
 	{
 		for (int32 x = 0; x < Accessor.NumX; ++x)
 		{
@@ -3312,7 +3312,7 @@ bool ULandscapeHeightfieldCollisionComponent::FillHeightTile(TArrayView<float> H
 			const int32 WriteIndex = Offset + y * Stride + x;
 			Heights[WriteIndex] = WorldHeight;
 		}
-	}
+	});
 
 	return true;
 }
