Prevent landscape components that render in an additional tool-specific pass from being occluded

diff --git a/Engine/Source/Editor/LandscapeEditor/Private/LandscapeEdModeBrushes.cpp b/Engine/Source/Editor/LandscapeEditor/Private/LandscapeEdModeBrushes.cpp
--- a/Engine/Source/Editor/LandscapeEditor/Private/LandscapeEdModeBrushes.cpp
+++ b/Engine/Source/Editor/LandscapeEditor/Private/LandscapeEdModeBrushes.cpp
@@ -48,7 +48,7 @@ class FLandscapeBrushCircle : public FLandscapeBrush
 {
 	using Super = FLandscapeBrush;
 
-	// Components which previously we're under the area of the brush.
+	// Components which previously were under the area of the brush.
 	TSet<TWeakObjectPtr<ULandscapeComponent>, TWeakObjectPtrSetKeyFuncs<TWeakObjectPtr<ULandscapeComponent>>> BrushMaterialComponents;
 	
 	// A Cache of previously created and now unused MIDs we can reuse if required.
diff --git a/Engine/Source/Runtime/Landscape/Private/LandscapeComponent.cpp b/Engine/Source/Runtime/Landscape/Private/LandscapeComponent.cpp
--- a/Engine/Source/Runtime/Landscape/Private/LandscapeComponent.cpp
+++ b/Engine/Source/Runtime/Landscape/Private/LandscapeComponent.cpp
@@ -73,6 +73,15 @@ void ULandscapeComponent::UpdateEditToolRenderData()
 		const bool bGetDebugMaterials = true;
 		GetUsedMaterials(UsedMaterialsForVerification, bGetDebugMaterials);
 
+		// The selection/tool/gizmo materials are displayed in a translucent pass with Disable Depth Test == true so that they always show despite being underneath other objects.
+		//  However, for this to work under all circumstances, we must make sure they don't get CPU-culled by primitives whose bounds fully occlude it, so we artificially inflate
+		//  their bounds dynamically (without having to invalidate the render state) : 
+		const bool bDisableCulling = (EditToolRenderData.SelectedType != FLandscapeEditToolRenderData::ST_NONE)
+			|| (EditToolRenderData.ToolMaterial != nullptr)
+			|| (EditToolRenderData.GizmoMaterial != nullptr);
+
+		UpdateOcclusionBoundsSlack(bDisableCulling ? TNumericLimits<float>::Max() : 0.0f);
+
 		FLandscapeEditToolRenderData LandscapeEditToolRenderData = EditToolRenderData;
 		ENQUEUE_RENDER_COMMAND(UpdateEditToolRenderData)(
 			[LandscapeEditToolRenderData, LandscapeSceneProxy, UsedMaterialsForVerification](FRHICommandListImmediate& RHICmdList)
