Fixed landscape thumbnail material warning

diff --git a/Engine/Source/Runtime/Landscape/Private/LandscapeRender.cpp b/Engine/Source/Runtime/Landscape/Private/LandscapeRender.cpp
--- a/Engine/Source/Runtime/Landscape/Private/LandscapeRender.cpp
+++ b/Engine/Source/Runtime/Landscape/Private/LandscapeRender.cpp
@@ -4146,6 +4146,11 @@ public:
 			FName(TEXT("FLumenCardCS")),
 			FName(TEXT("FLumenCardVS")),
 			FName(TEXT("FLumenCardPS")),
+
+			// No material cache on thumbnails
+			FName(TEXT("FMaterialCacheUnwrapVS")),
+			FName(TEXT("FMaterialCacheUnwrapPS")),
+			FName(TEXT("FMaterialCacheNaniteShadeCS")),
 		};
 		return ExcludedShaderTypes;
 	}
