Re-add the weightmap-related assert now that we've tracked the root code of those being null ()

diff --git a/Engine/Source/Runtime/Landscape/Private/LandscapeEditLayers.cpp b/Engine/Source/Runtime/Landscape/Private/LandscapeEditLayers.cpp
--- a/Engine/Source/Runtime/Landscape/Private/LandscapeEditLayers.cpp
+++ b/Engine/Source/Runtime/Landscape/Private/LandscapeEditLayers.cpp
@@ -4103,7 +4103,7 @@ bool ALandscape::PrepareTextureResources(bool bInWaitForStreaming)
 		for (ULandscapeComponent* Component : Proxy->LandscapeComponents)
 		{
 			UTexture2D* ComponentHeightmap = Component->GetHeightmap();
-			check(ComponentHeightmap);
+			check(ComponentHeightmap != nullptr);
 			{
 				bool bIsTextureReady = TextureStreamingManager->RequestTextureFullyStreamedInForever(ComponentHeightmap, bInWaitForStreaming);
 				if (!bIsTextureReady)
@@ -4124,10 +4124,7 @@ bool ALandscape::PrepareTextureResources(bool bInWaitForStreaming)
 
 			for (UTexture2D* ComponentWeightmap : Component->GetWeightmapTextures())
 			{
-				if (!ComponentWeightmap)
-				{
-					continue; // ignore missing weightmaps
-				}
+				check(ComponentWeightmap != nullptr);
 
 				bool bIsTextureReady = TextureStreamingManager->RequestTextureFullyStreamedInForever(ComponentWeightmap, bInWaitForStreaming);
 				// If the texture is not ready, start tracking its state changes to be notified when it's fully streamed in : 
@@ -4282,15 +4279,13 @@ bool ALandscape::PrepareLayersTextureResources(const TArray<FLandscapeLayer>& In
 			{
 				if (FLandscapeLayerComponentData* ComponentLayerData = Component->GetLayerData(Layer.Guid))
 				{
-					check(ComponentLayerData->HeightmapData.Texture);
+					check(ComponentLayerData->HeightmapData.Texture != nullptr);
 					bIsReady &= TextureStreamingManager->RequestTextureFullyStreamedInForever(ComponentLayerData->HeightmapData.Texture, bInWaitForStreaming);
 
 					for (UTexture2D* LayerWeightmap : ComponentLayerData->WeightmapData.Textures)
 					{
-						if (LayerWeightmap) // ignore missing weightmaps
-						{
-							bIsReady &= TextureStreamingManager->RequestTextureFullyStreamedInForever(LayerWeightmap, bInWaitForStreaming);
-						}
+						check(LayerWeightmap != nullptr);
+						bIsReady &= TextureStreamingManager->RequestTextureFullyStreamedInForever(LayerWeightmap, bInWaitForStreaming);
 					}
 				}
 			}
