Remove codepaths related to HasNormalCaptureBPBrushLayer.  No longer required with new landscape border normal handling and it conflicts with BatchMerge code.

diff --git a/Engine/Plugins/Experimental/Landmass/Source/Editor/Private/LandmassManagerBase.cpp b/Engine/Plugins/Experimental/Landmass/Source/Editor/Private/LandmassManagerBase.cpp
--- a/Engine/Plugins/Experimental/Landmass/Source/Editor/Private/LandmassManagerBase.cpp
+++ b/Engine/Plugins/Experimental/Landmass/Source/Editor/Private/LandmassManagerBase.cpp
@@ -492,9 +492,8 @@ void ALandmassManagerBase::TogglePreviewMode_Implementation(bool bEnablePreviewM
 
 }
 
-void ALandmassManagerBase::SetCaptureBoundaryNormals(bool bInCaptureBoundaryNormals)
+void ALandmassManagerBase::SetCaptureBoundaryNormals(bool)
 {
-	bCaptureBoundaryNormals = bInCaptureBoundaryNormals;
 }
 
 void ALandmassManagerBase::MoveBrushUp(ALandmassActor* BrushToMove)
diff --git a/Engine/Plugins/Experimental/Landmass/Source/Editor/Public/LandmassManagerBase.h b/Engine/Plugins/Experimental/Landmass/Source/Editor/Public/LandmassManagerBase.h
--- a/Engine/Plugins/Experimental/Landmass/Source/Editor/Public/LandmassManagerBase.h
+++ b/Engine/Plugins/Experimental/Landmass/Source/Editor/Public/LandmassManagerBase.h
@@ -124,6 +124,7 @@ public:
 	UFUNCTION(BlueprintNativeEvent, BlueprintCallable, meta = (ForceAsFunction), Category = Default)
 	void TogglePreviewMode(bool bEnablePreviewMode);
 
+	UE_DEPRECATED(5.6, "This no longer has an effect")
 	UFUNCTION(BlueprintCallable, Category = Default)
 	void SetCaptureBoundaryNormals(bool bInCaptureBoundaryNormals);
 	
diff --git a/Engine/Shaders/Private/Landscape/LandscapeEditLayersHeightmaps.usf b/Engine/Shaders/Private/Landscape/LandscapeEditLayersHeightmaps.usf
--- a/Engine/Shaders/Private/Landscape/LandscapeEditLayersHeightmaps.usf
+++ b/Engine/Shaders/Private/Landscape/LandscapeEditLayersHeightmaps.usf
@@ -194,28 +194,6 @@ void GenerateNormalsPS(in float4 SVPos : SV_POSITION, out float4 OutColor : SV_T
 	//BR.w *= MaskMaxBorderX * MaskMaxBorderY;
 	//BB.w *= MaskMaxBorderY;
 
-	// TODO [jonathan.bard] This is not supported and we probably don't want to keep this : 
-	//// HACK [chris.tchou] remove once we have a better boundary normal solution
-	//bool bApplyNormalBoundaryOverride = (OutputConfig.w >= 1.5);
-	//bool bIsOnBoundary = (TL.w * BR.w == 0.0);
-	//if (bApplyNormalBoundaryOverride && bIsOnBoundary)
-	//{
-	//	// Resample heights from the Captured Boundary texture, which is assumed to have height data everywhere
-	//	Red0 = ReadTexture2.GatherRed(ReadTexture2Sampler, GatherLocation - TexelSize);
-	//	Green0 = ReadTexture2.GatherGreen(ReadTexture2Sampler, GatherLocation - TexelSize);
-	//	Red1 = ReadTexture2.GatherRed(ReadTexture2Sampler, GatherLocation);
-	//	Green1 = ReadTexture2.GatherGreen(ReadTexture2Sampler, GatherLocation);
-
-	//	// recalculate heights and replace masks with 1.0
-	//	TL.zw = float2(ExtractAndScaleHeight(float2(Red0.w, Green0.w)), 1.0f);
-	//	TT.zw = float2(ExtractAndScaleHeight(float2(Red0.z, Green0.z)), 1.0f);
-	//	CC.zw = float2(ExtractAndScaleHeight(float2(Red0.y, Green0.y)), 1.0f);
-	//	LL.zw = float2(ExtractAndScaleHeight(float2(Red0.x, Green0.x)), 1.0f);
-	//	RR.zw = float2(ExtractAndScaleHeight(float2(Red1.z, Green1.z)), 1.0f);
-	//	BR.zw = float2(ExtractAndScaleHeight(float2(Red1.y, Green1.y)), 1.0f);
-	//	BB.zw = float2(ExtractAndScaleHeight(float2(Red1.x, Green1.x)), 1.0f);
-	//}
-
 	float3 N0 = ComputeNullableTriangleNormal(CC, LL, TL);
 	float3 N1 = ComputeNullableTriangleNormal(TL, TT, CC);
 	float3 N2 = ComputeNullableTriangleNormal(LL, CC, BB);
diff --git a/Engine/Shaders/Private/LandscapeLayersPS.usf b/Engine/Shaders/Private/LandscapeLayersPS.usf
--- a/Engine/Shaders/Private/LandscapeLayersPS.usf
+++ b/Engine/Shaders/Private/LandscapeLayersPS.usf
@@ -86,27 +86,6 @@ float4 CalculateNormalsFromHeights(float2 InTextureCoordinates, float4 OutColor)
 	BR.w *= MaskMaxBorderX * MaskMaxBorderY;
 	BB.w *= MaskMaxBorderY;
 
-	// HACK [chris.tchou] remove once we have a better boundary normal solution
-	bool bApplyNormalBoundaryOverride = (OutputConfig.w >= 1.5);
-	bool bIsOnBoundary = (TL.w * BR.w == 0.0);
-	if (bApplyNormalBoundaryOverride && bIsOnBoundary)
-	{
-		// Resample heights from the Captured Boundary texture, which is assumed to have height data everywhere
-		Red0 = ReadTexture2.GatherRed(ReadTexture2Sampler, GatherLocation - TexelSize);
-		Green0 = ReadTexture2.GatherGreen(ReadTexture2Sampler, GatherLocation - TexelSize);
-		Red1 = ReadTexture2.GatherRed(ReadTexture2Sampler, GatherLocation);
-		Green1 = ReadTexture2.GatherGreen(ReadTexture2Sampler, GatherLocation);
-
-		// recalculate heights and replace masks with 1.0
-		TL.zw = float2(ExtractAndScaleHeight(float2(Red0.w, Green0.w)), 1.0f);
-		TT.zw = float2(ExtractAndScaleHeight(float2(Red0.z, Green0.z)), 1.0f);
-		CC.zw = float2(ExtractAndScaleHeight(float2(Red0.y, Green0.y)), 1.0f);
-		LL.zw = float2(ExtractAndScaleHeight(float2(Red0.x, Green0.x)), 1.0f);
-		RR.zw = float2(ExtractAndScaleHeight(float2(Red1.z, Green1.z)), 1.0f);
-		BR.zw = float2(ExtractAndScaleHeight(float2(Red1.y, Green1.y)), 1.0f);
-		BB.zw = float2(ExtractAndScaleHeight(float2(Red1.x, Green1.x)), 1.0f);
-	}
-
 	float3 N0 = ComputeNullableTriangleNormal(CC, LL, TL);
 	float3 N1 = ComputeNullableTriangleNormal(TL, TT, CC);
 	float3 N2 = ComputeNullableTriangleNormal(LL, CC, BB);
diff --git a/Engine/Source/Runtime/Landscape/Classes/Landscape.h b/Engine/Source/Runtime/Landscape/Classes/Landscape.h
--- a/Engine/Source/Runtime/Landscape/Classes/Landscape.h
+++ b/Engine/Source/Runtime/Landscape/Classes/Landscape.h
@@ -64,7 +64,6 @@ enum class EHeightmapRTType : uint8
 	HeightmapRT_Scratch1,
 	HeightmapRT_Scratch2,
 	HeightmapRT_Scratch3,
-	HeightmapRT_BoundaryNormal, // HACK [chris.tchou] remove once we have a better boundary normal solution
 	// Mips RT
 	HeightmapRT_Mip1,
 	HeightmapRT_Mip2,
@@ -517,10 +516,9 @@ private:
 	void OnLayerCreatedInternal(FLandscapeLayer& Layer);
 
 	bool SupportsEditLayersLocalMerge();
-	bool HasNormalCaptureBPBrushLayer();
 
 	ELandscapeEditLayersMergeMode GetEditLayersMergeMode();
-	bool CreateLayersRenderingResource(bool bUseNormalCapture);
+	bool CreateLayersRenderingResource();
 	void PrepareEditLayersLocalMergeResources();
 	void UpdateLayersContent(bool bInWaitForStreaming = false, bool bInSkipMonitorLandscapeEdModeChanges = false, bool bIntermediateRender = false, bool bFlushRender = false);
 	void MonitorShaderCompilation();
@@ -595,7 +593,7 @@ private:
 
 	void CopyTexturePS(const FString& InSourceDebugName, FTextureResource* InSourceResource, const FString& InDestDebugName, FTextureResource* InDestResource) const;
 
-	void InitializeLayers(bool bUseNormalCapture);
+	void InitializeLayers();
 	
 	void PrintLayersDebugRT(const FString& InContext, UTextureRenderTarget2D* InDebugRT, uint8 InMipRender = 0, bool InOutputHeight = true, bool InOutputNormals = false) const;
 	void PrintLayersDebugTextureResource(const FString& InContext, FTextureResource* InTextureResource, uint8 InMipRender = 0, bool InOutputHeight = true, bool InOutputNormals = false) const;
@@ -702,9 +700,6 @@ private:
 
 	UPROPERTY(Transient)
 	bool bLandscapeLayersAreInitialized;
-	
-	UPROPERTY(Transient)
-	bool bLandscapeLayersAreInitializedForNormalCapture;
 
 	UPROPERTY(Transient)
 	bool bLandscapeLayersAreUsingLocalMerge;
diff --git a/Engine/Source/Runtime/Landscape/Private/LandscapeEditLayers.cpp b/Engine/Source/Runtime/Landscape/Private/LandscapeEditLayers.cpp
--- a/Engine/Source/Runtime/Landscape/Private/LandscapeEditLayers.cpp
+++ b/Engine/Source/Runtime/Landscape/Private/LandscapeEditLayers.cpp
@@ -668,7 +668,6 @@ struct FLandscapeLayersHeightmapShaderParameters
 		, LayerVisible(true)
 		, LayerBlendMode(LSBM_AdditiveBlend)
 		, GenerateNormals(false)
-		, OverrideBoundaryNormals(false) // HACK [chris.tchou] remove once we have a better boundary normal solution
 		, GridSize(0.0f, 0.0f, 0.0f)
 		, CurrentMipSize(0, 0)
 		, ParentMipSize(0, 0)
@@ -684,7 +683,6 @@ struct FLandscapeLayersHeightmapShaderParameters
 	bool LayerVisible;
 	ELandscapeBlendMode LayerBlendMode;
 	bool GenerateNormals;
-	bool OverrideBoundaryNormals; // HACK [chris.tchou] remove once we have a better boundary normal solution
 	FVector GridSize;
 	FIntPoint CurrentMipSize;
 	FIntPoint ParentMipSize;
@@ -726,16 +724,13 @@ public:
 	void SetParameters(FRHIBatchedShaderParameters& BatchedParameters, const FLandscapeLayersHeightmapShaderParameters& InParams)
 	{
 		SetTextureParameter(BatchedParameters, ReadTexture1Param, ReadTexture1SamplerParam, TStaticSamplerState<SF_Point, AM_Clamp, AM_Clamp, AM_Clamp>::GetRHI(), InParams.ReadHeightmap1->GetResource()->TextureRHI);
-		SetTextureParameter(BatchedParameters, ReadTexture2Param, ReadTexture2SamplerParam, 
-			(InParams.GenerateNormals && InParams.OverrideBoundaryNormals) 					// HACK [chris.tchou] remove once we have a better boundary normal solution
-				? TStaticSamplerState<SF_Point, AM_Wrap, AM_Wrap, AM_Wrap>::GetRHI()		// override boundary normals requires wrap mode
-				: TStaticSamplerState<SF_Point, AM_Clamp, AM_Clamp, AM_Clamp>::GetRHI(),
+		SetTextureParameter(BatchedParameters, ReadTexture2Param, ReadTexture2SamplerParam, TStaticSamplerState<SF_Point, AM_Clamp, AM_Clamp, AM_Clamp>::GetRHI(),
 			InParams.ReadHeightmap2 != nullptr ? InParams.ReadHeightmap2->GetResource()->TextureRHI : GWhiteTexture->TextureRHI);
 
 		FVector4f LayerInfo(InParams.LayerAlpha, InParams.LayerVisible ? 1.0f : 0.0f, InParams.LayerBlendMode == LSBM_AlphaBlend ? 1.0f : 0.f, 0.f);
 		FVector4f OutputConfig(InParams.ApplyLayerModifiers ? 1.0f : 0.0f, InParams.SetAlphaOne ? 1.0f : 0.0f,
-			(InParams.ReadHeightmap2 && !InParams.OverrideBoundaryNormals) ? 1.0f : 0.0f,			// HACK [chris.tchou] remove once we have a better boundary normal solution
-			InParams.GenerateNormals ? (InParams.OverrideBoundaryNormals ? 2.0f : 1.0f) : 0.0f);	// HACK [chris.tchou] remove once we have a better boundary normal solution
+			InParams.ReadHeightmap2 ? 1.0f : 0.0f,
+			InParams.GenerateNormals ? 1.0f : 0.0f);
 		FVector2f TextureSize(static_cast<float>(InParams.HeightmapSize.X), static_cast<float>(InParams.HeightmapSize.Y));
 
 		SetShaderValue(BatchedParameters, LayerInfoParam, LayerInfo);
@@ -2425,23 +2420,7 @@ ELandscapeEditLayersMergeMode ALandscape::GetEditLayersMergeMode()
 	return DesiredMode;
 }
 
-bool ALandscape::HasNormalCaptureBPBrushLayer()
-{
-	ALandscape* Landscape = GetLandscapeActor();
-	check(Landscape);
-
-	for (const FLandscapeLayer& Layer : Landscape->LandscapeEditLayers)
-	{
-		if (Layer.bVisible && Algo::AnyOf(Layer.Brushes, [](const FLandscapeLayerBrush& Brush) { return (Brush.GetBrush() != nullptr) && Brush.GetBrush()->IsVisible() && Brush.GetBrush()->GetCaptureBoundaryNormals(); }))
-		{
-			return true;
-		}
-	}
-
-	return false;
-}
-
-bool ALandscape::CreateLayersRenderingResource(bool bUseNormalCapture)
+bool ALandscape::CreateLayersRenderingResource()
 {
 	ULandscapeInfo* Info = GetLandscapeInfo();
 	if (Info == nullptr)
@@ -2474,23 +2453,6 @@ bool ALandscape::CreateLayersRenderingResource(bool bUseNormalCapture)
 		int32 CurrentMipSizeX = ComponentVerts * ComponentCounts.X;
 		int32 CurrentMipSizeY = ComponentVerts * ComponentCounts.Y;
 
-		// when using normal capture, we require an extra padding of 2 pixels --
-		//   but we round up to a full component of padding because the mip downsample shader code currently requires perfect 2:1 ratios in resolutions between mips
-		//   and keeping CurrentMipSizeX/Y a multiple of ComponentVerts ensures that (ComponentVerts is guaranteed to be a power of 2) 
-		if (bUseNormalCapture)
-		{
-			int32 PaddingX = FMath::RoundUpToPowerOfTwo(CurrentMipSizeX) - CurrentMipSizeX;
-			int32 PaddingY = FMath::RoundUpToPowerOfTwo(CurrentMipSizeY) - CurrentMipSizeY;
-			if (PaddingX < 2)
-			{
-				CurrentMipSizeX += ComponentVerts;
-			}
-			if (PaddingY < 2)
-			{
-				CurrentMipSizeY += ComponentVerts;
-			}
-		}
-
 		if (CurrentMipSizeX > GRHIGlobals.MaxTextureDimensions || CurrentMipSizeY > GRHIGlobals.MaxTextureDimensions)
 		{
 			if (!bWarnedGlobalMergeDimensionsExceeded)
@@ -7330,7 +7292,6 @@ int32 ALandscape::PerformLayersHeightmapsGlobalMerge(const FUpdateLayersContentC
 	FLandscapeLayersHeightmapShaderParameters ShaderParams;
 
 	bool FirstLayer = true;
-	bool bHasBoundaryNormalCopy = false;	// HACK [chris.tchou] remove once we have a better boundary normal solution
 	UTextureRenderTarget2D* CombinedHeightmapAtlasRT = HeightmapRTList[(int32)EHeightmapRTType::HeightmapRT_CombinedAtlas];
 	UTextureRenderTarget2D* CombinedHeightmapNonAtlasRT = HeightmapRTList[(int32)EHeightmapRTType::HeightmapRT_CombinedNonAtlas];
 	UTextureRenderTarget2D* LandscapeScratchRT1 = HeightmapRTList[(int32)EHeightmapRTType::HeightmapRT_Scratch1];
@@ -7412,20 +7373,6 @@ int32 ALandscape::PerformLayersHeightmapsGlobalMerge(const FUpdateLayersContentC
 
 				PrintLayersDebugRT(FString::Printf(TEXT("LS Height: %s %s -> BrushNonAtlas %s"), *Layer.Name.ToString(), *LandscapeBrush->GetName(), *BrushOutputNonAtlasRT->GetName()), BrushOutputNonAtlasRT);
 
-				// HACK [chris.tchou] remove once we have a better boundary normal solution
-				if (LandscapeBrush->GetCaptureBoundaryNormals())
-				{
-					if (bHasBoundaryNormalCopy)
-					{
-						UE_LOG(LogLandscape, Warning, TEXT("Multiple Landscape Blueprint Brush Layers are invoking Boundary Normal Capture!  Only the most recent one will be used."));
-					}
-
-					// GPU copy to another RT (we copy the entire thing... easier than picking out just the boundaries)
-					UTextureRenderTarget2D* BoundaryNormalRT = HeightmapRTList[(int32)EHeightmapRTType::HeightmapRT_BoundaryNormal];
-					ExecuteCopyLayersTexture({ FLandscapeLayersCopyTextureParams(BrushOutputNonAtlasRT, BoundaryNormalRT) });
-					bHasBoundaryNormalCopy = true;
-				}
-
 				// Resolve back to Combined heightmap (it's unlikely, but possible that the brush returns the same RT as input and output, if it did various operations on it, in which case the copy is useless) :
 				if (BrushOutputNonAtlasRT != CombinedHeightmapNonAtlasRT)
 				{
@@ -7478,16 +7425,13 @@ int32 ALandscape::PerformLayersHeightmapsGlobalMerge(const FUpdateLayersContentC
 
 	// compute Normals into LandscapeScratchRT1
 	ShaderParams.GenerateNormals = true;
-	ShaderParams.OverrideBoundaryNormals = bHasBoundaryNormalCopy;	// HACK [chris.tchou] remove once we have a better boundary normal solution
 	ShaderParams.GridSize = GetRootComponent()->GetRelativeScale3D();
-	UTextureRenderTarget2D* BoundaryNormalRT = HeightmapRTList[(int32)EHeightmapRTType::HeightmapRT_BoundaryNormal];	// HACK [chris.tchou] remove once we have a better boundary normal solution
 	DrawHeightmapComponentsToRenderTarget(FString::Printf(TEXT("LS Height: %s = -> CombinedNonAtlasNormals : %s"), *LandscapeScratchRT2->GetName(), *LandscapeScratchRT1->GetName()),
 		InUpdateLayersContentContext.LandscapeComponentsHeightmapsToRender, LandscapeExtent.Min, LandscapeScratchRT2,
-		bHasBoundaryNormalCopy ? BoundaryNormalRT : nullptr,	// HACK [chris.tchou] remove once we have a better boundary normal solution
+		nullptr,
 		LandscapeScratchRT1, ERTDrawingType::RTNonAtlas, true, ShaderParams);
 
 	ShaderParams.GenerateNormals = false;
-	ShaderParams.OverrideBoundaryNormals = false;	// HACK [chris.tchou] remove once we have a better boundary normal solution
 
 	// convert back to atlas (TODO: we could do this on the first mip downsample instead...)
 	DrawHeightmapComponentsToRenderTarget(FString::Printf(TEXT("LS Height: %s = -> CombinedAtlasFinal : %s"), *LandscapeScratchRT1->GetName(), *CombinedHeightmapAtlasRT->GetName()),
@@ -11509,19 +11453,15 @@ void ALandscape::UpdateLayersContent(bool bInWaitForStreaming, bool bInSkipMonit
 
 	// If we went from local merge to global merge or vice versa, we need to reinitialize layers : 
 	ELandscapeEditLayersMergeMode EditLayersMergeMode = GetEditLayersMergeMode();
-	// If we toggled normal capture, we need to reinitialize layers
-	bool bUseNormalCapture = HasNormalCaptureBPBrushLayer();
-	if ((CurrentEditLayersMergeMode != EditLayersMergeMode)
-		|| (bLandscapeLayersAreInitializedForNormalCapture != bUseNormalCapture))
+	if (CurrentEditLayersMergeMode != EditLayersMergeMode)
 	{
 		RequestLayersInitialization(/*bInRequestContentUpdate = */true);
 		CurrentEditLayersMergeMode = EditLayersMergeMode;
-		bLandscapeLayersAreInitializedForNormalCapture = bUseNormalCapture;
 	}
 
 	if (!bLandscapeLayersAreInitialized)
 	{
-		InitializeLayers(bUseNormalCapture);
+		InitializeLayers();
 	}
 
 	if (!bLandscapeLayersAreInitialized)
@@ -11874,11 +11814,11 @@ uint32 ALandscape::UpdateAfterReadbackResolves(const TArrayView<FLandscapeEditLa
 	return NewUpdateFlags;
 }
 
-void ALandscape::InitializeLayers(bool bUseNormalCapture)
+void ALandscape::InitializeLayers()
 {
 	check(HasLayersContent());
 
-	if (CreateLayersRenderingResource(bUseNormalCapture))
+	if (CreateLayersRenderingResource())
 	{
 		InitializeLandscapeLayersWeightmapUsage();
 		bLandscapeLayersAreInitialized = true;
diff --git a/Engine/Source/Runtime/Landscape/Public/LandscapeBlueprintBrushBase.h b/Engine/Source/Runtime/Landscape/Public/LandscapeBlueprintBrushBase.h
--- a/Engine/Source/Runtime/Landscape/Public/LandscapeBlueprintBrushBase.h
+++ b/Engine/Source/Runtime/Landscape/Public/LandscapeBlueprintBrushBase.h
@@ -74,8 +74,6 @@ protected:
 	FTransform CurrentRenderAreaWorldTransform;
 	FIntPoint CurrentRenderAreaSize = FIntPoint(ForceInitToZero);
 	FIntPoint CurrentRenderTargetSize = FIntPoint(ForceInitToZero);
-
-	bool bCaptureBoundaryNormals = false;	// HACK [chris.tchou] remove once we have a better boundary normal solution
 #endif // WITH_EDITORONLY_DATA
 
 public:
@@ -154,7 +152,5 @@ public:
 	LANDSCAPE_API virtual void PushDeferredLayersContentUpdate();
 
 	virtual bool CanChangeIsSpatiallyLoadedFlag() const override { return false; }
-
-	bool GetCaptureBoundaryNormals() { return bCaptureBoundaryNormals; }	// HACK [chris.tchou] remove once we have a better boundary normal solution
 #endif // WITH_EDITOR
 };
