Added CPU perf tag

diff --git a/Engine/Source/Runtime/Landscape/Private/LandscapeEditLayerRenderer.cpp b/Engine/Source/Runtime/Landscape/Private/LandscapeEditLayerRenderer.cpp
--- a/Engine/Source/Runtime/Landscape/Private/LandscapeEditLayerRenderer.cpp
+++ b/Engine/Source/Runtime/Landscape/Private/LandscapeEditLayerRenderer.cpp
@@ -1170,6 +1170,7 @@ void FMergeRenderContext::Render(TFunction<void(const FOnRenderBatchTargetGroupD
 		const FMergeRenderBatch& RenderBatch = RenderBatches[CurrentRenderBatchIndex];
 		FString RenderBatchDebugName = FString::Format(TEXT("Render Batch [{0}] : ({1},{2})->({3},{4})"), { CurrentRenderBatchIndex, RenderBatch.SectionRect.Min.X, RenderBatch.SectionRect.Min.Y, RenderBatch.SectionRect.Max.X, RenderBatch.SectionRect.Max.Y });
 		RHI_BREADCRUMB_EVENT_GAMETHREAD("%s", RenderBatchDebugName);
+		TRACE_CPUPROFILER_EVENT_SCOPE(RenderBatch);
 
 		checkf((RenderBatch.RenderSteps.Num() >= 1) && (RenderBatch.RenderSteps.Last().Type == FMergeRenderStep::EType::SignalBatchMergeGroupDone),
 			TEXT("Any batch should end with a SignalBatchMergeGroupDone step and there should be at least another step prior to that, otherwise, the batch is just useless."));
