Fix runaway memory usage on landscape painting.  Remove duplicated flags from ULandscapeScratchRenderTarget which were being tested but never set correctly, preventing scratch render target re-use.

diff --git a/Engine/Source/Runtime/Landscape/Private/LandscapeEditResourcesSubsystem.cpp b/Engine/Source/Runtime/Landscape/Private/LandscapeEditResourcesSubsystem.cpp
--- a/Engine/Source/Runtime/Landscape/Private/LandscapeEditResourcesSubsystem.cpp
+++ b/Engine/Source/Runtime/Landscape/Private/LandscapeEditResourcesSubsystem.cpp
@@ -301,7 +301,7 @@ bool ULandscapeScratchRenderTarget::IsCompatibleWith(const UE::Landscape::FScrat
 	}
 
 	// If texture flags are different, we cannot be compatible
-	if ((bUseUAV != InParams.bUseUAV) || (bTargetArraySlicesIndependently != InParams.bTargetArraySlicesIndependently))
+	if (RenderTarget->bCanCreateUAV != InParams.bUseUAV || RenderTarget->bTargetArraySlicesIndependently != InParams.bTargetArraySlicesIndependently)
 	{
 		return false;
 	}
diff --git a/Engine/Source/Runtime/Landscape/Public/LandscapeEditResourcesSubsystem.h b/Engine/Source/Runtime/Landscape/Public/LandscapeEditResourcesSubsystem.h
--- a/Engine/Source/Runtime/Landscape/Public/LandscapeEditResourcesSubsystem.h
+++ b/Engine/Source/Runtime/Landscape/Public/LandscapeEditResourcesSubsystem.h
@@ -247,12 +247,6 @@ private:
 	// Format of the render target. Technically, we could infer it from RenderTarget->Format but since it's stored as a EPixelFormat, 
 	//  end there's no easy conversion from EPixelFormat to ETextureRenderTargetFormat, we store a copy here instead: 
 	ETextureRenderTargetFormat RenderTargetFormat = ETextureRenderTargetFormat::RTF_RGBA8;
-
-	// Whether the render target has the ETextureCreateFlags::TargetArraySlicesIndependently flag
-	bool bTargetArraySlicesIndependently = false;
-
-	// Whether the render target has the ETextureCreateFlags::UAV flag
-	bool bUseUAV = false;
 	// END Un-mutable section
 
 	// Current state of the scratch render target :
