Prevent performance warning related to calling FindConsoleObject too often in landscape

diff --git a/Engine/Source/Runtime/Landscape/Private/LandscapeRender.cpp b/Engine/Source/Runtime/Landscape/Private/LandscapeRender.cpp
--- a/Engine/Source/Runtime/Landscape/Private/LandscapeRender.cpp
+++ b/Engine/Source/Runtime/Landscape/Private/LandscapeRender.cpp
@@ -312,7 +312,7 @@ namespace UE::Landscape
 
 		// We need the fixed grid vertex factory for rendering the water info texture : 
 		// This cvar is defined in the water plugin and searching for it should return nullptr if the plugin is not loaded
-		const bool bWaterPluginLoaded = IConsoleManager::Get().FindConsoleVariable(TEXT("r.Water.WaterInfo.RenderMethod")) != nullptr;
+		static const bool bWaterPluginLoaded = IConsoleManager::Get().FindConsoleVariable(TEXT("r.Water.WaterInfo.RenderMethod")) != nullptr;
 		bNeedsFixedGridVertexFactory |= bWaterPluginLoaded;
 
 		return bNeedsFixedGridVertexFactory;
