Improvements to landscape thumbnail renders :
* Don't bother generating landscape thumbnail MICs if the option is disabled and display a generic layer icon instead
* Use a single material context when populating thumbnails for all layers. This avoids repeatedly iterating through all primitive components to update those that might reference this material (O(N^2) -> O(N))

diff --git a/Engine/Source/Editor/LandscapeEditor/Private/LandscapeEdMode.cpp b/Engine/Source/Editor/LandscapeEditor/Private/LandscapeEdMode.cpp
--- a/Engine/Source/Editor/LandscapeEditor/Private/LandscapeEdMode.cpp
+++ b/Engine/Source/Editor/LandscapeEditor/Private/LandscapeEdMode.cpp
@@ -2244,6 +2244,14 @@ void FEdModeLandscape::RefreshDetailPanel()
 	}
 }
 
+void FEdModeLandscape::RegenerateLayerThumbnails()
+{
+	// Regenerate the lists that might hold on to a ThumbnailMIC : 
+	UpdateTargetList(/*bRegenerateThumbnails = */true);
+	UISettings->RefreshImportLayersList(/*bRefreshFromTarget = */false, /*bRegenerateThumbnails = */true); // Once for NewLandscape_Material
+	UISettings->RefreshImportLayersList(/*bRefreshFromTarget = */true, /*bRegenerateThumbnails = */true); // Once for the target list
+}
+
 void FEdModeLandscape::SetCurrentBrushSet(FName BrushSetName)
 {
 	for (int32 BrushIndex = 0; BrushIndex < LandscapeBrushSets.Num(); BrushIndex++)
@@ -2553,7 +2561,7 @@ bool FEdModeLandscape::CanEditCurrentTarget(FText* Reason) const
 	return true;
 }
 
-void FEdModeLandscape::UpdateTargetList()
+void FEdModeLandscape::UpdateTargetList(bool bRegenerateThumbnails)
 {
 	LandscapeTargetList.Empty();
 	TargetLayerStartingIndex = 0;
@@ -2586,11 +2594,10 @@ void FEdModeLandscape::UpdateTargetList()
 			}
 
 			// Add layers
-			UTexture2D* ThumbnailWeightmap = nullptr;
-			UTexture2D* ThumbnailHeightmap = nullptr;
-
 			TargetLayerStartingIndex = LandscapeTargetList.Num();
 
+			// Don't recreate the render state of everything, only update the materials context
+			FMaterialUpdateContext MaterialUpdateContext(FMaterialUpdateContext::EOptions::Default & ~FMaterialUpdateContext::EOptions::RecreateRenderStates);
 			for (auto It = CurrentToolTarget.LandscapeInfo->Layers.CreateIterator(); It; It++)
 			{
 				FLandscapeInfoLayerSettings& LayerSettings = *It;
@@ -2611,21 +2618,10 @@ void FEdModeLandscape::UpdateTargetList()
 					bFoundSelected = true;
 				}
 
-				// Ensure thumbnails are up valid
-				if (LayerSettings.ThumbnailMIC == nullptr)
+				// Ensure thumbnails are built 
+				if ((LayerSettings.ThumbnailMIC == nullptr) || bRegenerateThumbnails)
 				{
-					if (ThumbnailWeightmap == nullptr)
-					{
-						ThumbnailWeightmap = LoadObject<UTexture2D>(nullptr, TEXT("/Engine/EditorLandscapeResources/LandscapeThumbnailWeightmap.LandscapeThumbnailWeightmap"), nullptr, LOAD_None, nullptr);
-					}
-					if (ThumbnailHeightmap == nullptr)
-					{
-						ThumbnailHeightmap = LoadObject<UTexture2D>(nullptr, TEXT("/Engine/EditorLandscapeResources/LandscapeThumbnailHeightmap.LandscapeThumbnailHeightmap"), nullptr, LOAD_None, nullptr);
-					}
-
-					// Construct Thumbnail MIC
-					UMaterialInterface* LandscapeMaterial = LayerSettings.Owner.IsValid() ? LayerSettings.Owner.Get()->GetLandscapeMaterial() : UMaterial::GetDefaultMaterial(MD_Surface);
-					LayerSettings.ThumbnailMIC = ALandscapeProxy::GetLayerThumbnailMIC(LandscapeMaterial, LayerName, ThumbnailWeightmap, ThumbnailHeightmap, LayerSettings.Owner.Get());
+					LayerSettings.ThumbnailMIC = UE::Landscape::CreateLandscapeLayerThumbnailMIC(MaterialUpdateContext, LayerSettings.Owner.IsValid() ? LayerSettings.Owner.Get()->GetLandscapeMaterial() : nullptr, LayerName);
 				}
 
 				// Add the layer
diff --git a/Engine/Source/Editor/LandscapeEditor/Private/LandscapeEdMode.h b/Engine/Source/Editor/LandscapeEditor/Private/LandscapeEdMode.h
--- a/Engine/Source/Editor/LandscapeEditor/Private/LandscapeEdMode.h
+++ b/Engine/Source/Editor/LandscapeEditor/Private/LandscapeEdMode.h
@@ -561,7 +561,7 @@ public:
 	const FString GetTargetLayerAssetPackagePath() const;
 
 	int32 UpdateLandscapeList();
-	void UpdateTargetList();
+	void UpdateTargetList(bool bRegenerateThumbnails = false);
 	void SetTargetLandscape(const TWeakObjectPtr<ULandscapeInfo>& InLandscapeInfo);
 	bool CanEditCurrentTarget(FText* Reason = nullptr) const;
 
@@ -574,6 +574,7 @@ public:
 	void UpdateLayerUsageInformation(TWeakObjectPtr<ULandscapeLayerInfoObject>* LayerInfoObjectThatChanged = nullptr);
 	void OnLandscapeMaterialChangedDelegate();
 	void RefreshDetailPanel();
+	void RegenerateLayerThumbnails();
 
 	bool IsGridBased() const;
 
diff --git a/Engine/Source/Editor/LandscapeEditor/Private/LandscapeEditorModule.cpp b/Engine/Source/Editor/LandscapeEditor/Private/LandscapeEditorModule.cpp
--- a/Engine/Source/Editor/LandscapeEditor/Private/LandscapeEditorModule.cpp
+++ b/Engine/Source/Editor/LandscapeEditor/Private/LandscapeEditorModule.cpp
@@ -303,6 +303,7 @@ public:
 	*/
 	virtual int32 GetOrCreateEditLayer(FName InEditLayerName, ALandscape* InTargetLandscape, const TSubclassOf<ULandscapeEditLayerBase>& InEditLayerClass) override;
 	virtual void RefreshDetailPanel() override;
+	virtual void RegenerateLayerThumbnails() override;
 
 private:
 	TSharedRef<FPropertySection> RegisterPropertySection(FPropertyEditorModule& PropertyModule, FName ClassName, FName SectionName, FText DisplayName);
@@ -590,4 +591,12 @@ void FLandscapeEditorModule::RefreshDetailPanel()
 	}
 }
 
+void FLandscapeEditorModule::RegenerateLayerThumbnails()
+{
+	if (FEdModeLandscape* LandscapeMode = (FEdModeLandscape*)GLevelEditorModeTools().GetActiveMode(FBuiltinEditorModes::EM_Landscape))
+	{
+		LandscapeMode->RegenerateLayerThumbnails();
+	}
+}
+
 #undef LOCTEXT_NAMESPACE
diff --git a/Engine/Source/Editor/LandscapeEditor/Private/LandscapeEditorObject.cpp b/Engine/Source/Editor/LandscapeEditor/Private/LandscapeEditorObject.cpp
--- a/Engine/Source/Editor/LandscapeEditor/Private/LandscapeEditorObject.cpp
+++ b/Engine/Source/Editor/LandscapeEditor/Private/LandscapeEditorObject.cpp
@@ -13,6 +13,7 @@
 #include "LandscapeImportHelper.h"
 #include "LandscapeTiledImage.h"
 #include "LandscapeMaterialInstanceConstant.h"
+#include "LandscapeUtils.h"
 #include "Misc/ConfigCacheIni.h"
 #include "EngineUtils.h"
 #include "RenderUtils.h"
@@ -880,11 +881,8 @@ void ULandscapeEditorObject::ExpandImportData(TArray<uint16>& OutHeightData, TAr
 	}
 }
 
-void ULandscapeEditorObject::RefreshImportLayersList(bool bRefreshFromTarget)
+void ULandscapeEditorObject::RefreshImportLayersList(bool bRefreshFromTarget, bool bRegenerateThumbnails)
 {
-	UTexture2D* ThumbnailWeightmap = nullptr;
-	UTexture2D* ThumbnailHeightmap = nullptr;
-		
 	TArray<FName> LayerNames;
 	TArray<ULandscapeLayerInfoObject*> LayerInfoObjs;
 	UMaterialInterface* Material = nullptr;
@@ -913,6 +911,8 @@ void ULandscapeEditorObject::RefreshImportLayersList(bool bRefreshFromTarget)
 	const TArray<FLandscapeImportLayer> OldLayersList = MoveTemp(ImportLandscape_Layers);
 	ImportLandscape_Layers.Reset(LayerNames.Num());
 
+	// Don't recreate the render state of everything, only update the materials context
+	FMaterialUpdateContext MaterialUpdateContext(FMaterialUpdateContext::EOptions::Default & ~FMaterialUpdateContext::EOptions::RecreateRenderStates);
 	for (int32 i = 0; i < LayerNames.Num(); i++)
 	{
 		const FName& LayerName = LayerNames[i];
@@ -934,29 +934,18 @@ void ULandscapeEditorObject::RefreshImportLayersList(bool bRefreshFromTarget)
 				}
 			}
 
-			if (bFound)
+			if (bFound && !bRegenerateThumbnails)
 			{
-				if (NewImportLayer.ThumbnailMIC->Parent != Material)
+				if ((NewImportLayer.ThumbnailMIC != nullptr) && (NewImportLayer.ThumbnailMIC->Parent != Material))
 				{
-					FMaterialUpdateContext Context;
 					NewImportLayer.ThumbnailMIC->SetParentEditorOnly(Material);
-					Context.AddMaterialInterface(NewImportLayer.ThumbnailMIC);
+					MaterialUpdateContext.AddMaterialInterface(NewImportLayer.ThumbnailMIC);
 				}
 			}
 			else
 			{
-				if (!ThumbnailWeightmap)
-				{
-					ThumbnailWeightmap = LoadObject<UTexture2D>(nullptr, TEXT("/Engine/EditorLandscapeResources/LandscapeThumbnailWeightmap.LandscapeThumbnailWeightmap"), nullptr, LOAD_None, nullptr);
-				}
-
-				if (!ThumbnailHeightmap)
-				{
-					ThumbnailHeightmap = LoadObject<UTexture2D>(nullptr, TEXT("/Engine/EditorLandscapeResources/LandscapeThumbnailHeightmap.LandscapeThumbnailHeightmap"), nullptr, LOAD_None, nullptr);
-				}
-
 				NewImportLayer.LayerName = LayerName;
-				NewImportLayer.ThumbnailMIC = ALandscapeProxy::GetLayerThumbnailMIC(Material, LayerName, ThumbnailWeightmap, ThumbnailHeightmap, nullptr);
+				NewImportLayer.ThumbnailMIC = UE::Landscape::CreateLandscapeLayerThumbnailMIC(MaterialUpdateContext, Material, LayerName);
 			}
 
 			if (bRefreshFromTarget)
diff --git a/Engine/Source/Editor/LandscapeEditor/Private/SLandscapeEditor.cpp b/Engine/Source/Editor/LandscapeEditor/Private/SLandscapeEditor.cpp
--- a/Engine/Source/Editor/LandscapeEditor/Private/SLandscapeEditor.cpp
+++ b/Engine/Source/Editor/LandscapeEditor/Private/SLandscapeEditor.cpp
@@ -37,26 +37,29 @@ void SLandscapeAssetThumbnail::Construct(const FArguments& InArgs, UObject* Asse
 {
 	FIntPoint ThumbnailSize = InArgs._ThumbnailSize;
 
-	const bool bShouldDisplayTargetLayerThumbnails = GetDefault<ULandscapeSettings>()->ShouldDisplayTargetLayerThumbnails();
-	AssetThumbnail = MakeShareable(new FAssetThumbnail(bShouldDisplayTargetLayerThumbnails ? Asset : nullptr, ThumbnailSize.X, ThumbnailSize.Y, ThumbnailPool));
+	AssetThumbnail = MakeShareable(new FAssetThumbnail(Asset, ThumbnailSize.X, ThumbnailSize.Y, ThumbnailPool));
 	OnAccessAsset = InArgs._OnAccessAsset;
 
+	FAssetThumbnailConfig AssetThumbnailConfig;
+	// If the asset is null, then we purposefully don't want this layer to have a thumbnail. Display a generic icon in that case :
+	if (Asset == nullptr)
+	{
+		AssetThumbnailConfig.bForceGenericThumbnail = true;
+		AssetThumbnailConfig.ClassThumbnailBrushOverride = FName("LevelEditor.LandscapeMode");
+	}
 	ChildSlot
 	[
 		SNew(SBox)
 		.WidthOverride(static_cast<float>(ThumbnailSize.X))
 		.HeightOverride(static_cast<float>(ThumbnailSize.Y))
 		[
-			AssetThumbnail->MakeThumbnailWidget()
+			AssetThumbnail->MakeThumbnailWidget(AssetThumbnailConfig)
 		]
 	];
 
-	if (bShouldDisplayTargetLayerThumbnails)
+	if (UMaterialInterface* MaterialInterface = Cast<UMaterialInterface>(Asset))
 	{
-		if (UMaterialInterface* MaterialInterface = Cast<UMaterialInterface>(Asset))
-		{
-			UMaterial::OnMaterialCompilationFinished().AddSP(this, &SLandscapeAssetThumbnail::OnMaterialCompilationFinished);
-		}
+		UMaterial::OnMaterialCompilationFinished().AddSP(this, &SLandscapeAssetThumbnail::OnMaterialCompilationFinished);
 	}
 }
 
diff --git a/Engine/Source/Editor/LandscapeEditor/Public/LandscapeEditorObject.h b/Engine/Source/Editor/LandscapeEditor/Public/LandscapeEditorObject.h
--- a/Engine/Source/Editor/LandscapeEditor/Public/LandscapeEditorObject.h
+++ b/Engine/Source/Editor/LandscapeEditor/Public/LandscapeEditorObject.h
@@ -765,7 +765,7 @@ public:
 	void ClearImportLandscapeData() { ImportLandscape_Data.Empty(); }
 	void ChooseBestComponentSizeForImport();
 	void ImportLandscapeData();
-	void RefreshImportLayersList(bool bRefreshFromTarget = false);
+	void RefreshImportLayersList(bool bRefreshFromTarget = false, bool bRegenerateThumbnails = false);
 	ELandscapeImportResult CreateImportLayersInfo(TArray<FLandscapeImportLayerInfo>& OutImportLayerInfos);
 	ELandscapeImportResult CreateNewLayersInfo(TArray<FLandscapeImportLayerInfo>& OutNewLayerInfos);
 	void InitializeDefaultHeightData(TArray<uint16>& OutData);
diff --git a/Engine/Source/Runtime/Landscape/Classes/LandscapeProxy.h b/Engine/Source/Runtime/Landscape/Classes/LandscapeProxy.h
--- a/Engine/Source/Runtime/Landscape/Classes/LandscapeProxy.h
+++ b/Engine/Source/Runtime/Landscape/Classes/LandscapeProxy.h
@@ -1393,6 +1393,7 @@ PRAGMA_ENABLE_DEPRECATION_WARNINGS
 	LANDSCAPE_API void UpdateAllComponentMaterialInstances(FMaterialUpdateContext& InOutMaterialContext, TArray<class FComponentRecreateRenderStateContext>& InOutRecreateRenderStateContext, bool bInInvalidateCombinationMaterials = false);
 
 	/** Create a thumbnail material for a given layer */
+	UE_DEPRECATED(5.6, "Use UE::Landscape::CreateLandscapeLayerThumbnailMIC")
 	LANDSCAPE_API static ULandscapeMaterialInstanceConstant* GetLayerThumbnailMIC(UMaterialInterface* LandscapeMaterial, FName LayerName, UTexture2D* ThumbnailWeightmap, UTexture2D* ThumbnailHeightmap, ALandscapeProxy* Proxy);
 
 	/** Import the given Height/Weight data into this landscape */
diff --git a/Engine/Source/Runtime/Landscape/Private/LandscapeEdit.cpp b/Engine/Source/Runtime/Landscape/Private/LandscapeEdit.cpp
--- a/Engine/Source/Runtime/Landscape/Private/LandscapeEdit.cpp
+++ b/Engine/Source/Runtime/Landscape/Private/LandscapeEdit.cpp
@@ -489,6 +489,7 @@ void ULandscapeComponent::UpdateRejectNavmeshUnderneath()
 }
 
 #if WITH_EDITOR
+// Deprecated
 ULandscapeMaterialInstanceConstant* ALandscapeProxy::GetLayerThumbnailMIC(UMaterialInterface* LandscapeMaterial, FName LayerName, UTexture2D* ThumbnailWeightmap, UTexture2D* ThumbnailHeightmap, ALandscapeProxy* Proxy)
 {
 	if (!LandscapeMaterial)
@@ -496,33 +497,8 @@ ULandscapeMaterialInstanceConstant* ALandscapeProxy::GetLayerThumbnailMIC(UMater
 		LandscapeMaterial = Proxy ? Proxy->GetLandscapeMaterial() : UMaterial::GetDefaultMaterial(MD_Surface);
 	}
 
-	FlushRenderingCommands();
-
-	ULandscapeMaterialInstanceConstant* MaterialInstance = NewObject<ULandscapeMaterialInstanceConstant>(GetTransientPackage());
-	MaterialInstance->bIsLayerThumbnail = true;
-	MaterialInstance->bMobile = false;
-	MaterialInstance->SetParentEditorOnly(LandscapeMaterial, false);
-
-	FStaticParameterSet StaticParameters;
-	MaterialInstance->GetStaticParameterValues(StaticParameters);
-
-	// Customize that material instance to only enable our terrain layer's weightmap : 
-	StaticParameters.EditorOnly.TerrainLayerWeightParameters.Add(FStaticTerrainLayerWeightParameter(LayerName, /*InWeightmapIndex = */0, /*bInWeightBasedBlend = */false));
-
-	// Don't recreate the render state of everything, only update the materials context
-	{
-		FMaterialUpdateContext MaterialUpdateContext(FMaterialUpdateContext::EOptions::Default & ~FMaterialUpdateContext::EOptions::RecreateRenderStates);
-		MaterialInstance->UpdateStaticPermutation(StaticParameters, &MaterialUpdateContext);
-	}
-	
-	FLinearColor Mask(1.0f, 0.0f, 0.0f, 0.0f);
-	MaterialInstance->SetVectorParameterValueEditorOnly(FName(*FString::Printf(TEXT("LayerMask_%s"), *LayerName.ToString())), Mask);
-	MaterialInstance->SetTextureParameterValueEditorOnly(FName(TEXT("Weightmap0")), ThumbnailWeightmap);
-	MaterialInstance->SetTextureParameterValueEditorOnly(FName(TEXT("Heightmap")), ThumbnailHeightmap);
-
-	MaterialInstance->PostEditChange();
-
-	return MaterialInstance;
+	FMaterialUpdateContext MaterialUpdateContext(FMaterialUpdateContext::EOptions::Default & ~FMaterialUpdateContext::EOptions::RecreateRenderStates);
+	return UE::Landscape::CreateLandscapeLayerThumbnailMIC(MaterialUpdateContext, LandscapeMaterial, LayerName);
 }
 
 bool ULandscapeComponent::ValidateCombinationMaterial(UMaterialInstanceConstant* InCombinationMaterial) const
@@ -532,7 +508,6 @@ bool ULandscapeComponent::ValidateCombinationMaterial(UMaterialInstanceConstant*
 		return false;
 	}
 
-
 	const TArray<FStaticTerrainLayerWeightParameter>& TerrainLayerWeightParameters = InCombinationMaterial->GetEditorOnlyStaticParameters().TerrainLayerWeightParameters;
 	const TArray<FWeightmapLayerAllocationInfo>& ComponentWeightmapAllocations = GetWeightmapLayerAllocations();
 
diff --git a/Engine/Source/Runtime/Landscape/Private/LandscapeInfo.cpp b/Engine/Source/Runtime/Landscape/Private/LandscapeInfo.cpp
--- a/Engine/Source/Runtime/Landscape/Private/LandscapeInfo.cpp
+++ b/Engine/Source/Runtime/Landscape/Private/LandscapeInfo.cpp
@@ -8,9 +8,9 @@
 
 LANDSCAPE_API FLandscapeInfoLayerSettings::FLandscapeInfoLayerSettings(ULandscapeLayerInfoObject* InLayerInfo, class ALandscapeProxy* InProxy)
 	: LayerInfoObj(InLayerInfo)
-	, LayerName((InLayerInfo != NULL) ? InLayerInfo->LayerName : NAME_None)
+	, LayerName((InLayerInfo != nullptr) ? InLayerInfo->LayerName : NAME_None)
 #if WITH_EDITORONLY_DATA
-	, ThumbnailMIC(NULL)
+	, ThumbnailMIC(nullptr)
 	, Owner(InProxy)
 	, DebugColorChannel(0)
 	, bValid(false)
diff --git a/Engine/Source/Runtime/Landscape/Private/LandscapeSettings.cpp b/Engine/Source/Runtime/Landscape/Private/LandscapeSettings.cpp
--- a/Engine/Source/Runtime/Landscape/Private/LandscapeSettings.cpp
+++ b/Engine/Source/Runtime/Landscape/Private/LandscapeSettings.cpp
@@ -11,15 +11,41 @@ void ULandscapeSettings::PostEditChangeProperty(FPropertyChangedEvent& PropertyC
 {
 	Super::PostEditChangeProperty(PropertyChangedEvent);
 
+	ILandscapeModule& LandscapeModule = FModuleManager::GetModuleChecked<ILandscapeModule>("Landscape");
 	const FName PropertyName = PropertyChangedEvent.MemberProperty ? PropertyChangedEvent.MemberProperty->GetFName() : NAME_None;
 	if ((PropertyName == GET_MEMBER_NAME_CHECKED(ULandscapeSettings, BrushSizeUIMax))
-		|| (PropertyName == GET_MEMBER_NAME_CHECKED(ULandscapeSettings, BrushSizeClampMax))
-		|| (PropertyName == GET_MEMBER_NAME_CHECKED(ULandscapeSettings, bDisplayTargetLayerThumbnails)))
+		|| (PropertyName == GET_MEMBER_NAME_CHECKED(ULandscapeSettings, BrushSizeClampMax)))
 	{
 		// If landscape mode is active, refresh the detail panel to apply the changes immediately : 
+		LandscapeModule.GetLandscapeEditorServices()->RefreshDetailPanel();
+	}
+
+	if (PropertyName == GET_MEMBER_NAME_CHECKED(ULandscapeSettings, bDisplayTargetLayerThumbnails))
+	{
+		LandscapeModule.GetLandscapeEditorServices()->RegenerateLayerThumbnails();
+		LandscapeModule.GetLandscapeEditorServices()->RefreshDetailPanel();
+	}
+}
+
+void ULandscapeSettings::PreEditUndo()
+{
+	Super::PreEditUndo();
+
+	check(!DisplayTargetLayerThumbnailsBeforeUndo.IsSet());
+	DisplayTargetLayerThumbnailsBeforeUndo = bDisplayTargetLayerThumbnails;
+}
+
+void ULandscapeSettings::PostEditUndo()
+{
+	Super::PostEditUndo();
+
+	if (DisplayTargetLayerThumbnailsBeforeUndo.GetValue() != bDisplayTargetLayerThumbnails)
+	{
 		ILandscapeModule& LandscapeModule = FModuleManager::GetModuleChecked<ILandscapeModule>("Landscape");
+		LandscapeModule.GetLandscapeEditorServices()->RegenerateLayerThumbnails();
 		LandscapeModule.GetLandscapeEditorServices()->RefreshDetailPanel();
 	}
+	DisplayTargetLayerThumbnailsBeforeUndo.Reset();
 }
 
 #endif // WITH_EDITOR
\ No newline at end of file
diff --git a/Engine/Source/Runtime/Landscape/Private/LandscapeUtils.cpp b/Engine/Source/Runtime/Landscape/Private/LandscapeUtils.cpp
--- a/Engine/Source/Runtime/Landscape/Private/LandscapeUtils.cpp
+++ b/Engine/Source/Runtime/Landscape/Private/LandscapeUtils.cpp
@@ -16,9 +16,11 @@
 #include "Editor/EditorEngine.h"
 #include "EditorDirectories.h"
 #include "Engine/Texture2D.h"
+#include "MaterialDomain.h"
 #include "LandscapeComponent.h"
-#include "ObjectTools.h"
 #include "LandscapeMaterialInstanceConstant.h"
+#include "LandscapeSettings.h"
+#include "ObjectTools.h"
 #endif // WITH_EDITOR
 
 #if WITH_EDITOR
@@ -254,10 +256,10 @@ uint32 GetTypeHash(const FTextureCopyRequest& InKey)
 bool operator==(const FTextureCopyRequest& InEntryA, const FTextureCopyRequest& InEntryB)
 {
 	return (InEntryA.Source == InEntryB.Source) &&
-			(InEntryA.Destination == InEntryB.Destination) &&
-			(InEntryA.DestinationSlice == InEntryB.DestinationSlice) &&
-			(InEntryA.TextureUsage == InEntryB.TextureUsage) &&
-			(InEntryA.TextureType == InEntryB.TextureType);
+		(InEntryA.Destination == InEntryB.Destination) &&
+		(InEntryA.DestinationSlice == InEntryB.DestinationSlice) &&
+		(InEntryA.TextureUsage == InEntryB.TextureUsage) &&
+		(InEntryA.TextureType == InEntryB.TextureType);
 }
 
 bool FBatchTextureCopy::AddWeightmapCopy(UTexture* InDestination, int8 InDestinationSlice, int8 InDestinationChannel, const ULandscapeComponent* InComponent, ULandscapeLayerInfoObject* InLayerInfo)
@@ -451,6 +453,46 @@ UMaterialInstance* CreateToolLandscapeMaterialInstanceConstant(UMaterialInterfac
 	return MaterialInstance;
 }
 
+ULandscapeMaterialInstanceConstant* CreateLandscapeLayerThumbnailMIC(FMaterialUpdateContext& MaterialUpdateContext, UMaterialInterface* LandscapeMaterial, FName LayerName)
+{
+	if (!GetDefault<ULandscapeSettings>()->ShouldDisplayTargetLayerThumbnails())
+	{
+		return nullptr;
+	}
+
+	if (LandscapeMaterial == nullptr)
+	{
+		LandscapeMaterial = UMaterial::GetDefaultMaterial(MD_Surface);
+	}
+
+	FlushRenderingCommands();
+
+	ULandscapeMaterialInstanceConstant* MaterialInstance = NewObject<ULandscapeMaterialInstanceConstant>(GetTransientPackage());
+	MaterialInstance->bIsLayerThumbnail = true;
+	MaterialInstance->bMobile = false;
+	MaterialInstance->SetParentEditorOnly(LandscapeMaterial, false);
+
+	FStaticParameterSet StaticParameters;
+	MaterialInstance->GetStaticParameterValues(StaticParameters);
+
+	// Customize that material instance to only enable our terrain layer's weightmap : 
+	StaticParameters.EditorOnly.TerrainLayerWeightParameters.Add(FStaticTerrainLayerWeightParameter(LayerName, /*InWeightmapIndex = */0, /*bInWeightBasedBlend = */false));
+
+	MaterialInstance->UpdateStaticPermutation(StaticParameters, &MaterialUpdateContext);
+
+	static UTexture2D* ThumbnailWeightmap = LoadObject<UTexture2D>(nullptr, TEXT("/Engine/EditorLandscapeResources/LandscapeThumbnailWeightmap.LandscapeThumbnailWeightmap"), nullptr, LOAD_None, nullptr);
+	static UTexture2D* ThumbnailHeightmap = LoadObject<UTexture2D>(nullptr, TEXT("/Engine/EditorLandscapeResources/LandscapeThumbnailHeightmap.LandscapeThumbnailHeightmap"), nullptr, LOAD_None, nullptr);
+
+	FLinearColor Mask(1.0f, 0.0f, 0.0f, 0.0f);
+	MaterialInstance->SetVectorParameterValueEditorOnly(FName(*FString::Printf(TEXT("LayerMask_%s"), *LayerName.ToString())), Mask);
+	MaterialInstance->SetTextureParameterValueEditorOnly(FName(TEXT("Weightmap0")), ThumbnailWeightmap);
+	MaterialInstance->SetTextureParameterValueEditorOnly(FName(TEXT("Heightmap")), ThumbnailHeightmap);
+
+	MaterialInstance->PostEditChange();
+
+	return MaterialInstance;
+}
+
 #endif // WITH_EDITOR
 
 } // end namespace UE::Landscape
\ No newline at end of file
diff --git a/Engine/Source/Runtime/Landscape/Public/LandscapeEditorServices.h b/Engine/Source/Runtime/Landscape/Public/LandscapeEditorServices.h
--- a/Engine/Source/Runtime/Landscape/Public/LandscapeEditorServices.h
+++ b/Engine/Source/Runtime/Landscape/Public/LandscapeEditorServices.h
@@ -20,6 +20,9 @@ public:
 	 */
 	virtual int32 GetOrCreateEditLayer(FName InEditLayerName, ALandscape* InTargetLandscape, const TSubclassOf<ULandscapeEditLayerBase>& InEditLayerClass = TSubclassOf<ULandscapeEditLayerBase>()) = 0;
 
-	/** Requires the landscape editor mode, if active to refresh its detail panel */
+	/** Requires the landscape editor mode, if active, to refresh its detail panel */
 	virtual void RefreshDetailPanel() = 0;
+
+	/** Requires the landscape editor mode, if active, to invalidate all of its thumbnails */
+	virtual void RegenerateLayerThumbnails() = 0;
 };
diff --git a/Engine/Source/Runtime/Landscape/Public/LandscapeSettings.h b/Engine/Source/Runtime/Landscape/Public/LandscapeSettings.h
--- a/Engine/Source/Runtime/Landscape/Public/LandscapeSettings.h
+++ b/Engine/Source/Runtime/Landscape/Public/LandscapeSettings.h
@@ -70,6 +70,8 @@ public:
 #if WITH_EDITOR
 	//~ Begin UObject Interface.
 	virtual void PostEditChangeProperty(FPropertyChangedEvent& PropertyChangedEvent) override;
+	virtual void PreEditUndo() override;
+	virtual void PostEditUndo() override;
 #endif // WITH_EDITOR
 
 public:
@@ -132,4 +134,7 @@ protected:
 	UPROPERTY(EditAnywhere, config, Category = "Configuration", meta = (ToolTip = 
 		"When true, temporal anti-aliasing will be inactive while in landscape mode. This avoids the ghosting effect on the landscape brush but can lead to aliasing or shimmering on other parts of the image."))
 	bool bDisableTemporalAntiAliasingInLandscapeMode = true;
+
+private:
+	TOptional<bool> DisplayTargetLayerThumbnailsBeforeUndo;
 };
diff --git a/Engine/Source/Runtime/Landscape/Public/LandscapeUtils.h b/Engine/Source/Runtime/Landscape/Public/LandscapeUtils.h
--- a/Engine/Source/Runtime/Landscape/Public/LandscapeUtils.h
+++ b/Engine/Source/Runtime/Landscape/Public/LandscapeUtils.h
@@ -16,10 +16,12 @@ class FRDGBuilder;
 class FRDGEventName;
 class ULandscapeComponent;
 class ULandscapeLayerInfoObject;
+class ULandscapeMaterialInstanceConstant;
 class ULevel;
 class UTexture;
 class UTexture2D;
 class FTextureResource;
+class FMaterialUpdateContext;
 enum class ELandscapeToolTargetType : uint8;
 enum class ELandscapeToolTargetTypeFlags : uint8;
 
@@ -29,7 +31,7 @@ enum class EBuildFlags : uint8;
 
 /**
 * Returns true if edit layers (GPU landscape tools) are enabled on this platform :
-* Note: this is intended for the editor but is in runtime code since global shaders need to exist in runtime modules 
+* Note: this is intended for the editor but is in runtime code since global shaders need to exist in runtime modules
 */
 LANDSCAPE_API bool DoesPlatformSupportEditLayers(EShaderPlatform InShaderPlatform);
 
@@ -42,38 +44,38 @@ LANDSCAPE_API FString GetLandscapeToolTargetTypeFlagsAsString(ELandscapeToolTarg
 
 /** This struct is usually meant to be allocated on the game thread (where there's no FRDGBuilder, which is Render Thread-only) and allows to queue successive operations (lambdas) onto a single render command,
  *   sharing the same FRDGBuilder. This allows to sequence a list of RDG passes from the game thread and makes it possible to interleave render thread operations (in a single render command) with game thread-initiated
- *   render commands. 
+ *   render commands.
  *   Here's an extensive use case :
- * 
+ *
  *   FRDGBuilderRecorder Recorder;
- * 
- *   // Start recording commands to a single graph builder : 
- *   Recorder.StartRecording(); 
- * 
+ *
+ *   // Start recording commands to a single graph builder :
+ *   Recorder.StartRecording();
+ *
  *   Recorder.EnqueueRDGCommand([](FRDGBuilder* GraphBuilder) { GraphBuilder->AddPass(); }); // Append a Pass_A
  *   ENQUEUE_RENDER_COMMAND(...)  // Push Render_Command_A (immediately)
  *   Recorder.EnqueueRDGCommand([](FRDGBuilder* GraphBuilder) { GraphBuilder->AddPass(); }, { { SomeTexture, ERHIAccess::RTV } }); // Append a Pass_B and inform of the final state of a given texture to prevent the RDG from auto-transitioning to SRVMask at the end
  *   Recorder.EnqueueRenderCommand([](FRHICommandListImmediate& InRHICmdList) mutable { [...]; }); // Append a Render_Command_B to the graph builder
  *   Recorder.EnqueueRDGCommand([](FRDGBuilder* GraphBuilder) { GraphBuilder->AddPass(); }, { { }SomeTexture, ERHIAccess::CopySrc } }); // Append a Pass_C and inform of the final state of a given texture to prevent the RDG from auto-transitioning to SRVMask at the end
- *   
+ *
  *   // Stop recording and issue a render command with all that's been recorded so far
- *   Recorder.StopRecordingAndFlush(RDG_EVENT_NAME("Pass ABC")); 
- * 
+ *   Recorder.StopRecordingAndFlush(RDG_EVENT_NAME("Pass ABC"));
+ *
  *   // Enqueue some game-thread render commands (e.g. BP render) either directly or via the recorder in immediate mode :
  *   ENQUEUE_RENDER_COMMAND(...)  // Render_Command_C
  *   Recorder.EnqueueRDGCommand([](FRDGBuilder* GraphBuilder) { GraphBuilder->AddPass(); }, { { SomeTexture, ERHIAccess::CopyDst } }); // Append a Pass_D and inform of the final state of a given texture to prevent the RDG from auto-transitioning to SRVMask at the end
  *   Recorder.EnqueueRenderCommand([RHIOperation](FRHICommandListImmediate& InRHICmdList) mutable { [...]; }); // Append a Render_Command_D
- *   
- *   // Start recording commands again to a new single graph builder : 
- *   Recorder.StartRecording(); 
+ *
+ *   // Start recording commands again to a new single graph builder :
+ *   Recorder.StartRecording();
  *   Recorder.EnqueueRDGCommand([](FRDGBuilder* GraphBuilder) { GraphBuilder->AddPass(); }); // Append a Pass_E
- *  
+ *
  *   // Stop recording and issue a render command with all that's been recorded so far
  *   Recorder.StopRecordingAndFlush(RDG_EVENT_NAME("Pass E"));
- * 
- *   --> Will yield the following sequence on the render thread: 
+ *
+ *   --> Will yield the following sequence on the render thread:
  *   + Render_Command_A
- *   + Render Command "Pass ABC" 
+ *   + Render Command "Pass ABC"
  *     +- (RDGBuilder_0)
  *     +- RDGBuilder_0.Pass_A
  *     +- RDGBuilder_0.Pass_B
@@ -82,10 +84,10 @@ LANDSCAPE_API FString GetLandscapeToolTargetTypeFlagsAsString(ELandscapeToolTarg
  *     +- RDGBuilder_0.SetTextureAccessFinal(OutputTexture, ERHIAccess::CopySrc); // Only the final state recorded for a given texture is set
  *     +- RDGBuilder_0.Execute
  *   + Render_Command_C
- *   + Render Command 
+ *   + Render Command
  *     +- (RDGBuilder_1)
  *     +- RDGBuilder_1.Pass_D
- *     +- RDGBuilder_1.SetTextureAccessFinal(OutputTexture, ERHIAccess::CopyDst); 
+ *     +- RDGBuilder_1.SetTextureAccessFinal(OutputTexture, ERHIAccess::CopyDst);
  *     +- RDGBuilder_1.Execute
  *   + Render_Command_D
  *   + Render Command "Pass E"
@@ -106,7 +108,7 @@ public:
 
 	LANDSCAPE_API FRDGBuilderRecorder() = default;
 	LANDSCAPE_API ~FRDGBuilderRecorder();
-	
+
 	enum class EState
 	{
 		Immediate, // In immediate mode, any command that is enqueued will be pushed to the render thread immediately (effectively acting like a ENQUEUE_RENDER_COMMAND) 
@@ -140,7 +142,7 @@ public:
 
 	/**
 	 * Records a FRDGRecorderRDGCommand to execute when registering passes to the single FRDGBuilder when in Recording mode or pushes it immediately to the render thread when in Immediate mode
-	 * 
+	 *
 	 * @param InRDGCommand Lambda to execute on the render thread
 	 * @param InRDGTextureAccessFinalList (optional) : list of external texture with the RHIAccess they should have when executing the FRDGBuilder. This is to prevent the RDG from auto-transitioning to SRVMask at the end
 	 */
@@ -158,7 +160,7 @@ public:
 	 */
 	LANDSCAPE_API bool IsEmpty() const;
 
-	/** 
+	/**
 	 * Cancels all recorder operations. This must be used if the FRDGEventRenderRecorder is "cancelled" (i.e. its sequence of operations is not flushed to a render command).
 	 *  Otherwise, there will be an assert in ~FRDGBuilderRecorder().
 	 */
@@ -175,8 +177,8 @@ private:
 
 #if RDG_EVENTS
 public:
-	/** 
-	 * Scope object meant to insert a RDG event in the RDG operations, as if it was inserted on the render thread on a FRDGBuilder. 
+	/**
+	 * Scope object meant to insert a RDG event in the RDG operations, as if it was inserted on the render thread on a FRDGBuilder.
 	 *  Use RDG_RENDER_COMMAND_RECORDER_BREADCRUMB_EVENT to create one.
 	 */
 	class FScopedBreadcrumbEvent final
@@ -323,12 +325,16 @@ struct UE_DEPRECATED(5.6, "This helper struct is deprecated. Please use utility
 };
 
 /**
- * Returns a pointer to a newly created ULandscapeMaterialInstanceConstant 
+ * Returns a pointer to a newly created ULandscapeMaterialInstanceConstant
  * @param	BaseMaterial	The base material appended to the debug name and set as parent of MaterialInstance
  * @return  Pointer to ULandscapeMaterialInstanceConstant
  */
 LANDSCAPE_API UMaterialInstance* CreateToolLandscapeMaterialInstanceConstant(UMaterialInterface* BaseMaterial);
 
+
+/** Create a thumbnail material for a given layer. Can return nullptr if the option to disable landscape thumbnails has been turned on */
+LANDSCAPE_API ULandscapeMaterialInstanceConstant* CreateLandscapeLayerThumbnailMIC(FMaterialUpdateContext& MaterialUpdateContext, UMaterialInterface* LandscapeMaterial, FName LayerName);
+
 #endif // WITH_EDITOR
 
-} // end namespace UE::Landscape
+} // end namespace UE::Landscape
\ No newline at end of file
