Improved logging relative to invalid landscape collision data + made the code resilient to invalid data

#rnx

diff --git a/Engine/Source/Runtime/Landscape/Private/LandscapeCollision.cpp b/Engine/Source/Runtime/Landscape/Private/LandscapeCollision.cpp
--- a/Engine/Source/Runtime/Landscape/Private/LandscapeCollision.cpp
+++ b/Engine/Source/Runtime/Landscape/Private/LandscapeCollision.cpp
@@ -1148,7 +1148,8 @@ ULandscapeHeightfieldCollisionComponent::FWriteRuntimeDataParams ULandscapeHeigh
 	const FCollisionSampleInfo Info = GetCollisionSampleInfo();
 
 	const uint16* Heights = (const uint16*)CollisionHeightData.LockReadOnly();
-	check(CollisionHeightData.GetElementCount() == Info.NumSamples + Info.NumSimpleSamples);
+	checkf(CollisionHeightData.GetElementCount() == Info.NumSamples + Info.NumSimpleSamples, 
+		TEXT("Invalid collision height data element count : %d found while there should be (%d samples + %d simple samples)"), CollisionHeightData.GetElementCount(), Info.NumSamples, Info.NumSimpleSamples);
 	const uint16* SimpleHeights = Heights + Info.NumSamples;
 
 	// Physical material data from layer system
@@ -1157,7 +1158,8 @@ ULandscapeHeightfieldCollisionComponent::FWriteRuntimeDataParams ULandscapeHeigh
 	if (DominantLayerData.GetElementCount())
 	{
 		DominantLayers = (const uint8*)DominantLayerData.LockReadOnly();
-		check(DominantLayerData.GetElementCount() == Info.NumSamples + Info.NumSimpleSamples);
+		checkf(DominantLayerData.GetElementCount() == Info.NumSamples + Info.NumSimpleSamples,
+			TEXT("Invalid dominant layer data element count : %d found while there should be (%d samples + %d simple samples)"), DominantLayerData.GetElementCount(), Info.NumSamples, Info.NumSimpleSamples);
 		SimpleDominantLayers = DominantLayers + Info.NumSamples;
 	}
 
@@ -1167,7 +1169,8 @@ ULandscapeHeightfieldCollisionComponent::FWriteRuntimeDataParams ULandscapeHeigh
 	if (PhysicalMaterialRenderData.GetElementCount())
 	{
 		RenderPhysicalMaterialIds = (const uint8*)PhysicalMaterialRenderData.LockReadOnly();
-		check(PhysicalMaterialRenderData.GetElementCount() == Info.NumSamples + Info.NumSimpleSamples);
+		checkf(PhysicalMaterialRenderData.GetElementCount() == Info.NumSamples + Info.NumSimpleSamples,
+			TEXT("Invalid physical material render data element count : %d found while there should be (%d samples + %d simple samples)"), PhysicalMaterialRenderData.GetElementCount(), Info.NumSamples, Info.NumSimpleSamples);
 		SimpleRenderPhysicalMaterialIds = RenderPhysicalMaterialIds + Info.NumSamples;
 	}
 
@@ -1538,7 +1541,8 @@ bool ULandscapeMeshCollisionComponent::CookCollisionData(const FName& Format, bo
 
 	const uint16* Heights = (const uint16*)CollisionHeightData.LockReadOnly();
 	const uint16* XYOffsets = (const uint16*)CollisionXYOffsetData.LockReadOnly();
-	check(CollisionHeightData.GetElementCount() == NumVerts + NumSimpleVerts);
+	checkf(CollisionHeightData.GetElementCount() == NumVerts + NumSimpleVerts,
+		TEXT("Invalid collision height data element count : %d found while there should be (%d samples + %d simple samples)"), CollisionHeightData.GetElementCount(), NumVerts, NumSimpleVerts);
 	check(CollisionXYOffsetData.GetElementCount() == NumVerts * 2);
 
 	const uint8* DominantLayers = nullptr;
@@ -2073,7 +2077,8 @@ void ULandscapeHeightfieldCollisionComponent::UpdateHeightfieldRegion(int32 Comp
 			bool bIsMirrored = GetComponentToWorld().GetDeterminant() < 0.f;
 	
 			uint16* Heights = (uint16*)CollisionHeightData.Lock(LOCK_READ_ONLY);
-			check(CollisionHeightData.GetElementCount() == (FMath::Square(CollisionSizeVerts) + FMath::Square(SimpleCollisionSizeVerts)));
+			checkf(CollisionHeightData.GetElementCount() == FMath::Square(CollisionSizeVerts) + FMath::Square(SimpleCollisionSizeVerts),
+				TEXT("Invalid collision height data element count : %d found while there should be (%d samples + %d simple samples)"), CollisionHeightData.GetElementCount(), FMath::Square(CollisionSizeVerts), FMath::Square(SimpleCollisionSizeVerts));
 	
 			int32 HeightfieldY1 = ComponentY1;
 			int32 HeightfieldX1 = (bIsMirrored ? ComponentX1 : (CollisionSizeVerts - ComponentX2 - 1));
@@ -2812,76 +2817,92 @@ void ULandscapeInfo::UpdateAddCollision(FIntPoint LandscapeKey)
 	// Corner Cases...
 	if (NeighborCollisions[0])
 	{
-		uint16* Heights = (uint16*)NeighborCollisions[0]->CollisionHeightData.Lock(LOCK_READ_ONLY);
-		int32 CollisionSizeVerts = NeighborCollisions[0]->CollisionSizeQuads + 1;
-		HeightCorner[0] = Heights[CollisionSizeVerts - 1 + (CollisionSizeVerts - 1)*CollisionSizeVerts];
-		CornerSet |= 1;
+		if (uint16* Heights = (uint16*)NeighborCollisions[0]->CollisionHeightData.Lock(LOCK_READ_ONLY))
+		{
+			int32 CollisionSizeVerts = NeighborCollisions[0]->CollisionSizeQuads + 1;
+			HeightCorner[0] = Heights[CollisionSizeVerts - 1 + (CollisionSizeVerts - 1)*CollisionSizeVerts];
+			CornerSet |= 1;
+		}
 		NeighborCollisions[0]->CollisionHeightData.Unlock();
 	}
 	if (NeighborCollisions[2])
 	{
-		uint16* Heights = (uint16*)NeighborCollisions[2]->CollisionHeightData.Lock(LOCK_READ_ONLY);
-		int32 CollisionSizeVerts = NeighborCollisions[2]->CollisionSizeQuads + 1;
-		HeightCorner[1] = Heights[(CollisionSizeVerts - 1)*CollisionSizeVerts];
-		CornerSet |= 1 << 1;
+		if (uint16* Heights = (uint16*)NeighborCollisions[2]->CollisionHeightData.Lock(LOCK_READ_ONLY))
+		{
+			int32 CollisionSizeVerts = NeighborCollisions[2]->CollisionSizeQuads + 1;
+			HeightCorner[1] = Heights[(CollisionSizeVerts - 1)*CollisionSizeVerts];
+			CornerSet |= 1 << 1;
+		}
 		NeighborCollisions[2]->CollisionHeightData.Unlock();
 	}
 	if (NeighborCollisions[5])
 	{
-		uint16* Heights = (uint16*)NeighborCollisions[5]->CollisionHeightData.Lock(LOCK_READ_ONLY);
-		int32 CollisionSizeVerts = NeighborCollisions[5]->CollisionSizeQuads + 1;
-		HeightCorner[2] = Heights[(CollisionSizeVerts - 1)];
-		CornerSet |= 1 << 2;
+		if (uint16* Heights = (uint16*)NeighborCollisions[5]->CollisionHeightData.Lock(LOCK_READ_ONLY))
+		{
+			int32 CollisionSizeVerts = NeighborCollisions[5]->CollisionSizeQuads + 1;
+			HeightCorner[2] = Heights[(CollisionSizeVerts - 1)];
+			CornerSet |= 1 << 2;
+		}
 		NeighborCollisions[5]->CollisionHeightData.Unlock();
 	}
 	if (NeighborCollisions[7])
 	{
-		uint16* Heights = (uint16*)NeighborCollisions[7]->CollisionHeightData.Lock(LOCK_READ_ONLY);
-		int32 CollisionSizeVerts = NeighborCollisions[7]->CollisionSizeQuads + 1;
-		HeightCorner[3] = Heights[0];
-		CornerSet |= 1 << 3;
+		if (uint16* Heights = (uint16*)NeighborCollisions[7]->CollisionHeightData.Lock(LOCK_READ_ONLY))
+		{
+			int32 CollisionSizeVerts = NeighborCollisions[7]->CollisionSizeQuads + 1;
+			HeightCorner[3] = Heights[0];
+			CornerSet |= 1 << 3;
+		}
 		NeighborCollisions[7]->CollisionHeightData.Unlock();
 	}
 
 	// Other cases...
 	if (NeighborCollisions[1])
 	{
-		uint16* Heights = (uint16*)NeighborCollisions[1]->CollisionHeightData.Lock(LOCK_READ_ONLY);
-		int32 CollisionSizeVerts = NeighborCollisions[1]->CollisionSizeQuads + 1;
-		HeightCorner[0] = Heights[(CollisionSizeVerts - 1)*CollisionSizeVerts];
-		CornerSet |= 1;
-		HeightCorner[1] = Heights[CollisionSizeVerts - 1 + (CollisionSizeVerts - 1)*CollisionSizeVerts];
-		CornerSet |= 1 << 1;
+		if (uint16* Heights = (uint16*)NeighborCollisions[1]->CollisionHeightData.Lock(LOCK_READ_ONLY))
+		{
+			int32 CollisionSizeVerts = NeighborCollisions[1]->CollisionSizeQuads + 1;
+			HeightCorner[0] = Heights[(CollisionSizeVerts - 1)*CollisionSizeVerts];
+			CornerSet |= 1;
+			HeightCorner[1] = Heights[CollisionSizeVerts - 1 + (CollisionSizeVerts - 1)*CollisionSizeVerts];
+			CornerSet |= 1 << 1;
+		}
 		NeighborCollisions[1]->CollisionHeightData.Unlock();
 	}
 	if (NeighborCollisions[3])
 	{
-		uint16* Heights = (uint16*)NeighborCollisions[3]->CollisionHeightData.Lock(LOCK_READ_ONLY);
-		int32 CollisionSizeVerts = NeighborCollisions[3]->CollisionSizeQuads + 1;
-		HeightCorner[0] = Heights[(CollisionSizeVerts - 1)];
-		CornerSet |= 1;
-		HeightCorner[2] = Heights[CollisionSizeVerts - 1 + (CollisionSizeVerts - 1)*CollisionSizeVerts];
-		CornerSet |= 1 << 2;
+		if (uint16* Heights = (uint16*)NeighborCollisions[3]->CollisionHeightData.Lock(LOCK_READ_ONLY))
+		{
+			int32 CollisionSizeVerts = NeighborCollisions[3]->CollisionSizeQuads + 1;
+			HeightCorner[0] = Heights[(CollisionSizeVerts - 1)];
+			CornerSet |= 1;
+			HeightCorner[2] = Heights[CollisionSizeVerts - 1 + (CollisionSizeVerts - 1)*CollisionSizeVerts];
+			CornerSet |= 1 << 2;
+		}
 		NeighborCollisions[3]->CollisionHeightData.Unlock();
 	}
 	if (NeighborCollisions[4])
 	{
-		uint16* Heights = (uint16*)NeighborCollisions[4]->CollisionHeightData.Lock(LOCK_READ_ONLY);
-		int32 CollisionSizeVerts = NeighborCollisions[4]->CollisionSizeQuads + 1;
-		HeightCorner[1] = Heights[0];
-		CornerSet |= 1 << 1;
-		HeightCorner[3] = Heights[(CollisionSizeVerts - 1)*CollisionSizeVerts];
-		CornerSet |= 1 << 3;
+		if (uint16* Heights = (uint16*)NeighborCollisions[4]->CollisionHeightData.Lock(LOCK_READ_ONLY))
+		{
+			int32 CollisionSizeVerts = NeighborCollisions[4]->CollisionSizeQuads + 1;
+			HeightCorner[1] = Heights[0];
+			CornerSet |= 1 << 1;
+			HeightCorner[3] = Heights[(CollisionSizeVerts - 1)*CollisionSizeVerts];
+			CornerSet |= 1 << 3;
+		}
 		NeighborCollisions[4]->CollisionHeightData.Unlock();
 	}
 	if (NeighborCollisions[6])
 	{
-		uint16* Heights = (uint16*)NeighborCollisions[6]->CollisionHeightData.Lock(LOCK_READ_ONLY);
-		int32 CollisionSizeVerts = NeighborCollisions[6]->CollisionSizeQuads + 1;
-		HeightCorner[2] = Heights[0];
-		CornerSet |= 1 << 2;
-		HeightCorner[3] = Heights[(CollisionSizeVerts - 1)];
-		CornerSet |= 1 << 3;
+		if (uint16* Heights = (uint16*)NeighborCollisions[6]->CollisionHeightData.Lock(LOCK_READ_ONLY))
+		{
+			int32 CollisionSizeVerts = NeighborCollisions[6]->CollisionSizeQuads + 1;
+			HeightCorner[2] = Heights[0];
+			CornerSet |= 1 << 2;
+			HeightCorner[3] = Heights[(CollisionSizeVerts - 1)];
+			CornerSet |= 1 << 3;
+		}
 		NeighborCollisions[6]->CollisionHeightData.Unlock();
 	}
 
@@ -2911,7 +2932,8 @@ void ULandscapeHeightfieldCollisionComponent::ExportCustomProperties(FOutputDevi
 	int32 CollisionSizeVerts = CollisionSizeQuads + 1;
 	int32 SimpleCollisionSizeVerts = SimpleCollisionSizeQuads > 0 ? SimpleCollisionSizeQuads + 1 : 0;
 	int32 NumHeights = FMath::Square(CollisionSizeVerts) + FMath::Square(SimpleCollisionSizeVerts);
-	check(CollisionHeightData.GetElementCount() == NumHeights);
+	checkf(CollisionHeightData.GetElementCount() == FMath::Square(CollisionSizeVerts) + FMath::Square(SimpleCollisionSizeVerts),
+		TEXT("Invalid collision height data element count : %d found while there should be (%d samples + %d simple samples)"), CollisionHeightData.GetElementCount(), FMath::Square(CollisionSizeVerts), FMath::Square(SimpleCollisionSizeVerts));
 
 	uint16* Heights = (uint16*)CollisionHeightData.Lock(LOCK_READ_ONLY);
 
