Fix for Metal validation issues during texture upload for Landscape Edge Patching.

diff --git a/Engine/Source/Runtime/Landscape/Private/LandscapeEdgeFixup.cpp b/Engine/Source/Runtime/Landscape/Private/LandscapeEdgeFixup.cpp
--- a/Engine/Source/Runtime/Landscape/Private/LandscapeEdgeFixup.cpp
+++ b/Engine/Source/Runtime/Landscape/Private/LandscapeEdgeFixup.cpp
@@ -1170,7 +1170,7 @@ void ULandscapeHeightmapTextureEdgeFixup::PatchTextureEdge_Internal(UE::Landscap
 					}
 
 					{
-						constexpr uint32 SourcePitch = 4;
+						const uint32 SourcePitch = 4 * SrcWidth;
 						FUpdateTextureRegion2D UpdateRegion(DstStartX, DstStartY, SrcStartX, SrcStartY, SrcWidth, SrcHeight);
 
 						FIXUP_DEBUG_LOG_RENDER(TEXT("    PatchTextureEdge_Internal [RHIUpdateTexture2D] res:%p e:%d mip:%d reg:[%d %d %d %d %d %d]"), HeightTextureResource, EdgeIndex, MipIndex,
