Prevented landscape edit layers merge context from automatically changing the RT state when cycling render targets. The problem is that since we're mixing automatic transitions by the RDG graph command recorder with manual transitions, we need to make sure that every manual transition is actually performed and there were cases where we would mark a transition as done on the game thread side but it wasn't actually performed. So instead, we're relying on the user to lazily perform the transitions before the draws

#rnx

diff --git a/Engine/Source/Runtime/Landscape/Private/LandscapeBlueprintBrushBase.cpp b/Engine/Source/Runtime/Landscape/Private/LandscapeBlueprintBrushBase.cpp
--- a/Engine/Source/Runtime/Landscape/Private/LandscapeBlueprintBrushBase.cpp
+++ b/Engine/Source/Runtime/Landscape/Private/LandscapeBlueprintBrushBase.cpp
@@ -300,10 +300,13 @@ void ALandscapeBlueprintBrushBase::RenderLayer(ILandscapeEditLayerRenderer::FRen
 	// By default, use the old way of rendering BP brushes : 
 
 	// Swap the render targets so that the layer's input RT is the latest combined result :
-	// The write render target will be accessed as ERHIAccess::CopyDest all along : 
-	RenderParams.MergeRenderContext->CycleBlendRenderTargets(/*InDesiredWriteAccess = */ERHIAccess::CopyDest, RDGBuilderRecorder);
+	RenderParams.MergeRenderContext->CycleBlendRenderTargets(RDGBuilderRecorder);
 	ULandscapeScratchRenderTarget* WriteRT = RenderParams.MergeRenderContext->GetBlendRenderTargetWrite();
 	ULandscapeScratchRenderTarget* CurrentLayerReadRT = RenderParams.MergeRenderContext->GetBlendRenderTargetRead();
+	
+	// The write render target will be accessed as ERHIAccess::CopyDest all along : 
+	WriteRT->TransitionTo(ERHIAccess::CopyDest, RDGBuilderRecorder);
+	CurrentLayerReadRT->TransitionTo(ERHIAccess::SRVMask, RDGBuilderRecorder);
 
 	const bool bIsHeightmapMerge = RenderParams.MergeRenderContext->IsHeightmapMerge();
 	const TArray<FName> EnabledWeightmaps = RenderParams.RendererState.GetActiveTargetWeightmaps();
diff --git a/Engine/Source/Runtime/Landscape/Private/LandscapeEditLayerRenderer.cpp b/Engine/Source/Runtime/Landscape/Private/LandscapeEditLayerRenderer.cpp
--- a/Engine/Source/Runtime/Landscape/Private/LandscapeEditLayerRenderer.cpp
+++ b/Engine/Source/Runtime/Landscape/Private/LandscapeEditLayerRenderer.cpp
@@ -913,7 +913,7 @@ void FMergeRenderContext::FreeBatchResources(const FMergeRenderBatch& InRenderBa
 	PerTargetLayerValidityRenderTargets.Empty();
 }
 
-void FMergeRenderContext::CycleBlendRenderTargets(ERHIAccess InDesiredWriteAccess, UE::Landscape::FRDGBuilderRecorder& RDGBuilderRecorder)
+void FMergeRenderContext::CycleBlendRenderTargets(UE::Landscape::FRDGBuilderRecorder& RDGBuilderRecorder)
 {
 	const bool bFirstWrite = (CurrentBlendRenderTargetWriteIndex < 0);
 	CurrentBlendRenderTargetWriteIndex = (CurrentBlendRenderTargetWriteIndex + 1) % BlendRenderTargets.Num();
@@ -925,14 +925,6 @@ void FMergeRenderContext::CycleBlendRenderTargets(ERHIAccess InDesiredWriteAcces
 		{
 			GetBlendRenderTargetWrite()->Clear(RDGBuilderRecorder);
 		}
-
-		// Change the state of the new Read and Write (ReadPrevious is already SRV since it was Read before)
-		ULandscapeScratchRenderTarget::FTransitionBatcherScope TransitionsScope(RDGBuilderRecorder);
-		TransitionsScope.TransitionTo(GetBlendRenderTargetRead(), ERHIAccess::SRVMask);
-		if (InDesiredWriteAccess != ERHIAccess::None)
-		{
-			TransitionsScope.TransitionTo(GetBlendRenderTargetWrite(), InDesiredWriteAccess);
-		}
 	}
 }
 
@@ -1162,7 +1154,7 @@ void FMergeRenderContext::Render(TFunction<void(const FOnRenderBatchTargetGroupD
 	FRDGBuilderRecorder RDGBuilderRecorder;
 
 	// Kick start the blend render targets : 
-	CycleBlendRenderTargets(/*InDesiredWriteAccess = */ERHIAccess::None, RDGBuilderRecorder);
+	CycleBlendRenderTargets(RDGBuilderRecorder);
 
 	const int32 NumBatches = RenderBatches.Num();
 	for (CurrentRenderBatchIndex = 0; CurrentRenderBatchIndex < NumBatches; ++CurrentRenderBatchIndex)
diff --git a/Engine/Source/Runtime/Landscape/Private/LandscapeEditLayerRendererPrivate.cpp b/Engine/Source/Runtime/Landscape/Private/LandscapeEditLayerRendererPrivate.cpp
--- a/Engine/Source/Runtime/Landscape/Private/LandscapeEditLayerRendererPrivate.cpp
+++ b/Engine/Source/Runtime/Landscape/Private/LandscapeEditLayerRendererPrivate.cpp
@@ -52,11 +52,12 @@ void ULandscapeDefaultEditLayerRenderer::RenderLayer(FRenderParams& RenderParams
 	FMergeRenderContext* RenderContext = RenderParams.MergeRenderContext;
 	const FMergeRenderBatch* RenderBatch = RenderContext->GetCurrentRenderBatch();
 
-	RenderContext->CycleBlendRenderTargets(/*InDesiredWriteAccess = */ERHIAccess::RTV, RDGBuilderRecorder);
+	RenderContext->CycleBlendRenderTargets(RDGBuilderRecorder);
 	ULandscapeScratchRenderTarget* WriteRT = RenderContext->GetBlendRenderTargetWrite();
 
 	// Start from a blank canvas so that the first layer is blended with nothing underneath :
 	WriteRT->Clear(RDGBuilderRecorder);
+	check(WriteRT->GetCurrentState() == ERHIAccess::RTV);
 
 	// Render the components of the batch for each target layer into the "pseudo-stencil" buffer, so that it can be sampled by users as a UTextures in UMaterials and such : 
 	RenderContext->RenderValidityRenderTargets(*RenderBatch, RDGBuilderRecorder);
diff --git a/Engine/Source/Runtime/Landscape/Private/LandscapeEditLayers.cpp b/Engine/Source/Runtime/Landscape/Private/LandscapeEditLayers.cpp
--- a/Engine/Source/Runtime/Landscape/Private/LandscapeEditLayers.cpp
+++ b/Engine/Source/Runtime/Landscape/Private/LandscapeEditLayers.cpp
@@ -5648,211 +5648,216 @@ void ULandscapeEditLayerPersistent::RenderLayer(FRenderParams& RenderParams, UE:
 		return;
 	}
 
-	RenderParams.MergeRenderContext->CycleBlendRenderTargets(/*InDesiredWriteAccess = */ERHIAccess::None, RDGBuilderRecorder);
+	RenderParams.MergeRenderContext->CycleBlendRenderTargets(RDGBuilderRecorder);
 	ULandscapeScratchRenderTarget* WriteRT = RenderParams.MergeRenderContext->GetBlendRenderTargetWrite();
+
 	WriteRT->Clear(RDGBuilderRecorder);
 	// We will write to the RT using a (bunch of) PS : 
-	WriteRT->TransitionTo(ERHIAccess::RTV, RDGBuilderRecorder);
+	check(WriteRT->GetCurrentState() == ERHIAccess::RTV);
 
 	// First step : render the edit layer's component quads
 	{
-	// Sort the list of quads by output texture first, then by source texture, because we'll coalesce several texture copies into the same pass (one output, multiple inputs)
-	CopyQuadsParams.Sort([](const FCopyQuadParams& InLHS, const FCopyQuadParams& InRHS)
-	{
-		if (InLHS.DestinationArrayIndex == InRHS.DestinationArrayIndex)
-		{
-			if (InLHS.SourceTextureResource == InRHS.SourceTextureResource)
+		// Sort the list of quads by output texture first, then by source texture, because we'll coalesce several texture copies into the same pass (one output, multiple inputs)
+		CopyQuadsParams.Sort([](const FCopyQuadParams& InLHS, const FCopyQuadParams& InRHS)
 			{
-				return (InLHS.SourceRect.Min.Y == InRHS.SourceRect.Min.Y) 
-					? (InLHS.SourceRect.Min.X < InRHS.SourceRect.Min.X)
-					: (InLHS.SourceRect.Min.Y < InRHS.SourceRect.Min.Y);
-			}
-			return (InLHS.SourceTextureResource < InRHS.SourceTextureResource);
-		}
+				if (InLHS.DestinationArrayIndex == InRHS.DestinationArrayIndex)
+				{
+					if (InLHS.SourceTextureResource == InRHS.SourceTextureResource)
+					{
+						return (InLHS.SourceRect.Min.Y == InRHS.SourceRect.Min.Y)
+							? (InLHS.SourceRect.Min.X < InRHS.SourceRect.Min.X)
+							: (InLHS.SourceRect.Min.Y < InRHS.SourceRect.Min.Y);
+					}
+					return (InLHS.SourceTextureResource < InRHS.SourceTextureResource);
+				}
 
-		return (InLHS.DestinationArrayIndex < InRHS.DestinationArrayIndex);
-	});
+				return (InLHS.DestinationArrayIndex < InRHS.DestinationArrayIndex);
+			});
 
-	// Now process this list of quads and prepare as many passes as necessary for performing all the copies :
+		// Now process this list of quads and prepare as many passes as necessary for performing all the copies :
 		auto RDGCommand =
-			[ CopyQuadsParams = MoveTemp(CopyQuadsParams)
+			[CopyQuadsParams = MoveTemp(CopyQuadsParams)
 			, OutputResource = WriteRT->GetRenderTarget()->GetResource()
 			, OutputResourceName = WriteRT->GetDebugName()
 			, SceneInterface
 			, CommonSourceTextureSize
 			, bUseSwizzleMask] (FRDGBuilder& GraphBuilder)
-		{
+			{
 				RDG_EVENT_SCOPE(GraphBuilder, "CopyEditLayer -> %s", *OutputResourceName);
 
-			FRDGTextureSRVRef BlackDummySRVRef = GraphBuilder.CreateSRV(FRDGTextureSRVDesc::Create(GSystemTextures.GetBlackDummy(GraphBuilder)));
-
-			FCopyQuadsMultiSourcePS::FPermutationDomain PermutationVector;
-			PermutationVector.Set< FCopyQuadsMultiSourcePS::FUseSwizleMask >(bUseSwizzleMask);
-			FGlobalShaderMap* ShaderMap = GetGlobalShaderMap(GMaxRHIFeatureLevel);
-			TShaderRef<FCopyQuadsMultiSourcePS> PixelShader = ShaderMap->GetShader<FCopyQuadsMultiSourcePS>(PermutationVector);
+				FRDGTextureSRVRef BlackDummySRVRef = GraphBuilder.CreateSRV(FRDGTextureSRVDesc::Create(GSystemTextures.GetBlackDummy(GraphBuilder)));
 
-			// TODO [jonathan.bard] this is just to avoid a RHI validation error for unoptimized shaders... once validation is made to not issue those errors, we can remove this
-			// Create a SceneView to please the shader bindings, but it's unused in practice 
-			FSceneViewFamilyContext ViewFamily(FSceneViewFamily::ConstructionValues(nullptr, SceneInterface, FEngineShowFlags(ESFIM_Game)).SetTime(FGameTime::GetTimeSinceAppStart()));
-			FSceneViewInitOptions ViewInitOptions;
-			ViewInitOptions.ViewFamily = &ViewFamily;
-			ViewInitOptions.SetViewRectangle(FIntRect(0, 0, 1, 1)); // Use a dummy rect to avoid a check(slow)
+				FCopyQuadsMultiSourcePS::FPermutationDomain PermutationVector;
+				PermutationVector.Set< FCopyQuadsMultiSourcePS::FUseSwizleMask >(bUseSwizzleMask);
+				FGlobalShaderMap* ShaderMap = GetGlobalShaderMap(GMaxRHIFeatureLevel);
+				TShaderRef<FCopyQuadsMultiSourcePS> PixelShader = ShaderMap->GetShader<FCopyQuadsMultiSourcePS>(PermutationVector);
+
+				// TODO [jonathan.bard] this is just to avoid a RHI validation error for unoptimized shaders... once validation is made to not issue those errors, we can remove this
+				// Create a SceneView to please the shader bindings, but it's unused in practice 
+				FSceneViewFamilyContext ViewFamily(FSceneViewFamily::ConstructionValues(nullptr, SceneInterface, FEngineShowFlags(ESFIM_Game)).SetTime(FGameTime::GetTimeSinceAppStart()));
+				FSceneViewInitOptions ViewInitOptions;
+				ViewInitOptions.ViewFamily = &ViewFamily;
+				ViewInitOptions.SetViewRectangle(FIntRect(0, 0, 1, 1)); // Use a dummy rect to avoid a check(slow)
 				GetRendererModule().CreateAndInitSingleView(GraphBuilder.RHICmdList, &ViewFamily, &ViewInitOptions);
-			const FSceneView* View = ViewFamily.Views[0];
+				const FSceneView* View = ViewFamily.Views[0];
 
-			// The following variables allow to accumulate quads for rendering in a single pass. Then, we flush them all out when the pass is full and we start again : 
-			int32 CurrentDestinationArrayIndex = INDEX_NONE;
-			TArray<FTextureResource*> CurrentSourceTextureResources;
-			TArray<FUintVector2> CurrentQuadInfos;
-			TArray<FUintVector4> CurrentSourceRects, CurrentDestinationRects;
+				// The following variables allow to accumulate quads for rendering in a single pass. Then, we flush them all out when the pass is full and we start again : 
+				int32 CurrentDestinationArrayIndex = INDEX_NONE;
+				TArray<FTextureResource*> CurrentSourceTextureResources;
+				TArray<FUintVector2> CurrentQuadInfos;
+				TArray<FUintVector4> CurrentSourceRects, CurrentDestinationRects;
 
-			FRDGTextureRef OutputTexture = GraphBuilder.RegisterExternalTexture(CreateRenderTarget(OutputResource->GetTextureRHI(), TEXT("OutputTexture")));
+				FRDGTextureRef OutputTexture = GraphBuilder.RegisterExternalTexture(CreateRenderTarget(OutputResource->GetTextureRHI(), TEXT("OutputTexture")));
 
-			// When we've reached the limit of what we can render in one pass, flush all and add a render pass : 
-			auto FlushCopyQuads = [&GraphBuilder, View, OutputTexture, &CurrentDestinationArrayIndex, &CurrentSourceTextureResources, &CurrentSourceRects, &CurrentDestinationRects,
-				&CurrentQuadInfos, &CommonSourceTextureSize, BlackDummySRVRef, ShaderMap, PixelShader]()
-			{
-				if (CurrentSourceTextureResources.IsEmpty())
-				{
-					// Nothing to flush :
-					check(CurrentSourceRects.IsEmpty() && CurrentDestinationRects.IsEmpty() && CurrentQuadInfos.IsEmpty());
-					return;
-				}
+				// When we've reached the limit of what we can render in one pass, flush all and add a render pass : 
+				auto FlushCopyQuads = [&GraphBuilder, View, OutputTexture, &CurrentDestinationArrayIndex, &CurrentSourceTextureResources, &CurrentSourceRects, &CurrentDestinationRects,
+					&CurrentQuadInfos, &CommonSourceTextureSize, BlackDummySRVRef, ShaderMap, PixelShader]()
+					{
+						if (CurrentSourceTextureResources.IsEmpty())
+						{
+							// Nothing to flush :
+							check(CurrentSourceRects.IsEmpty() && CurrentDestinationRects.IsEmpty() && CurrentQuadInfos.IsEmpty());
+							return;
+						}
 
-				check(!CurrentSourceRects.IsEmpty() && (CurrentSourceRects.Num() == CurrentDestinationRects.Num()) && (CurrentSourceRects.Num() == CurrentQuadInfos.Num()));
+						check(!CurrentSourceRects.IsEmpty() && (CurrentSourceRects.Num() == CurrentDestinationRects.Num()) && (CurrentSourceRects.Num() == CurrentQuadInfos.Num()));
 
-				FRDGBufferRef RectBuffer = CreateUploadBuffer(GraphBuilder, TEXT("DestinationRects"), TConstArrayView<FUintVector4>(MakeArrayView(CurrentDestinationRects)));
-				FRDGBufferSRVRef RectBufferSRV = GraphBuilder.CreateSRV(FRDGBufferSRVDesc(RectBuffer, PF_R32G32B32A32_UINT));
+						FRDGBufferRef RectBuffer = CreateUploadBuffer(GraphBuilder, TEXT("DestinationRects"), TConstArrayView<FUintVector4>(MakeArrayView(CurrentDestinationRects)));
+						FRDGBufferSRVRef RectBufferSRV = GraphBuilder.CreateSRV(FRDGBufferSRVDesc(RectBuffer, PF_R32G32B32A32_UINT));
 
-				FRDGBufferRef RectUVBuffer = CreateUploadBuffer(GraphBuilder, TEXT("RectUVs"), TConstArrayView<FUintVector4>(MakeArrayView(CurrentSourceRects)));
-				FRDGBufferSRVRef RectUVBufferSRV = GraphBuilder.CreateSRV(FRDGBufferSRVDesc(RectUVBuffer, PF_R32G32B32A32_UINT));
+						FRDGBufferRef RectUVBuffer = CreateUploadBuffer(GraphBuilder, TEXT("RectUVs"), TConstArrayView<FUintVector4>(MakeArrayView(CurrentSourceRects)));
+						FRDGBufferSRVRef RectUVBufferSRV = GraphBuilder.CreateSRV(FRDGBufferSRVDesc(RectUVBuffer, PF_R32G32B32A32_UINT));
 
-				FRDGBufferRef QuadInfosBuffer = CreateUploadBuffer(GraphBuilder, TEXT("QuadInfos"), TConstArrayView<FUintVector2>(MakeArrayView(CurrentQuadInfos)));
-				FRDGBufferSRVRef QuadInfosBufferSRV = GraphBuilder.CreateSRV(FRDGBufferSRVDesc(QuadInfosBuffer, PF_R32G32_UINT));
-				
-				// We may target an array slice directly : 
-				int32 ArrayIndex = -1;
-				if (OutputTexture->Desc.IsTextureArray())
-				{
-					check(CurrentDestinationArrayIndex < OutputTexture->Desc.ArraySize);
-					ArrayIndex = CurrentDestinationArrayIndex;
-				}
+						FRDGBufferRef QuadInfosBuffer = CreateUploadBuffer(GraphBuilder, TEXT("QuadInfos"), TConstArrayView<FUintVector2>(MakeArrayView(CurrentQuadInfos)));
+						FRDGBufferSRVRef QuadInfosBufferSRV = GraphBuilder.CreateSRV(FRDGBufferSRVDesc(QuadInfosBuffer, PF_R32G32_UINT));
 
-				FCopyQuadsMultiSourcePSParameters* PassParameters = GraphBuilder.AllocParameters<FCopyQuadsMultiSourcePSParameters>();
-				PassParameters->RenderTargets[0] = FRenderTargetBinding(OutputTexture, ERenderTargetLoadAction::ELoad, /*InMipIndex = */0, ArrayIndex);
-				PassParameters->PS.View = View->ViewUniformBuffer;
-				PassParameters->PS.InQuadInfos = QuadInfosBufferSRV;
-				check(CurrentSourceTextureResources.Num() <= FCopyQuadsMultiSourcePS::NumMultiSources);
-				for (int32 TextureIndex = 0; TextureIndex < FCopyQuadsMultiSourcePS::NumMultiSources; ++TextureIndex)
-				{
-					FRDGTextureRef SourceTexture = BlackDummySRVRef->GetParent();
-					if (CurrentSourceTextureResources.IsValidIndex(TextureIndex))
-					{
-						SourceTexture = GraphBuilder.RegisterExternalTexture(CreateRenderTarget(CurrentSourceTextureResources[TextureIndex]->GetTexture2DRHI(), TEXT("SourceTexture")));
-					}
-					PassParameters->PS.InSourceTexture[TextureIndex] = SourceTexture;
-				}
+						// We may target an array slice directly : 
+						int32 ArrayIndex = -1;
+						if (OutputTexture->Desc.IsTextureArray())
+						{
+							check(CurrentDestinationArrayIndex < OutputTexture->Desc.ArraySize);
+							ArrayIndex = CurrentDestinationArrayIndex;
+						}
 
-				FPixelShaderUtils::AddRasterizeToRectsPass<FCopyQuadsMultiSourcePS>(GraphBuilder,
-					ShaderMap,
-					RDG_EVENT_NAME("CopyQuadsMultiSourcePS"),
-					PixelShader,
-					PassParameters,
-					/*ViewportSize = */OutputTexture->Desc.Extent,
-					RectBufferSRV,
-					CurrentDestinationRects.Num(),
-					/*BlendState = */nullptr,
-					/*RasterizerState = */nullptr,
-					/*DepthStencilState = */nullptr,
-					/*StencilRef = */0,
-					/*TextureSize = */FIntPoint(FMath::FloorToInt32(CommonSourceTextureSize.X), FMath::FloorToInt32(CommonSourceTextureSize.Y)),
-					RectUVBufferSRV);
-
-				// Reset all for the next pass that comes along : 
-				CurrentDestinationArrayIndex = INDEX_NONE;
-				CurrentSourceTextureResources.Empty();
-				CurrentSourceRects.Empty();
-				CurrentDestinationRects.Empty();
-				CurrentQuadInfos.Empty();
-			};
+						FCopyQuadsMultiSourcePSParameters* PassParameters = GraphBuilder.AllocParameters<FCopyQuadsMultiSourcePSParameters>();
+						PassParameters->RenderTargets[0] = FRenderTargetBinding(OutputTexture, ERenderTargetLoadAction::ELoad, /*InMipIndex = */0, ArrayIndex);
+						PassParameters->PS.View = View->ViewUniformBuffer;
+						PassParameters->PS.InQuadInfos = QuadInfosBufferSRV;
+						check(CurrentSourceTextureResources.Num() <= FCopyQuadsMultiSourcePS::NumMultiSources);
+						for (int32 TextureIndex = 0; TextureIndex < FCopyQuadsMultiSourcePS::NumMultiSources; ++TextureIndex)
+						{
+							FRDGTextureRef SourceTexture = BlackDummySRVRef->GetParent();
+							if (CurrentSourceTextureResources.IsValidIndex(TextureIndex))
+							{
+								SourceTexture = GraphBuilder.RegisterExternalTexture(CreateRenderTarget(CurrentSourceTextureResources[TextureIndex]->GetTexture2DRHI(), TEXT("SourceTexture")));
+							}
+							PassParameters->PS.InSourceTexture[TextureIndex] = SourceTexture;
+						}
 
-			for (const FCopyQuadParams& InSingleQuadParams : CopyQuadsParams)
-			{
-				check(InSingleQuadParams.DestinationArrayIndex != INDEX_NONE);
-				// If the output texture/slice has changed since the last iteration, it's time to flush and start a new pass :
-				if (CurrentDestinationArrayIndex != InSingleQuadParams.DestinationArrayIndex)
-				{
-					if (CurrentDestinationArrayIndex != INDEX_NONE)
+						FPixelShaderUtils::AddRasterizeToRectsPass<FCopyQuadsMultiSourcePS>(GraphBuilder,
+							ShaderMap,
+							RDG_EVENT_NAME("CopyQuadsMultiSourcePS"),
+							PixelShader,
+							PassParameters,
+							/*ViewportSize = */OutputTexture->Desc.Extent,
+							RectBufferSRV,
+							CurrentDestinationRects.Num(),
+							/*BlendState = */nullptr,
+							/*RasterizerState = */nullptr,
+							/*DepthStencilState = */nullptr,
+							/*StencilRef = */0,
+							/*TextureSize = */FIntPoint(FMath::FloorToInt32(CommonSourceTextureSize.X), FMath::FloorToInt32(CommonSourceTextureSize.Y)),
+							RectUVBufferSRV);
+
+						// Reset all for the next pass that comes along : 
+						CurrentDestinationArrayIndex = INDEX_NONE;
+						CurrentSourceTextureResources.Empty();
+						CurrentSourceRects.Empty();
+						CurrentDestinationRects.Empty();
+						CurrentQuadInfos.Empty();
+					};
+
+				for (const FCopyQuadParams& InSingleQuadParams : CopyQuadsParams)
+				{
+					check(InSingleQuadParams.DestinationArrayIndex != INDEX_NONE);
+					// If the output texture/slice has changed since the last iteration, it's time to flush and start a new pass :
+					if (CurrentDestinationArrayIndex != InSingleQuadParams.DestinationArrayIndex)
 					{
-						FlushCopyQuads();
+						if (CurrentDestinationArrayIndex != INDEX_NONE)
+						{
+							FlushCopyQuads();
+						}
+						CurrentDestinationArrayIndex = InSingleQuadParams.DestinationArrayIndex;
 					}
-					CurrentDestinationArrayIndex = InSingleQuadParams.DestinationArrayIndex;
-				}
 
-				if (CurrentSourceTextureResources.IsEmpty())
-				{
-					CurrentSourceTextureResources.Add(InSingleQuadParams.SourceTextureResource);
-				}
-				else if (InSingleQuadParams.SourceTextureResource != CurrentSourceTextureResources.Last())
-				{
-					// If we've reached the amount of textures we can render in a single pass, we flush the pass and initiate a new one : 
-					if (CurrentSourceTextureResources.Num() == FCopyQuadsMultiSourcePS::NumMultiSources)
+					if (CurrentSourceTextureResources.IsEmpty())
+					{
+						CurrentSourceTextureResources.Add(InSingleQuadParams.SourceTextureResource);
+					}
+					else if (InSingleQuadParams.SourceTextureResource != CurrentSourceTextureResources.Last())
 					{
-						FlushCopyQuads();
-						check(CurrentSourceTextureResources.IsEmpty());
+						// If we've reached the amount of textures we can render in a single pass, we flush the pass and initiate a new one : 
+						if (CurrentSourceTextureResources.Num() == FCopyQuadsMultiSourcePS::NumMultiSources)
+						{
+							FlushCopyQuads();
+							check(CurrentSourceTextureResources.IsEmpty());
 
-						CurrentDestinationArrayIndex = InSingleQuadParams.DestinationArrayIndex;
+							CurrentDestinationArrayIndex = InSingleQuadParams.DestinationArrayIndex;
+						}
+						CurrentSourceTextureResources.Add(InSingleQuadParams.SourceTextureResource);
 					}
-					CurrentSourceTextureResources.Add(InSingleQuadParams.SourceTextureResource);
-				}
 
-				// If we are using the same texture as the last one, we can render it in the same pass, just append our quad : 
-				if (InSingleQuadParams.SourceTextureResource == CurrentSourceTextureResources.Last())
-				{
-					const int32 SourceTextureIndex = CurrentSourceTextureResources.Num() - 1;
-					CurrentQuadInfos.Add(FUintVector2(static_cast<uint32>(SourceTextureIndex), InSingleQuadParams.ChannelSwizzleMask));
-					CurrentSourceRects.Add(FUintVector4(InSingleQuadParams.SourceRect.Min.X, InSingleQuadParams.SourceRect.Min.Y, InSingleQuadParams.SourceRect.Max.X, InSingleQuadParams.SourceRect.Max.Y));
-					CurrentDestinationRects.Add(FUintVector4(InSingleQuadParams.DestinationRect.Min.X, InSingleQuadParams.DestinationRect.Min.Y, InSingleQuadParams.DestinationRect.Max.X, InSingleQuadParams.DestinationRect.Max.Y));
+					// If we are using the same texture as the last one, we can render it in the same pass, just append our quad : 
+					if (InSingleQuadParams.SourceTextureResource == CurrentSourceTextureResources.Last())
+					{
+						const int32 SourceTextureIndex = CurrentSourceTextureResources.Num() - 1;
+						CurrentQuadInfos.Add(FUintVector2(static_cast<uint32>(SourceTextureIndex), InSingleQuadParams.ChannelSwizzleMask));
+						CurrentSourceRects.Add(FUintVector4(InSingleQuadParams.SourceRect.Min.X, InSingleQuadParams.SourceRect.Min.Y, InSingleQuadParams.SourceRect.Max.X, InSingleQuadParams.SourceRect.Max.Y));
+						CurrentDestinationRects.Add(FUintVector4(InSingleQuadParams.DestinationRect.Min.X, InSingleQuadParams.DestinationRect.Min.Y, InSingleQuadParams.DestinationRect.Max.X, InSingleQuadParams.DestinationRect.Max.Y));
+					}
 				}
-			}
 
-			// Flush the remaining quads if any : 
-			FlushCopyQuads();
+				// Flush the remaining quads if any : 
+				FlushCopyQuads();
 			};
 
-			// We need to specify the final state of the external texture to prevent the graph builder from transitioning it to SRVMask :
+		// We need to specify the final state of the external texture to prevent the graph builder from transitioning it to SRVMask :
 		RDGBuilderRecorder.EnqueueRDGCommand(RDGCommand, { { WriteRT->GetRenderTarget()->GetResource(), ERHIAccess::RTV } });
 	}
 
 	// The second step is to perform the blend : merge the edit layer with the previous layers in the stack
 	{
-		RenderParams.MergeRenderContext->CycleBlendRenderTargets(/*InDesiredWriteAccess = */ERHIAccess::RTV, RDGBuilderRecorder);
+		RenderParams.MergeRenderContext->CycleBlendRenderTargets(RDGBuilderRecorder);
 		WriteRT = RenderParams.MergeRenderContext->GetBlendRenderTargetWrite();
 		ULandscapeScratchRenderTarget* CurrentLayerReadRT = RenderParams.MergeRenderContext->GetBlendRenderTargetRead();
 		ULandscapeScratchRenderTarget* PreviousLayersReadRT = RenderParams.MergeRenderContext->GetBlendRenderTargetReadPrevious();
-	 
+
+		WriteRT->TransitionTo(ERHIAccess::RTV, RDGBuilderRecorder);
+		CurrentLayerReadRT->TransitionTo(ERHIAccess::SRVMask, RDGBuilderRecorder);
+		PreviousLayersReadRT->TransitionTo(ERHIAccess::SRVMask, RDGBuilderRecorder);
+
 		if (RenderParams.MergeRenderContext->IsHeightmapMerge())
 		{
 			auto RDGCommand =
-				[ TargetLayerName = RenderParams.TargetLayerGroupLayerNames[0]
+				[TargetLayerName = RenderParams.TargetLayerGroupLayerNames[0]
 				, OutputResource = WriteRT->GetRenderTarget2D()->GetResource()
 				, OutputResourceName = WriteRT->GetDebugName()
 				, CurrentEditLayerResource = CurrentLayerReadRT->GetRenderTarget2D()->GetResource()
 				, PreviousEditLayersResource = PreviousLayersReadRT->GetRenderTarget2D()->GetResource()
 				, EffectiveTextureSize = RenderParams.MergeRenderContext->GetCurrentRenderBatch()->GetRenderTargetResolution(/*bInWithDuplicateBorders = */false)
-				, BlendMode = LayerStruct->BlendMode, HeightmapAlpha = LayerStruct->HeightmapAlpha ]
+				, BlendMode = LayerStruct->BlendMode, HeightmapAlpha = LayerStruct->HeightmapAlpha]
 				(FRDGBuilder& GraphBuilder)
-			{
-				FRDGTextureRef OutputTextureRef = GraphBuilder.RegisterExternalTexture(CreateRenderTarget(OutputResource->TextureRHI, TEXT("OutputTexture")));
-				FRDGTextureRef CurrentEditLayerTextureRef = GraphBuilder.RegisterExternalTexture(CreateRenderTarget(CurrentEditLayerResource->TextureRHI, TEXT("CurrentEditLayerTexture")));
-				FRDGTextureRef PreviousEditLayersTextureRef = GraphBuilder.RegisterExternalTexture(CreateRenderTarget(PreviousEditLayersResource->TextureRHI, TEXT("PreviousEditLayersTexture")));
+				{
+					FRDGTextureRef OutputTextureRef = GraphBuilder.RegisterExternalTexture(CreateRenderTarget(OutputResource->TextureRHI, TEXT("OutputTexture")));
+					FRDGTextureRef CurrentEditLayerTextureRef = GraphBuilder.RegisterExternalTexture(CreateRenderTarget(CurrentEditLayerResource->TextureRHI, TEXT("CurrentEditLayerTexture")));
+					FRDGTextureRef PreviousEditLayersTextureRef = GraphBuilder.RegisterExternalTexture(CreateRenderTarget(PreviousEditLayersResource->TextureRHI, TEXT("PreviousEditLayersTexture")));
 
-				FLandscapeEditLayersHeightmapsMergeEditLayerPS::FParameters* PSParams = GraphBuilder.AllocParameters<FLandscapeEditLayersHeightmapsMergeEditLayerPS::FParameters>();
-				PSParams->RenderTargets[0] = FRenderTargetBinding(OutputTextureRef, ERenderTargetLoadAction::ENoAction);
-				PSParams->InEditLayerBlendMode = static_cast<uint32>(LandscapeBlendModeToEditLayerBlendMode(BlendMode));
-				PSParams->InEditLayerAlpha = HeightmapAlpha;
-				PSParams->InCurrentEditLayerHeightmap = GraphBuilder.CreateSRV(FRDGTextureSRVDesc::Create(CurrentEditLayerTextureRef));
-				PSParams->InPreviousEditLayersHeightmap = GraphBuilder.CreateSRV(FRDGTextureSRVDesc::Create(PreviousEditLayersTextureRef));
+					FLandscapeEditLayersHeightmapsMergeEditLayerPS::FParameters* PSParams = GraphBuilder.AllocParameters<FLandscapeEditLayersHeightmapsMergeEditLayerPS::FParameters>();
+					PSParams->RenderTargets[0] = FRenderTargetBinding(OutputTextureRef, ERenderTargetLoadAction::ENoAction);
+					PSParams->InEditLayerBlendMode = static_cast<uint32>(LandscapeBlendModeToEditLayerBlendMode(BlendMode));
+					PSParams->InEditLayerAlpha = HeightmapAlpha;
+					PSParams->InCurrentEditLayerHeightmap = GraphBuilder.CreateSRV(FRDGTextureSRVDesc::Create(CurrentEditLayerTextureRef));
+					PSParams->InPreviousEditLayersHeightmap = GraphBuilder.CreateSRV(FRDGTextureSRVDesc::Create(PreviousEditLayersTextureRef));
 
 					FLandscapeEditLayersHeightmapsMergeEditLayerPS::MergeEditLayerPS(
 						RDG_EVENT_NAME("MergeEditLayer(%s) -> %s", *TargetLayerName.ToString(), *OutputResourceName), GraphBuilder, PSParams, EffectiveTextureSize);
@@ -5909,34 +5914,34 @@ void ULandscapeEditLayerPersistent::RenderLayer(FRenderParams& RenderParams, UE:
 			}
 
 			auto RDGCommand =
-				[ MergePaintLayersParams
+				[MergePaintLayersParams
 				, OutputResource = WriteRT->GetRenderTarget2DArray()->GetResource()
 				, OutputResourceName = WriteRT->GetDebugName()
 				, CurrentEditLayerResource = CurrentLayerReadRT->GetRenderTarget2DArray()->GetResource()
 				, PreviousEditLayersResource = PreviousLayersReadRT->GetRenderTarget2DArray()->GetResource()
 				, EffectiveTextureSize = RenderParams.MergeRenderContext->GetCurrentRenderBatch()->GetRenderTargetResolution(/*bInWithDuplicateBorders = */false)]
 				(FRDGBuilder& GraphBuilder)
-			{
-				FRDGTextureRef OutputTextureRef = GraphBuilder.RegisterExternalTexture(CreateRenderTarget(OutputResource->TextureRHI, TEXT("OutputTexture")));
-				FRDGTextureRef CurrentEditLayerTextureRef = GraphBuilder.RegisterExternalTexture(CreateRenderTarget(CurrentEditLayerResource->TextureRHI, TEXT("CurrentEditLayerTexture")));
-				FRDGTextureSRVRef CurrentEditLayerTextureSRVRef = GraphBuilder.CreateSRV(FRDGTextureSRVDesc::Create(CurrentEditLayerTextureRef));
-				FRDGTextureRef PreviousEditLayersTextureRef = GraphBuilder.RegisterExternalTexture(CreateRenderTarget(PreviousEditLayersResource->TextureRHI, TEXT("PreviousEditLayersTexture")));
-				FRDGTextureSRVRef PreviousEditLayersTextureSRVRef = GraphBuilder.CreateSRV(FRDGTextureSRVDesc::Create(PreviousEditLayersTextureRef));
-
-				const int32 NumTargetLayers = MergePaintLayersParams.Num();
-				for (int32 TargetLayerIndex = 0; TargetLayerIndex < NumTargetLayers; ++TargetLayerIndex)
 				{
-					FLandscapeEditLayersWeightmapsMergeEditLayerPS::FParameters* PSParams = GraphBuilder.AllocParameters<FLandscapeEditLayersWeightmapsMergeEditLayerPS::FParameters>();
-					PSParams->RenderTargets[0] = FRenderTargetBinding(OutputTextureRef, ERenderTargetLoadAction::ENoAction, /*InMipIndex = */0, /*InArraySlice = */TargetLayerIndex);
-					PSParams->InPaintLayerIndex = TargetLayerIndex;
-					PSParams->InEditLayerPaintLayerBlendMode = static_cast<uint32>(MergePaintLayersParams[TargetLayerIndex].EditLayerPaintLayerBlendMode);
-					PSParams->InEditLayerAlpha = MergePaintLayersParams[TargetLayerIndex].EditLayerAlpha;
-					PSParams->InCurrentEditLayerWeightmaps = CurrentEditLayerTextureSRVRef;
-					PSParams->InPreviousEditLayersWeightmaps = PreviousEditLayersTextureSRVRef;
+					FRDGTextureRef OutputTextureRef = GraphBuilder.RegisterExternalTexture(CreateRenderTarget(OutputResource->TextureRHI, TEXT("OutputTexture")));
+					FRDGTextureRef CurrentEditLayerTextureRef = GraphBuilder.RegisterExternalTexture(CreateRenderTarget(CurrentEditLayerResource->TextureRHI, TEXT("CurrentEditLayerTexture")));
+					FRDGTextureSRVRef CurrentEditLayerTextureSRVRef = GraphBuilder.CreateSRV(FRDGTextureSRVDesc::Create(CurrentEditLayerTextureRef));
+					FRDGTextureRef PreviousEditLayersTextureRef = GraphBuilder.RegisterExternalTexture(CreateRenderTarget(PreviousEditLayersResource->TextureRHI, TEXT("PreviousEditLayersTexture")));
+					FRDGTextureSRVRef PreviousEditLayersTextureSRVRef = GraphBuilder.CreateSRV(FRDGTextureSRVDesc::Create(PreviousEditLayersTextureRef));
+
+					const int32 NumTargetLayers = MergePaintLayersParams.Num();
+					for (int32 TargetLayerIndex = 0; TargetLayerIndex < NumTargetLayers; ++TargetLayerIndex)
+					{
+						FLandscapeEditLayersWeightmapsMergeEditLayerPS::FParameters* PSParams = GraphBuilder.AllocParameters<FLandscapeEditLayersWeightmapsMergeEditLayerPS::FParameters>();
+						PSParams->RenderTargets[0] = FRenderTargetBinding(OutputTextureRef, ERenderTargetLoadAction::ENoAction, /*InMipIndex = */0, /*InArraySlice = */TargetLayerIndex);
+						PSParams->InPaintLayerIndex = TargetLayerIndex;
+						PSParams->InEditLayerPaintLayerBlendMode = static_cast<uint32>(MergePaintLayersParams[TargetLayerIndex].EditLayerPaintLayerBlendMode);
+						PSParams->InEditLayerAlpha = MergePaintLayersParams[TargetLayerIndex].EditLayerAlpha;
+						PSParams->InCurrentEditLayerWeightmaps = CurrentEditLayerTextureSRVRef;
+						PSParams->InPreviousEditLayersWeightmaps = PreviousEditLayersTextureSRVRef;
 
 						FLandscapeEditLayersWeightmapsMergeEditLayerPS::MergeEditLayerPS(
 							RDG_EVENT_NAME("MergeEditLayer(%s) -> %s", *MergePaintLayersParams[TargetLayerIndex].PaintLayerName.ToString(), *OutputResourceName), GraphBuilder, PSParams, EffectiveTextureSize);
-				}
+					}
 				};
 
 			// We need to specify the final state of the external textures to prevent the graph builder from transitioning them to SRVMask (even those that end up as SRVMask at the end of this command,
@@ -6439,12 +6444,14 @@ void ULandscapeHeightmapNormalsEditLayerRenderer::RenderLayer(FRenderParams& Ren
 	checkf(RDGBuilderRecorder.IsRecording(), TEXT("ERenderFlags::RenderMode_Recorded means the command recorder should be recording at this point"));
 	checkf(RenderParams.TargetLayerGroupLayerNames.Num() == 1, TEXT("Normals should only be generated on heightmap merge, which should have 1 and only target layer"));
 
-	RenderContext->CycleBlendRenderTargets(/*InDesiredWriteAccess = */ERHIAccess::RTV, RDGBuilderRecorder);
+	RenderContext->CycleBlendRenderTargets(RDGBuilderRecorder);
 	ULandscapeScratchRenderTarget* WriteRT = RenderContext->GetBlendRenderTargetWrite();
 	ULandscapeScratchRenderTarget* ReadRT = RenderContext->GetBlendRenderTargetRead();
 	ULandscapeScratchRenderTarget* ValidityRT = RenderContext->GetValidityRenderTarget(RenderParams.TargetLayerGroupLayerNames[0]);
 	check((WriteRT != nullptr) && (ReadRT != nullptr) && (ValidityRT != nullptr));
 
+	WriteRT->TransitionTo(ERHIAccess::RTV, RDGBuilderRecorder);
+	ReadRT->TransitionTo(ERHIAccess::SRVMask, RDGBuilderRecorder);
 	ValidityRT->TransitionTo(ERHIAccess::SRVMask, RDGBuilderRecorder);
 
 	FIntPoint NumComponentsInRenderArea = RenderBatch->MaxComponentKey - RenderBatch->MinComponentKey + 1;
@@ -6498,10 +6505,13 @@ void ULandscapeWeightmapWeightBlendedLayersRenderer::RenderLayer(FRenderParams&
 	checkf(RDGBuilderRecorder.IsRecording(), TEXT("ERenderFlags::RenderMode_Recorded means the command recorder should be recording at this point"));
 	check(!RenderParams.MergeRenderContext->IsHeightmapMerge());
 
-	RenderParams.MergeRenderContext->CycleBlendRenderTargets(ERHIAccess::RTV, RDGBuilderRecorder);
+	RenderParams.MergeRenderContext->CycleBlendRenderTargets(RDGBuilderRecorder);
 	ULandscapeScratchRenderTarget* WriteRT = RenderParams.MergeRenderContext->GetBlendRenderTargetWrite();
 	ULandscapeScratchRenderTarget* ReadRT = RenderParams.MergeRenderContext->GetBlendRenderTargetRead();
 
+	WriteRT->TransitionTo(ERHIAccess::RTV, RDGBuilderRecorder);
+	ReadRT->TransitionTo(ERHIAccess::SRVMask, RDGBuilderRecorder);
+
 	TArray<FWeightmapPaintLayerInfo> WeightmapPaintLayerInfos;
 	WeightmapPaintLayerInfos.Reserve(RenderParams.TargetLayerGroupLayerInfos.Num());
 	Algo::Transform(RenderParams.TargetLayerGroupLayerInfos, WeightmapPaintLayerInfos, [](ULandscapeLayerInfoObject* InLayerInfo)
@@ -6520,34 +6530,34 @@ void ULandscapeWeightmapWeightBlendedLayersRenderer::RenderLayer(FRenderParams&
 		});
 
 	auto RDGCommand =
-		[ TargetLayerNames = RenderParams.TargetLayerGroupLayerNames
+		[TargetLayerNames = RenderParams.TargetLayerGroupLayerNames
 		, WeightmapPaintLayerInfos
 		, OutputResource = WriteRT->GetRenderTarget2DArray()->GetResource()
 		, OutputResourceName = WriteRT->GetDebugName()
 		, CurrentEditLayerResource = ReadRT->GetRenderTarget2DArray()->GetResource()
 		, EffectiveTextureSize = RenderParams.MergeRenderContext->GetCurrentRenderBatch()->GetRenderTargetResolution(/*bInWithDuplicateBorders = */false)](FRDGBuilder& GraphBuilder)
-	{
+		{
 			FRDGTextureRef OutputTextureRef = GraphBuilder.RegisterExternalTexture(CreateRenderTarget(OutputResource->TextureRHI, TEXT("OutputTexture")));
-		FRDGTextureRef CurrentEditLayerTextureRef = GraphBuilder.RegisterExternalTexture(CreateRenderTarget(CurrentEditLayerResource->TextureRHI, TEXT("CurrentEditLayerTexture")));
-		FRDGTextureSRVRef CurrentEditLayerTextureSRVRef = GraphBuilder.CreateSRV(FRDGTextureSRVDesc::Create(CurrentEditLayerTextureRef));
+			FRDGTextureRef CurrentEditLayerTextureRef = GraphBuilder.RegisterExternalTexture(CreateRenderTarget(CurrentEditLayerResource->TextureRHI, TEXT("CurrentEditLayerTexture")));
+			FRDGTextureSRVRef CurrentEditLayerTextureSRVRef = GraphBuilder.CreateSRV(FRDGTextureSRVDesc::Create(CurrentEditLayerTextureRef));
 
-		FRDGBufferRef PaintLayerInfosBuffer = CreateStructuredBuffer(GraphBuilder, TEXT("LandscapePaintLayerInfosBuffer"), WeightmapPaintLayerInfos);
-		FRDGBufferSRVRef PaintLayerInfosBufferSRV = GraphBuilder.CreateSRV(FRDGBufferSRVDesc(PaintLayerInfosBuffer));
+			FRDGBufferRef PaintLayerInfosBuffer = CreateStructuredBuffer(GraphBuilder, TEXT("LandscapePaintLayerInfosBuffer"), WeightmapPaintLayerInfos);
+			FRDGBufferSRVRef PaintLayerInfosBufferSRV = GraphBuilder.CreateSRV(FRDGBufferSRVDesc(PaintLayerInfosBuffer));
 
-		const int32 NumTargetLayers = TargetLayerNames.Num();
-		for (int32 TargetLayerIndex = 0; TargetLayerIndex < NumTargetLayers; ++TargetLayerIndex)
-		{
-			FLandscapeEditLayersWeightmapsPerformLegacyWeightBlendingPS::FParameters* PSParams = GraphBuilder.AllocParameters<FLandscapeEditLayersWeightmapsPerformLegacyWeightBlendingPS::FParameters>();
-			PSParams->RenderTargets[0] = FRenderTargetBinding(OutputTextureRef, ERenderTargetLoadAction::ENoAction, /*InMipIndex = */0, /*InArraySlice = */TargetLayerIndex);
-			PSParams->InPaintLayerIndex = TargetLayerIndex;
-			PSParams->InNumPaintLayers = NumTargetLayers;
-			PSParams->InPaintLayerInfos = PaintLayerInfosBufferSRV;
-			PSParams->InCurrentEditLayerWeightmaps = CurrentEditLayerTextureSRVRef;
+			const int32 NumTargetLayers = TargetLayerNames.Num();
+			for (int32 TargetLayerIndex = 0; TargetLayerIndex < NumTargetLayers; ++TargetLayerIndex)
+			{
+				FLandscapeEditLayersWeightmapsPerformLegacyWeightBlendingPS::FParameters* PSParams = GraphBuilder.AllocParameters<FLandscapeEditLayersWeightmapsPerformLegacyWeightBlendingPS::FParameters>();
+				PSParams->RenderTargets[0] = FRenderTargetBinding(OutputTextureRef, ERenderTargetLoadAction::ENoAction, /*InMipIndex = */0, /*InArraySlice = */TargetLayerIndex);
+				PSParams->InPaintLayerIndex = TargetLayerIndex;
+				PSParams->InNumPaintLayers = NumTargetLayers;
+				PSParams->InPaintLayerInfos = PaintLayerInfosBufferSRV;
+				PSParams->InCurrentEditLayerWeightmaps = CurrentEditLayerTextureSRVRef;
 
 				FLandscapeEditLayersWeightmapsPerformLegacyWeightBlendingPS::PerformLegacyWeightBlendingPS(
 					RDG_EVENT_NAME("LegacyWeightBlending(%s) -> %s", *TargetLayerNames[TargetLayerIndex].ToString(), *OutputResourceName),
 					GraphBuilder, PSParams, EffectiveTextureSize);
-		}
+			}
 		};
 
 	// We need to specify the final state of the external textures to prevent the graph builder from transitioning them to SRVMask (even those that end up as SRVMask at the end of this command,
@@ -7573,8 +7583,10 @@ int32 ALandscape::PerformLayersHeightmapsBatchedMerge(const FUpdateLayersContent
 			RHI_BREADCRUMB_EVENT_GAMETHREAD("CopyMip0AndExpand");
 
 			// Recompose mip0 of the final heightmaps, subsection by subsection, to duplicate borders : 
-			InParams.MergeRenderContext->CycleBlendRenderTargets(/*InDesiredWriteAccess = */ERHIAccess::SRVMask, RDGBuilderRecorder); // TODO [jonathan.bard] : This should be CopyDst but ExecuteCopyLayersTexture doesn't allow for it ATM
+			InParams.MergeRenderContext->CycleBlendRenderTargets(RDGBuilderRecorder); 
 			ULandscapeScratchRenderTarget* ReadRT = InParams.MergeRenderContext->GetBlendRenderTargetRead();
+			
+			ReadRT->TransitionTo(ERHIAccess::SRVMask, RDGBuilderRecorder); // TODO [jonathan.bard] : This should be CopyDst but ExecuteCopyLayersTexture doesn't allow for it ATM
 
 			OnEditLayersMergedDelegate->Broadcast(FOnLandscapeEditLayersMergedParams(ReadRT->GetRenderTarget(),
 				/*InRenderAreaResolution = */InParams.RenderBatch->GetRenderTargetResolution(/*bInWithDuplicateBorders = */false), /*bInIsHeightmapMerge = */true));
@@ -7646,9 +7658,13 @@ int32 ALandscape::PerformLayersHeightmapsBatchedMerge(const FUpdateLayersContent
 			FIntPoint CurrentMipSubsectionSize = Landscape->SubsectionSizeQuads + 1;
 			for (int32 MipIndex = 1; MipIndex < NumMips; ++MipIndex)
 			{
-				InParams.MergeRenderContext->CycleBlendRenderTargets(/*InDesiredWriteAccess = */ERHIAccess::RTV, RDGBuilderRecorder);
+				InParams.MergeRenderContext->CycleBlendRenderTargets(RDGBuilderRecorder);
 				ULandscapeScratchRenderTarget* WriteRT = InParams.MergeRenderContext->GetBlendRenderTargetWrite();
 				ULandscapeScratchRenderTarget* ReadRT = InParams.MergeRenderContext->GetBlendRenderTargetRead();
+
+				WriteRT->TransitionTo(ERHIAccess::RTV, RDGBuilderRecorder);
+				ReadRT->TransitionTo(ERHIAccess::SRVMask, RDGBuilderRecorder);
+
 				CurrentMipResolution /= 2;
 				check(CurrentMipResolution.X > 0 && CurrentMipResolution.Y > 0);
 				CurrentMipSubsectionSize /= 2;
@@ -10447,9 +10463,11 @@ int32 ALandscape::PerformLayersWeightmapsBatchedMerge(FUpdateLayersContentContex
 		ALandscape* Landscape = InParams.MergeRenderContext->GetLandscape();
 
 		// Cycle render targets one last time so we can read from the last RT we've written to :
-			InParams.MergeRenderContext->CycleBlendRenderTargets(/*InDesiredWriteAccess = */ERHIAccess::None, RDGBuilderRecorder);
+		InParams.MergeRenderContext->CycleBlendRenderTargets(RDGBuilderRecorder);
 		ULandscapeScratchRenderTarget* ReadRT = InParams.MergeRenderContext->GetBlendRenderTargetRead();
 
+		ReadRT->TransitionTo(ERHIAccess::SRVMask, RDGBuilderRecorder);
+
 		OnEditLayersMergedDelegate->Broadcast(FOnLandscapeEditLayersMergedParams(ReadRT->GetRenderTarget(),
 			/*InRenderAreaResolution = */InParams.RenderBatch->GetRenderTargetResolution(/*bInWithDuplicateBorders = */false), /*bInIsHeightmapMerge = */false));
 
diff --git a/Engine/Source/Runtime/Landscape/Public/LandscapeEditLayerRenderer.h b/Engine/Source/Runtime/Landscape/Public/LandscapeEditLayerRenderer.h
--- a/Engine/Source/Runtime/Landscape/Public/LandscapeEditLayerRenderer.h
+++ b/Engine/Source/Runtime/Landscape/Public/LandscapeEditLayerRenderer.h
@@ -809,7 +809,7 @@ public:
 	 * @param RDGBuilderRecorder recorder to append operations to a single FRDGBuilder, if the the recorder in "recording" mode. In "immediate" mode, a render command will be enqueued immediately
 	 */
 	LANDSCAPE_EDIT_LAYERS_BATCHED_MERGE_EXPERIMENTAL
-	LANDSCAPE_API void CycleBlendRenderTargets(ERHIAccess InDesiredWriteAccess, UE::Landscape::FRDGBuilderRecorder& RDGBuilderRecorder);
+	LANDSCAPE_API void CycleBlendRenderTargets(UE::Landscape::FRDGBuilderRecorder& RDGBuilderRecorder);
 
 	LANDSCAPE_EDIT_LAYERS_BATCHED_MERGE_EXPERIMENTAL
 	LANDSCAPE_API ULandscapeScratchRenderTarget* GetBlendRenderTargetWrite() const;
