Fix landscape splines always dirtying the proxies on load when they affect the weightmaps.

diff --git a/Engine/Source/Runtime/Landscape/Private/LandscapeEdit.cpp b/Engine/Source/Runtime/Landscape/Private/LandscapeEdit.cpp
--- a/Engine/Source/Runtime/Landscape/Private/LandscapeEdit.cpp
+++ b/Engine/Source/Runtime/Landscape/Private/LandscapeEdit.cpp
@@ -7458,6 +7458,10 @@ TArray<UTexture*> ULandscapeComponent::ReallocateWeightmaps(FLandscapeEditDataIn
 		UTexture2D* CurrentWeightmapTexture = nullptr;
 		ULandscapeWeightmapUsage* CurrentWeightmapUsage = nullptr;
 
+		// Testing TotalNeededChannels here is not correct.  There's no apparent reason to do it and it prevents direct re-use of available channels in
+		// WeightmapUsageMap if more than one texture is needed.  In simple cases, this works if this TotalNeededChannels test (and the one ten lines down)
+		// are removed. But... this bug is at least 14 years old in this very central function.  Changing this has some risk and it solves no known problems.
+		// For now, this will be left as is.
 		if (TotalNeededChannels < 4)
 		{
 			// UE_LOG(LogLandscape, Log, TEXT("Looking for nearest"));
@@ -7518,7 +7522,10 @@ TArray<UTexture*> ULandscapeComponent::ReallocateWeightmaps(FLandscapeEditDataIn
 		// No suitable weightmap texture
 		if (CurrentWeightmapTexture == nullptr)
 		{
-			LandscapeInfo->MarkObjectDirty(this);
+			if (bMarkPackageDirty)
+			{
+				LandscapeInfo->MarkObjectDirty(this);
+			}
 
 			// Weightmap is sized the same as the component
 			int32 WeightmapSize = (SubsectionSizeQuads + 1) * NumSubsections;
