Added separate blend step for landscape edit layer renderers. This moves the blend step into a separate, standardized blend step, which allows to apply blend methods in a way that works across renderers

Details :
* All heightmap and weightmap renderers can now return their height/weight value (16 bits) as well as an alpha value (14 bits) and alpha flags (2 bits)
* Edit layer renderers can now implement one of the following functions :
** BlendLayer, called if ERenderFlags::BlendMode_SeparateBlend is returned by GetRenderFlags. In this case, each RenderLayer is followed by a BlendLayer.
** CanGroupRenderLayerWith, called if ERenderFlags::RenderLayerGroup_SupportsGrouping is returned by GetRenderFlags : indicates that the renderer's RenderLayer can be followed by another one (instead of RenderLayer/BlendLayer/RenderLayer/...)
** BeginRenderLayerGroup / EndRenderLayerGroup, also called if ERenderFlags::RenderLayerGroup_SupportsGrouping is returned by GetRenderFlags : called at the beginning end of a series of RenderLayer calls (i.e. BeginRenderLayerGroup/RenderLayer/RenderLayer/.../ EndRenderLayerGroup/BlendLayer) : this allows to coalesce several successive renders
* RenderLayer now returns a value to indicate whether anything was actually rendered (avoids calling BlendLayer if nothing was actually rendered by a renderer)
* FMergeRenderContext now has a GenericBlendLayer method that can be used by any renderer (in a BlendLayer render step)
* UE::Landscape::EditLayers::EHeightmapBlendMode/EWeightmapBlendMode/FWeightmapPaintLayerInfo are for batched merge while EEditLayerHeightmapBlendMode/EEditLayerWeightmapBlendMode/FLandscapeEditLayerWeightmapPaintLayerInfo are now deprecated (only used by local merge)
* FMergeRenderContext now has a blackboard system where anyone at any point can add/remove blackboard items, in order to let successive renderers communicate information (unused at the moment, but the goal is to let render groups batch succesive draws into a single one)

diff --git a/Engine/Plugins/Experimental/LandscapePatch/Shaders/Private/LandscapeCircleHeightPatchPS.usf b/Engine/Plugins/Experimental/LandscapePatch/Shaders/Private/LandscapeCircleHeightPatchPS.usf
--- a/Engine/Plugins/Experimental/LandscapePatch/Shaders/Private/LandscapeCircleHeightPatchPS.usf
+++ b/Engine/Plugins/Experimental/LandscapePatch/Shaders/Private/LandscapeCircleHeightPatchPS.usf
@@ -8,6 +8,7 @@
 float3 InCenter;
 float InRadius;
 float InFalloff;
+
 Texture2D<float4> InSourceTexture;
 uint2 InSourceTextureOffset;
 
@@ -19,9 +20,12 @@ void ApplyLandscapeCircleHeightPatch(in float4 SVPos : SV_POSITION, out float4 O
 	const float ClampedFalloff = max(InFalloff, KINDA_SMALL_NUMBER);
 	
 	int2 DestinationTextureCoordinates = floor(SVPos.xy);
+
+#if PERFORM_BLENDING
 	int2 SourceTextureCoordinates = DestinationTextureCoordinates - InSourceTextureOffset;
 	float4 CurrentPackedHeight = InSourceTexture.Load(int3(SourceTextureCoordinates, 0));
 	float CurrentHeight = UnpackHeight(CurrentPackedHeight.xy);
+#endif // PERFORM_BLENDING
 	
 	float PatchHeight = InCenter.z;
 	
@@ -33,21 +37,29 @@ void ApplyLandscapeCircleHeightPatch(in float4 SVPos : SV_POSITION, out float4 O
 		Alpha = 1 - min((Distance - InRadius) / ClampedFalloff, 1.0);
 	}
 	
+// TODO [jonathan.bard] Deprecate :
+#if PERFORM_BLENDING
 	float NewHeight = lerp(CurrentHeight, PatchHeight, Alpha);
-	
 	OutColor = float4(PackHeight(NewHeight), 0.0, 0.0);
+#else // PERFORM_BLENDING
+
+	OutColor = float4(PackHeight(PatchHeight), PackHeightAlpha(Alpha, EHEIGHTMAPALPHAFLAGS_ALPHABLEND));
+
+#endif // !PERFORM_BLENDING
 }
 
 #endif // CIRCLE_HEIGHT_PATCH
 
 #if CIRCLE_VISIBILITY_PATCH
 
-void ApplyLandscapeCircleVisibilityPatch(in float4 SVPos : SV_POSITION, out float OutColor : SV_Target0)
+void ApplyLandscapeCircleVisibilityPatch(in float4 SVPos : SV_POSITION, out float4 OutColor : SV_Target0)
 {
 	int2 DestinationTextureCoordinates = floor(SVPos.xy);
-	int2 SourceTextureCoordinates = DestinationTextureCoordinates - InSourceTextureOffset;
 	float Distance = distance(DestinationTextureCoordinates, InCenter.xy);
 	
+// TODO [jonathan.bard] Deprecate :
+#if PERFORM_BLENDING
+	int2 SourceTextureCoordinates = DestinationTextureCoordinates - InSourceTextureOffset;
 	float4 CurrentWeight = InSourceTexture.Load(int3(SourceTextureCoordinates, 0)).x;
 
 	if (Distance <= InRadius)
@@ -55,12 +67,17 @@ void ApplyLandscapeCircleVisibilityPatch(in float4 SVPos : SV_POSITION, out floa
 		// TODO: This should be changed to calculate alpha just like a height patch and write out
 		//  the alpha value, which allows for smoother holes. But needs to be changed in a
 		//  separate CL with a fixup to set the falloff of older assets to 0 to keep them unchanged.
-		OutColor = 255.f;
+		OutColor = 1.0;
 	}
 	else
 	{
 		OutColor = CurrentWeight.x;
 	}
+#else // PERFORM_BLENDING
+
+	OutColor = (Distance <= InRadius) ? 1.0f : 0.0f;
+
+#endif // !PERFORM_BLENDING
 }
 
 #endif // CIRCLE_VISIBILITY_PATCH
diff --git a/Engine/Plugins/Experimental/LandscapePatch/Shaders/Private/LandscapeTexturePatchPS.usf b/Engine/Plugins/Experimental/LandscapePatch/Shaders/Private/LandscapeTexturePatchPS.usf
--- a/Engine/Plugins/Experimental/LandscapePatch/Shaders/Private/LandscapeTexturePatchPS.usf
+++ b/Engine/Plugins/Experimental/LandscapePatch/Shaders/Private/LandscapeTexturePatchPS.usf
@@ -16,8 +16,6 @@
 //#define REINITIALIZE_PATCH 1
 #endif // defined (__INTELLISENSE__)
 
-// This constant gets used in multiple shaders, hence it's out here.
-static const float LANDSCAPE_MID_VALUE = 32768.0f;
 
 #if APPLY_HEIGHT_PATCH || APPLY_WEIGHT_PATCH
 
@@ -102,6 +100,23 @@ uint InBlendMode;
 // A combination of flags, whose positions are set in the corresponding cpp file.
 uint InFlags;
 
+uint ConvertPatchBlendModeToHeightAlphaFlags()
+{
+	switch (InBlendMode)
+	{
+		case ADDITIVE_MODE:
+			return EHEIGHTMAPALPHAFLAGS_ADDITIVE;
+		case ALPHA_BLEND_MODE:
+			return (EHEIGHTMAPALPHAFLAGS_MIN | EHEIGHTMAPALPHAFLAGS_MAX);
+		case MIN_MODE:
+			return EHEIGHTMAPALPHAFLAGS_MIN;
+		case MAX_MODE:
+			return EHEIGHTMAPALPHAFLAGS_MAX;
+		default:
+			return EHEIGHTMAPALPHAFLAGS_NONE;
+	}
+}
+
 void ApplyLandscapeTextureHeightPatch(in float4 SVPos : SV_POSITION, out float4 OutColor : SV_Target0)
 {
 	// Unpack our flags. The flag choices are set in FApplyLandscapeTextureHeightPatchPS::ModifyCompilationEnvironment
@@ -121,9 +136,11 @@ void ApplyLandscapeTextureHeightPatch(in float4 SVPos : SV_POSITION, out float4
 	
 	float PatchSignedHeight = InHeightScale * (PatchStoredHeight - InZeroInEncoding) + InHeightOffset;
 	
+#if PERFORM_BLENDING
 	int2 HeightmapCoordinates = floor(SVPos.xy) - InSourceHeightmapOffset;
 	float4 CurrentPackedHeight = InSourceHeightmap.Load(int3(HeightmapCoordinates, 0));
 	float CurrentHeight = UnpackHeight(CurrentPackedHeight.xy);
+#endif // PERFORM_BLENDING
 	
 	float Alpha = GetFalloffAlpha(InFalloffWorldMargin, InPatchWorldDimensions, PatchUVCoordinates, InEdgeUVDeadBorder, bRectangularFalloff);
 	
@@ -132,6 +149,8 @@ void ApplyLandscapeTextureHeightPatch(in float4 SVPos : SV_POSITION, out float4
 		Alpha *= PatchSampledValue.a;
 	}
 	
+	// TODO [jonathan.bard] Deprecate :
+#if PERFORM_BLENDING
 	float NewHeight = 0;
 	switch (InBlendMode)
 	{
@@ -149,7 +168,14 @@ void ApplyLandscapeTextureHeightPatch(in float4 SVPos : SV_POSITION, out float4
 			break;
 	}
 	
-	OutColor = float4(PackHeight(NewHeight), CurrentPackedHeight.z, CurrentPackedHeight.w);
+	OutColor = float4(PackHeight(NewHeight), PackHeightAlpha(1.0f, EHEIGHTMAPALPHAFLAGS_NONE));
+
+#else // PERFORM_BLENDING
+
+	float OutputHeight = (InBlendMode == ADDITIVE_MODE) ? PatchSignedHeight : (LANDSCAPE_MID_VALUE + PatchSignedHeight);
+	OutColor = float4(PackHeight(OutputHeight), PackHeightAlpha(Alpha, ConvertPatchBlendModeToHeightAlphaFlags()));
+
+#endif // !PERFORM_BLENDING
 }
 #endif // APPLY_HEIGHT_PATCH
 
@@ -212,8 +238,24 @@ uint InBlendMode;
 // A combination of flags, whose positions are set in the corresponding cpp file.
 uint InFlags;
 
+uint ConvertPatchBlendModeToWeightAlphaFlags()
+{
+	switch (InBlendMode)
+	{
+		case ADDITIVE_MODE:
+			return EWEIGHTMAPALPHAFLAGS_ADDITIVE;
+		case ALPHA_BLEND_MODE:
+			return (EWEIGHTMAPALPHAFLAGS_MIN | EWEIGHTMAPALPHAFLAGS_MAX);
+		case MIN_MODE:
+			return EWEIGHTMAPALPHAFLAGS_MIN;
+		case MAX_MODE:
+			return EWEIGHTMAPALPHAFLAGS_MAX;
+		default:
+			return EWEIGHTMAPALPHAFLAGS_NONE;
+	}
+}
 
-void ApplyLandscapeTextureWeightPatch(in float4 SVPos : SV_POSITION, out float OutColor : SV_Target0)
+void ApplyLandscapeTextureWeightPatch(in float4 SVPos : SV_POSITION, out float4 OutColor : SV_Target0)
 {
 	// Unpack our flags. The flag choices are set in FApplyLandscapeTextureWeightPatchPS::ModifyCompilationEnvironment
 	// so that they match to the cpp file
@@ -229,8 +271,10 @@ void ApplyLandscapeTextureWeightPatch(in float4 SVPos : SV_POSITION, out float O
 	float4 PatchSampledValue = InWeightPatch.Sample(InWeightPatchSampler, PatchUVCoordinates);
 	float PatchWeight = PatchSampledValue.x;
 	
+#if PERFORM_BLENDING
 	int2 SourceWeightmapCoordinates = floor(SVPos.xy) - InSourceWeightmapCoordOffset;
 	float CurrentWeight = InSourceWeightmap.Load(int3(SourceWeightmapCoordinates, 0)).x;
+#endif // PERFORM_BLENDING
 	
 	float Alpha = GetFalloffAlpha(InFalloffWorldMargin, InPatchWorldDimensions, PatchUVCoordinates, InEdgeUVDeadBorder, bRectangularFalloff);
 	
@@ -239,6 +283,8 @@ void ApplyLandscapeTextureWeightPatch(in float4 SVPos : SV_POSITION, out float O
 		Alpha *= PatchSampledValue.a;
 	}
 	
+	// TODO [jonathan.bard] Deprecate :
+#if PERFORM_BLENDING
 	float NewWeight = 0;
 	switch (InBlendMode)
 	{
@@ -256,7 +302,12 @@ void ApplyLandscapeTextureWeightPatch(in float4 SVPos : SV_POSITION, out float O
 			break;
 	}
 	
-	OutColor = NewWeight;
+	OutColor = float4(PackWeight(NewWeight), PackWeightAlpha(1.0f, EWEIGHTMAPALPHAFLAGS_NONE));
+#else // PERFORM_BLENDING
+
+	OutColor = float4(PackWeight(PatchWeight), PackWeightAlpha(Alpha, ConvertPatchBlendModeToWeightAlphaFlags()));
+
+#endif // !PERFORM_BLENDING
 }
 
 
diff --git a/Engine/Plugins/Experimental/LandscapePatch/Source/LandscapePatch/LandscapePatch.Build.cs b/Engine/Plugins/Experimental/LandscapePatch/Source/LandscapePatch/LandscapePatch.Build.cs
--- a/Engine/Plugins/Experimental/LandscapePatch/Source/LandscapePatch/LandscapePatch.Build.cs
+++ b/Engine/Plugins/Experimental/LandscapePatch/Source/LandscapePatch/LandscapePatch.Build.cs
@@ -39,6 +39,7 @@ public class LandscapePatch : ModuleRules
 				"RenderCore",
 				"RHI",
 				"Projects", // IPluginManager
+				"Renderer",
 			}
 			);
 			
diff --git a/Engine/Plugins/Experimental/LandscapePatch/Source/LandscapePatch/Private/LandscapeCircleHeightPatch.cpp b/Engine/Plugins/Experimental/LandscapePatch/Source/LandscapePatch/Private/LandscapeCircleHeightPatch.cpp
--- a/Engine/Plugins/Experimental/LandscapePatch/Source/LandscapePatch/Private/LandscapeCircleHeightPatch.cpp
+++ b/Engine/Plugins/Experimental/LandscapePatch/Source/LandscapePatch/Private/LandscapeCircleHeightPatch.cpp
@@ -10,9 +10,11 @@
 #include "LandscapeInfo.h"
 #include "LandscapePatchUtil.h" // GetHeightmapToWorld
 #include "LandscapeUtils.h" // IsVisibilityLayer
+#include "Materials/MaterialExpressionLandscapeVisibilityMask.h" 
 #include "RenderGraphBuilder.h"
 #include "RenderGraphUtils.h" // AddCopyTexturePass
 #include "RenderingThread.h"
+#include "SystemTextures.h"
 #include "TextureResource.h"
 
 #include UE_INLINE_GENERATED_CPP_BY_NAME(LandscapeCircleHeightPatch)
@@ -43,6 +45,7 @@ void ULandscapeCircleHeightPatch::OnComponentCreated()
 	}
 }
 
+#if WITH_EDITOR
 UTextureRenderTarget2D* ULandscapeCircleHeightPatch::RenderLayer_Native(const FLandscapeBrushParameters& InParameters, const FTransform& HeightmapCoordsToWorld)
 {
 	using namespace UE::Landscape;
@@ -55,7 +58,7 @@ UTextureRenderTarget2D* ULandscapeCircleHeightPatch::RenderLayer_Native(const FL
 
 	// We render in immediate mode for the legacy path so we use an immediate recorder :  
 	FRDGBuilderRecorder RDGBuilderRecorderImmediate;
-	ApplyCirclePatch(RDGBuilderRecorderImmediate, InParameters.LayerType == ELandscapeToolTargetType::Visibility,
+	ApplyCirclePatch(/*bInPerformBlending = */true, /*RenderParams = */nullptr, RDGBuilderRecorderImmediate, InParameters.LayerType == ELandscapeToolTargetType::Visibility,
 		InParameters.CombinedResult->GetResource(), 0,
 		FIntPoint(InParameters.CombinedResult->SizeX, InParameters.CombinedResult->SizeY),
 		HeightmapCoordsToWorld, /*OutputAccess = */ERHIAccess::SRVMask); // in the legacy path the RT leaves as a SRV
@@ -63,62 +66,93 @@ UTextureRenderTarget2D* ULandscapeCircleHeightPatch::RenderLayer_Native(const FL
 	return InParameters.CombinedResult;
 }
 
-#if WITH_EDITOR
-void ULandscapeCircleHeightPatch::RenderLayer(FRenderParams& InRenderParams, UE::Landscape::FRDGBuilderRecorder& RDGBuilderRecorder)
+UE::Landscape::EditLayers::ERenderFlags ULandscapeCircleHeightPatch::GetRenderFlags(const UE::Landscape::EditLayers::FMergeContext* InMergeContext) const
+{
+	using namespace UE::Landscape::EditLayers;
+
+	return ERenderFlags::RenderMode_Recorded | ERenderFlags::BlendMode_SeparateBlend | ERenderFlags::RenderLayerGroup_SupportsGrouping;
+}
+
+bool ULandscapeCircleHeightPatch::CanGroupRenderLayerWith(TScriptInterface<ILandscapeEditLayerRenderer> InOtherRenderer) const
+{
+	UObject* OtherRenderer = InOtherRenderer.GetObject();
+	check(OtherRenderer != nullptr);
+	// Circle height patches are compatible with one another : 
+	return OtherRenderer->IsA<ULandscapeCircleHeightPatch>();
+}
+
+bool ULandscapeCircleHeightPatch::RenderLayer(FRenderParams& RenderParams, UE::Landscape::FRDGBuilderRecorder& RDGBuilderRecorder)
 {
 	using namespace UE::Landscape::PatchUtil;
 	using namespace UE::Landscape::EditLayers;
 
 	checkf(RDGBuilderRecorder.IsRecording(), TEXT("ERenderFlags::RenderMode_Recorded means the command recorder should be recording at this point"));
-	FTransform HeightmapCoordsToWorld = GetHeightmapToWorld(InRenderParams.RenderAreaWorldTransform);
+	FTransform HeightmapCoordsToWorld = GetHeightmapToWorld(RenderParams.RenderAreaWorldTransform);
 
-	ULandscapeScratchRenderTarget* LandscapeRT = InRenderParams.MergeRenderContext->GetBlendRenderTargetWrite();
+	ULandscapeScratchRenderTarget* LandscapeRT = RenderParams.MergeRenderContext->GetBlendRenderTargetWrite();
 
-	const bool bIsHeightmapTarget = InRenderParams.MergeRenderContext->IsHeightmapMerge();
+	const bool bIsHeightmapTarget = RenderParams.MergeRenderContext->IsHeightmapMerge();
 	if (bIsHeightmapTarget)
 	{
 		UTextureRenderTarget2D* OutputToBlendInto = LandscapeRT->TryGetRenderTarget2D();
-		if (ensure(OutputToBlendInto))
-		{
-			ApplyCirclePatch(RDGBuilderRecorder, /*bIsVisibilityLayer = */false, OutputToBlendInto->GetResource(),
-				0, InRenderParams.RenderAreaSectionRect.Size(), HeightmapCoordsToWorld);
-		}
-		return;
+		check(OutputToBlendInto != nullptr);
+		return ApplyCirclePatch(/*bInPerformBlending = */false, &RenderParams, RDGBuilderRecorder, /*bIsVisibilityLayer = */false, OutputToBlendInto->GetResource(),
+			0, RenderParams.RenderAreaSectionRect.Size(), HeightmapCoordsToWorld);
 	}
 
 	// If we got to here, we're not processing a heightmap, so we only need to do anything if the 
 	//  patch edits visibility.
 	if (!bEditVisibility)
 	{
-		return;
+		return false;
 	}
 
 	UTextureRenderTarget2DArray* TextureArray = LandscapeRT->TryGetRenderTarget2DArray();
 	check(TextureArray != nullptr);
 
-	int32 NumLayers = InRenderParams.TargetLayerGroupLayerInfos.Num();
+	int32 NumLayers = RenderParams.TargetLayerGroupLayerInfos.Num();
 	check(LandscapeRT->GetEffectiveNumSlices() == NumLayers);
 
+	bool bDidRenderSomething = false;
 	for (int32 LayerIndex = 0; LayerIndex < NumLayers; ++LayerIndex)
 	{
-		if (UE::Landscape::IsVisibilityLayer(InRenderParams.TargetLayerGroupLayerInfos[LayerIndex]))
+		if (UE::Landscape::IsVisibilityLayer(RenderParams.TargetLayerGroupLayerInfos[LayerIndex]))
 		{
-			ApplyCirclePatch(RDGBuilderRecorder, /*bIsVisibilityLayer = */true, TextureArray->GetResource(),
-				LayerIndex, InRenderParams.RenderAreaSectionRect.Size(), HeightmapCoordsToWorld);
+			bDidRenderSomething |= ApplyCirclePatch(/*bInPerformBlending = */false, &RenderParams, RDGBuilderRecorder, /*bIsVisibilityLayer = */true, TextureArray->GetResource(),
+				LayerIndex, RenderParams.RenderAreaSectionRect.Size(), HeightmapCoordsToWorld);
 		}
 	}
+	return bDidRenderSomething;
 }
-#endif // WITH_EDITOR
 
-void ULandscapeCircleHeightPatch::ApplyCirclePatch(UE::Landscape::FRDGBuilderRecorder& RDGBuilderRecorder, bool bIsVisibilityLayer,
+void ULandscapeCircleHeightPatch::BlendLayer(FRenderParams& RenderParams, UE::Landscape::FRDGBuilderRecorder& RDGBuilderRecorder)
+{
+	using namespace UE::Landscape::EditLayers;
+
+	FBlendParams BlendParams;
+	if (RenderParams.MergeRenderContext->IsHeightmapMerge())
+	{
+		BlendParams.HeightmapBlendParams.BlendMode = EHeightmapBlendMode::AlphaBlend;
+	}
+	// Circle height patch only supports visibility (the others are using EWeightmapBlendMode::Passthrough): 
+	else if (bEditVisibility)
+	{
+		BlendParams.WeightmapBlendParams.Emplace(UMaterialExpressionLandscapeVisibilityMask::ParameterName, EWeightmapBlendMode::Additive);
+	}
+
+	// Then perform the generic blend : 
+	RenderParams.MergeRenderContext->GenericBlendLayer(BlendParams, RDGBuilderRecorder);
+}
+
+bool ULandscapeCircleHeightPatch::ApplyCirclePatch(bool bInPerformBlending, FRenderParams* RenderParams, UE::Landscape::FRDGBuilderRecorder& RDGBuilderRecorder, bool bIsVisibilityLayer,
 	FTextureResource* InMergedLandscapeTextureResource, int32 LandscapeTextureSliceIndex,
-	const FIntPoint& DestinationResolution, const FTransform & HeightmapCoordsToWorld, ERHIAccess OutputAccess)
+	const FIntPoint& DestinationResolution, const FTransform& HeightmapCoordsToWorld, ERHIAccess OutputAccess)
 {
 	using namespace UE::Landscape;
 
 	if (bEditVisibility != bIsVisibilityLayer)
 	{
-		return;
+		return false;
 	}
 
 	double ToHeightmapRadiusScale = GetComponentTransform().GetScale3D().X / HeightmapCoordsToWorld.GetScale3D().X;
@@ -131,77 +165,127 @@ void ULandscapeCircleHeightPatch::ApplyCirclePatch(UE::Landscape::FRDGBuilderRec
 	float HeightmapFalloff = Falloff * ToHeightmapRadiusScale + RadiusAdjustment;
 
 	FIntRect DestinationBounds(
-		FMath::Clamp(FMath::Floor(CenterInHeightmapCoordinates.X - HeightmapRadius - HeightmapFalloff), 
+		FMath::Clamp(FMath::Floor(CenterInHeightmapCoordinates.X - HeightmapRadius - HeightmapFalloff),
 			0, DestinationResolution.X),
-		FMath::Clamp(FMath::Floor(CenterInHeightmapCoordinates.Y - HeightmapRadius - HeightmapFalloff), 
+		FMath::Clamp(FMath::Floor(CenterInHeightmapCoordinates.Y - HeightmapRadius - HeightmapFalloff),
 			0, DestinationResolution.Y),
-		FMath::Clamp(FMath::CeilToInt(CenterInHeightmapCoordinates.X + HeightmapRadius + HeightmapFalloff) + 1, 
+		FMath::Clamp(FMath::CeilToInt(CenterInHeightmapCoordinates.X + HeightmapRadius + HeightmapFalloff) + 1,
 			0, DestinationResolution.X),
-		FMath::Clamp(FMath::CeilToInt(CenterInHeightmapCoordinates.Y + HeightmapRadius + HeightmapFalloff) + 1, 
+		FMath::Clamp(FMath::CeilToInt(CenterInHeightmapCoordinates.Y + HeightmapRadius + HeightmapFalloff) + 1,
 			0, DestinationResolution.Y));
 
 	if (DestinationBounds.Area() <= 0)
 	{
 		// Must be outside the landscape
-		return;
+		return false;
+	}
+
+	FTextureResource* OutputResource = InMergedLandscapeTextureResource;
+	FString OutputResourceName = OutputResource->GetResourceName().ToString();
+	if (!bInPerformBlending)
+	{ 
+		check(RenderParams != nullptr);
+		ULandscapeScratchRenderTarget* WriteRT = RenderParams->MergeRenderContext->GetBlendRenderTargetWrite();
+	    // After this point, the render cannot fail so if we're the first in our render layer group to render, we can cycle the blend render targets and 
+	    //  start rendering in the write one : 
+	    if (RenderParams->NumSuccessfulRenderLayerStepsUntilBlendLayerStep == 0)
+	    {
+		    RenderParams->MergeRenderContext->CycleBlendRenderTargets(RDGBuilderRecorder);
+		    WriteRT = RenderParams->MergeRenderContext->GetBlendRenderTargetWrite();
+		    WriteRT->Clear(RDGBuilderRecorder);
+			check(WriteRT->GetCurrentState() == ERHIAccess::RTV);
+	    }
+		OutputResource = WriteRT->GetRenderTarget()->GetResource();
+		OutputResourceName = WriteRT->GetDebugName();
 	}
 
 	auto RDGCommand =
-		[InMergedLandscapeTextureResource, LandscapeTextureSliceIndex, DestinationBounds,
-		CenterInHeightmapCoordinates, HeightmapRadius, HeightmapFalloff, bEditVisibility = bEditVisibility](FRDGBuilder& GraphBuilder)
+		[ OutputResource
+		, OutputResourceName
+		, LandscapeTextureSliceIndex
+		, DestinationBounds
+		, CenterInHeightmapCoordinates
+		, HeightmapRadius
+		, HeightmapFalloff
+		, bEditVisibility = bEditVisibility
+		, bInPerformBlending](FRDGBuilder& GraphBuilder)
 		{
 			TRACE_CPUPROFILER_EVENT_SCOPE(LandscapeCircleHeightPatch);
 
+			FRDGTextureRef DestinationTexture = GraphBuilder.RegisterExternalTexture(CreateRenderTarget(OutputResource->TextureRHI, TEXT("OutputTexture")));
 			const TCHAR* OutputName = bEditVisibility ? TEXT("LandscapeCircleVisibilityPatchOutput") : TEXT("LandscapeCircleHeightPatchOutput");
-			const TCHAR* InputCopyName = bEditVisibility ? TEXT("LandscapeCircleVisibilityPatchInputCopy") : TEXT("LandscapeCircleHeightPatchInputCopy");
-
-			TRefCountPtr<IPooledRenderTarget> RenderTarget = CreateRenderTarget(InMergedLandscapeTextureResource->GetTextureRHI(), OutputName);
-			FRDGTextureRef DestinationTexture = GraphBuilder.RegisterExternalTexture(RenderTarget);
-
-			// Make a copy of the portion of our weightmap input that we're writing to so that we can 
-			//  read and write at the same time (needed for blending)
-			FRDGTextureDesc InputCopyDescription = DestinationTexture->Desc;
-			InputCopyDescription.Dimension = ETextureDimension::Texture2D;
-			InputCopyDescription.ArraySize = 1;
-			InputCopyDescription.NumMips = 1;
-			InputCopyDescription.Extent = DestinationBounds.Size();
-			FRDGTextureRef InputCopy = GraphBuilder.CreateTexture(InputCopyDescription, InputCopyName);
-
-			FRHICopyTextureInfo CopyTextureInfo;
-			CopyTextureInfo.SourceMipIndex = 0;
-			CopyTextureInfo.NumMips = 1;
-			CopyTextureInfo.SourceSliceIndex = LandscapeTextureSliceIndex;
-			CopyTextureInfo.NumSlices = 1;
-			CopyTextureInfo.SourcePosition = FIntVector(DestinationBounds.Min.X, DestinationBounds.Min.Y, 0);
-			CopyTextureInfo.Size = FIntVector(InputCopyDescription.Extent.X, InputCopyDescription.Extent.Y, 0);
-			AddCopyTexturePass(GraphBuilder, DestinationTexture, InputCopy, CopyTextureInfo);
-
-			FRDGTextureSRVRef InputCopySRV = GraphBuilder.CreateSRV(FRDGTextureSRVDesc::CreateForMipLevel(InputCopy, 0));
-
-			FLandscapeCircleHeightPatchPS::FParameters* ShaderParams = GraphBuilder.AllocParameters<FLandscapeCircleHeightPatchPS::FParameters>();
+
+			FRDGTextureSRVRef InputCopySRV = nullptr;
+			if (bInPerformBlending)
+			{
+			    // Make a copy of the portion of our weightmap input that we're writing to so that we can 
+			    //  read and write at the same time (needed for blending)
+			    FRDGTextureDesc InputCopyDescription = DestinationTexture->Desc;
+			    InputCopyDescription.Dimension = ETextureDimension::Texture2D;
+			    InputCopyDescription.ArraySize = 1;
+			    InputCopyDescription.NumMips = 1;
+			    InputCopyDescription.Extent = DestinationBounds.Size();
+			    const TCHAR* InputCopyName = bEditVisibility ? TEXT("LandscapeCircleVisibilityPatchInputCopy") : TEXT("LandscapeCircleHeightPatchInputCopy");
+			    FRDGTextureRef InputCopy = GraphBuilder.CreateTexture(InputCopyDescription, InputCopyName);
+    
+			    FRHICopyTextureInfo CopyTextureInfo;
+			    CopyTextureInfo.SourceMipIndex = 0;
+			    CopyTextureInfo.NumMips = 1;
+			    CopyTextureInfo.SourceSliceIndex = LandscapeTextureSliceIndex;
+			    CopyTextureInfo.NumSlices = 1;
+			    CopyTextureInfo.SourcePosition = FIntVector(DestinationBounds.Min.X, DestinationBounds.Min.Y, 0);
+			    CopyTextureInfo.Size = FIntVector(InputCopyDescription.Extent.X, InputCopyDescription.Extent.Y, 0);
+			    AddCopyTexturePass(GraphBuilder, DestinationTexture, InputCopy, CopyTextureInfo);
+    
+			    InputCopySRV = GraphBuilder.CreateSRV(FRDGTextureSRVDesc::CreateForMipLevel(InputCopy, 0));
+			}
+
+			FLandscapeCircleHeightPatchPSBase::FParameters* ShaderParams = GraphBuilder.AllocParameters<FLandscapeCircleHeightPatchPSBase::FParameters>();
 			ShaderParams->InCenter = (FVector3f)CenterInHeightmapCoordinates;
 			ShaderParams->InRadius = HeightmapRadius;
 			ShaderParams->InFalloff = HeightmapFalloff;
-			ShaderParams->InSourceTexture = InputCopySRV;
-			ShaderParams->InSourceTextureOffset = DestinationBounds.Min;
+			
+			ShaderParams->InSourceTexture = GraphBuilder.CreateSRV(FRDGTextureSRVDesc::Create(GSystemTextures.GetBlackDummy(GraphBuilder)));
+			ShaderParams->InSourceTextureOffset = FIntPoint(ForceInit);
+			if (bInPerformBlending)
+			{
+				ShaderParams->InSourceTexture = InputCopySRV;
+				ShaderParams->InSourceTextureOffset = DestinationBounds.Min;
+			}
+
 			ShaderParams->RenderTargets[0] = FRenderTargetBinding(DestinationTexture, ERenderTargetLoadAction::ENoAction,
 				/*InMipIndex = */0, LandscapeTextureSliceIndex);
 
 			if (bEditVisibility)
 			{
-				FLandscapeCircleVisibilityPatchPS::AddToRenderGraph(GraphBuilder, ShaderParams, DestinationBounds);
+				if (bInPerformBlending)
+				{
+					FLandscapeCircleVisibilityPatchPS</*bPerformBlending = */true>::AddToRenderGraph(RDG_EVENT_NAME("RenderCircleVisibilityPatch -> %s", *OutputResourceName), GraphBuilder, ShaderParams, DestinationBounds);
+				}
+				else
+				{
+					FLandscapeCircleVisibilityPatchPS</*bPerformBlending = */false>::AddToRenderGraph(RDG_EVENT_NAME("RenderCircleVisibilityPatch -> %s", *OutputResourceName), GraphBuilder, ShaderParams, DestinationBounds);
+				}
 			}
 			else
 			{
-				FLandscapeCircleHeightPatchPS::AddToRenderGraph(GraphBuilder, ShaderParams, DestinationBounds);
+				if (bInPerformBlending)
+				{
+					FLandscapeCircleHeightPatchPS</*bPerformBlending = */true>::AddToRenderGraph(RDG_EVENT_NAME("RenderCircleHeightPatch -> %s", *OutputResourceName), GraphBuilder, ShaderParams, DestinationBounds);
+				}
+				else
+				{
+					FLandscapeCircleHeightPatchPS</*bPerformBlending = */false>::AddToRenderGraph(RDG_EVENT_NAME("RenderCircleHeightPatch -> %s", *OutputResourceName), GraphBuilder, ShaderParams, DestinationBounds);
+				}
 			}
 		};
 
 	// We need to specify the final state of the external texture to prevent the graph builder from transitioning it to SRVMask :
 	RDGBuilderRecorder.EnqueueRDGCommand(RDGCommand, { { InMergedLandscapeTextureResource, (OutputAccess == ERHIAccess::None) ? ERHIAccess::RTV : OutputAccess} });
+
+	return true;
 }
 
-#if WITH_EDITOR
 void ULandscapeCircleHeightPatch::GetRendererStateInfo(const UE::Landscape::EditLayers::FMergeContext* InMergeContext,
 	FEditLayerTargetTypeState& OutSupported,
 	FEditLayerTargetTypeState& OutEnabled,
diff --git a/Engine/Plugins/Experimental/LandscapePatch/Source/LandscapePatch/Private/LandscapeCircleHeightPatchPS.cpp b/Engine/Plugins/Experimental/LandscapePatch/Source/LandscapePatch/Private/LandscapeCircleHeightPatchPS.cpp
--- a/Engine/Plugins/Experimental/LandscapePatch/Source/LandscapePatch/Private/LandscapeCircleHeightPatchPS.cpp
+++ b/Engine/Plugins/Experimental/LandscapePatch/Source/LandscapePatch/Private/LandscapeCircleHeightPatchPS.cpp
@@ -5,7 +5,7 @@
 #include "PixelShaderUtils.h"
 #include "DataDrivenShaderPlatformInfo.h"
 
-bool FLandscapeCircleHeightPatchPS::ShouldCompilePermutation(const FGlobalShaderPermutationParameters& Parameters)
+bool FLandscapeCircleHeightPatchPSBase::ShouldCompilePermutation(const FGlobalShaderPermutationParameters& Parameters)
 {
 	// Apparently landscape requires a particular feature level
 	return IsFeatureLevelSupported(Parameters.Platform, ERHIFeatureLevel::SM5)
@@ -13,12 +13,18 @@ bool FLandscapeCircleHeightPatchPS::ShouldCompilePermutation(const FGlobalShader
 		&& !IsMetalMobilePlatform(Parameters.Platform);
 }
 
-void FLandscapeCircleHeightPatchPS::ModifyCompilationEnvironment(const FGlobalShaderPermutationParameters& Parameters, FShaderCompilerEnvironment& OutEnvironment)
+template <bool bPerformBlending>
+void FLandscapeCircleHeightPatchPS<bPerformBlending>::ModifyCompilationEnvironment(const FGlobalShaderPermutationParameters& Parameters, FShaderCompilerEnvironment& OutEnvironment)
 {
+	if (bPerformBlending)
+	{
+		OutEnvironment.SetDefine(TEXT("PERFORM_BLENDING"), 1);
+	}
 	OutEnvironment.SetDefine(TEXT("CIRCLE_HEIGHT_PATCH"), 1);
 }
 
-void FLandscapeCircleHeightPatchPS::AddToRenderGraph(FRDGBuilder& GraphBuilder, FParameters* Parameters, const FIntRect& DestinationBounds)
+template <bool bPerformBlending>
+void FLandscapeCircleHeightPatchPS<bPerformBlending>::AddToRenderGraph(FRDGEventName&& InRDGEventName, FRDGBuilder& GraphBuilder, FLandscapeCircleHeightPatchPSBase::FParameters* Parameters, const FIntRect& DestinationBounds)
 {
 	FGlobalShaderMap* ShaderMap = GetGlobalShaderMap(GMaxRHIFeatureLevel);
 	TShaderMapRef<FLandscapeCircleHeightPatchPS> PixelShader(ShaderMap);
@@ -26,18 +32,24 @@ void FLandscapeCircleHeightPatchPS::AddToRenderGraph(FRDGBuilder& GraphBuilder,
 	FPixelShaderUtils::AddFullscreenPass(
 		GraphBuilder,
 		ShaderMap,
-		RDG_EVENT_NAME("LandscapeCircleHeightPatch"),
+		MoveTemp(InRDGEventName),
 		PixelShader,
 		Parameters,
 		DestinationBounds);
 }
 
-void FLandscapeCircleVisibilityPatchPS::ModifyCompilationEnvironment(const FGlobalShaderPermutationParameters& Parameters, FShaderCompilerEnvironment& OutEnvironment)
+template <bool bPerformBlending>
+void FLandscapeCircleVisibilityPatchPS<bPerformBlending>::ModifyCompilationEnvironment(const FGlobalShaderPermutationParameters& Parameters, FShaderCompilerEnvironment& OutEnvironment)
 {
+	if (bPerformBlending)
+	{
+		OutEnvironment.SetDefine(TEXT("PERFORM_BLENDING"), 1);
+	}
 	OutEnvironment.SetDefine(TEXT("CIRCLE_VISIBILITY_PATCH"), 1);
 }
 
-void FLandscapeCircleVisibilityPatchPS::AddToRenderGraph(FRDGBuilder& GraphBuilder, FParameters* Parameters, const FIntRect& DestinationBounds)
+template <bool bPerformBlending>
+void FLandscapeCircleVisibilityPatchPS<bPerformBlending>::AddToRenderGraph(FRDGEventName&& InRDGEventName, FRDGBuilder& GraphBuilder, FLandscapeCircleHeightPatchPSBase::FParameters* Parameters, const FIntRect& DestinationBounds)
 {
 	FGlobalShaderMap* ShaderMap = GetGlobalShaderMap(GMaxRHIFeatureLevel);
 	TShaderMapRef<FLandscapeCircleVisibilityPatchPS> PixelShader(ShaderMap);
@@ -45,11 +57,13 @@ void FLandscapeCircleVisibilityPatchPS::AddToRenderGraph(FRDGBuilder& GraphBuild
 	FPixelShaderUtils::AddFullscreenPass(
 		GraphBuilder,
 		ShaderMap,
-		RDG_EVENT_NAME("LandscapeCircleVisibilityPatch"),
+		MoveTemp(InRDGEventName),
 		PixelShader,
 		Parameters,
 		DestinationBounds);
 }
 
-IMPLEMENT_GLOBAL_SHADER(FLandscapeCircleHeightPatchPS, "/Plugin/LandscapePatch/Private/LandscapeCircleHeightPatchPS.usf", "ApplyLandscapeCircleHeightPatch", SF_Pixel);
-IMPLEMENT_GLOBAL_SHADER(FLandscapeCircleVisibilityPatchPS, "/Plugin/LandscapePatch/Private/LandscapeCircleHeightPatchPS.usf", "ApplyLandscapeCircleVisibilityPatch", SF_Pixel);
+IMPLEMENT_SHADER_TYPE(template<>, FLandscapeCircleHeightPatchPS</*bPerformBlending = */false>, TEXT("/Plugin/LandscapePatch/Private/LandscapeCircleHeightPatchPS.usf"), TEXT("ApplyLandscapeCircleHeightPatch"), SF_Pixel);
+IMPLEMENT_SHADER_TYPE(template<>, FLandscapeCircleHeightPatchPS</*bPerformBlending = */true>, TEXT("/Plugin/LandscapePatch/Private/LandscapeCircleHeightPatchPS.usf"), TEXT("ApplyLandscapeCircleHeightPatch"), SF_Pixel);
+IMPLEMENT_SHADER_TYPE(template<>, FLandscapeCircleVisibilityPatchPS</*bPerformBlending = */false>, TEXT("/Plugin/LandscapePatch/Private/LandscapeCircleHeightPatchPS.usf"), TEXT("ApplyLandscapeCircleVisibilityPatch"), SF_Pixel);
+IMPLEMENT_SHADER_TYPE(template<>, FLandscapeCircleVisibilityPatchPS</*bPerformBlending = */true>, TEXT("/Plugin/LandscapePatch/Private/LandscapeCircleHeightPatchPS.usf"), TEXT("ApplyLandscapeCircleVisibilityPatch"), SF_Pixel);
diff --git a/Engine/Plugins/Experimental/LandscapePatch/Source/LandscapePatch/Private/LandscapeCircleHeightPatchPS.h b/Engine/Plugins/Experimental/LandscapePatch/Source/LandscapePatch/Private/LandscapeCircleHeightPatchPS.h
--- a/Engine/Plugins/Experimental/LandscapePatch/Source/LandscapePatch/Private/LandscapeCircleHeightPatchPS.h
+++ b/Engine/Plugins/Experimental/LandscapePatch/Source/LandscapePatch/Private/LandscapeCircleHeightPatchPS.h
@@ -8,11 +8,8 @@
 /**
  * Shader that applies a circle height patch to a landscape heightmap.
  */
-class FLandscapeCircleHeightPatchPS : public FGlobalShader
+class FLandscapeCircleHeightPatchPSBase : public FGlobalShader
 {
-	DECLARE_GLOBAL_SHADER(FLandscapeCircleHeightPatchPS);
-	SHADER_USE_PARAMETER_STRUCT(FLandscapeCircleHeightPatchPS, FGlobalShader);
-
 public:
 	BEGIN_SHADER_PARAMETER_STRUCT(FParameters, )
 		SHADER_PARAMETER_RDG_TEXTURE_SRV(Texture2D<float4>, InSourceTexture) // Our input texture
@@ -27,25 +24,37 @@ public:
 		RENDER_TARGET_BINDING_SLOTS() // Holds our output
 	END_SHADER_PARAMETER_STRUCT()
 
+	FLandscapeCircleHeightPatchPSBase() {}
+	FLandscapeCircleHeightPatchPSBase(const ShaderMetaType::CompiledShaderInitializerType& Initializer)
+		: FGlobalShader(Initializer)
+	{}
+
+public:
 	static bool ShouldCompilePermutation(const FGlobalShaderPermutationParameters& Parameters);
+};
+
+template <bool bPerformBlending = false>
+class FLandscapeCircleHeightPatchPS : public FLandscapeCircleHeightPatchPSBase
+{
+	DECLARE_GLOBAL_SHADER(FLandscapeCircleHeightPatchPS);
+	SHADER_USE_PARAMETER_STRUCT(FLandscapeCircleHeightPatchPS, FLandscapeCircleHeightPatchPSBase);
+
+public:
 	static void ModifyCompilationEnvironment(const FGlobalShaderPermutationParameters& Parameters, FShaderCompilerEnvironment& OutEnvironment);
 
-	static void AddToRenderGraph(FRDGBuilder& GraphBuilder, FParameters* InParameters, const FIntRect& DestinationBounds);
+	static void AddToRenderGraph(FRDGEventName&& InRDGEventName, FRDGBuilder& GraphBuilder, FLandscapeCircleHeightPatchPSBase::FParameters* InParameters, const FIntRect& DestinationBounds);
 };
 
 /**
  * Shader that applies a circle patch to a landscape visibility layer.
  */
-class FLandscapeCircleVisibilityPatchPS : public FLandscapeCircleHeightPatchPS
+template <bool bPerformBlending = false>
+class FLandscapeCircleVisibilityPatchPS : public FLandscapeCircleHeightPatchPSBase
 {
 	DECLARE_GLOBAL_SHADER(FLandscapeCircleVisibilityPatchPS);
+	SHADER_USE_PARAMETER_STRUCT(FLandscapeCircleVisibilityPatchPS, FLandscapeCircleHeightPatchPSBase);
 
 public:
-	FLandscapeCircleVisibilityPatchPS() = default;
-	FLandscapeCircleVisibilityPatchPS(const ShaderMetaType::CompiledShaderInitializerType& Initializer)
-		: FLandscapeCircleHeightPatchPS(Initializer)
-	{}
-
 	static void ModifyCompilationEnvironment(const FGlobalShaderPermutationParameters& Parameters, FShaderCompilerEnvironment& OutEnvironment);
-	static void AddToRenderGraph(FRDGBuilder& GraphBuilder, FParameters* InParameters, const FIntRect& DestinationBounds);
+	static void AddToRenderGraph(FRDGEventName&& InRDGEventName, FRDGBuilder& GraphBuilder, FLandscapeCircleHeightPatchPSBase::FParameters* InParameters, const FIntRect& DestinationBounds);
 };
diff --git a/Engine/Plugins/Experimental/LandscapePatch/Source/LandscapePatch/Private/LandscapeTexturePatch.cpp b/Engine/Plugins/Experimental/LandscapePatch/Source/LandscapePatch/Private/LandscapeTexturePatch.cpp
--- a/Engine/Plugins/Experimental/LandscapePatch/Source/LandscapePatch/Private/LandscapeTexturePatch.cpp
+++ b/Engine/Plugins/Experimental/LandscapePatch/Source/LandscapePatch/Private/LandscapeTexturePatch.cpp
@@ -17,8 +17,9 @@
 #include "Misc/UObjectToken.h"
 #include "RHIStaticStates.h"
 #include "RenderGraphUtils.h"
-#include "TextureResource.h"
 #include "RenderingThread.h"
+#include "SystemTextures.h"
+#include "TextureResource.h"
 #include "Algo/AnyOf.h"
 
 #include UE_INLINE_GENERATED_CPP_BY_NAME(LandscapeTexturePatch)
@@ -140,26 +141,35 @@ UE::Landscape::EditLayers::ERenderFlags ULandscapeTexturePatch::GetRenderFlags(c
 	{
 		RenderFlags |= ERenderFlags::RenderMode_Recorded;
 	}
+	RenderFlags |= ERenderFlags::BlendMode_SeparateBlend | ERenderFlags::RenderLayerGroup_SupportsGrouping;
 	return RenderFlags;
 }
 
-void ULandscapeTexturePatch::RenderLayer(FRenderParams& InRenderParams, UE::Landscape::FRDGBuilderRecorder& RDGBuilderRecorder)
+bool ULandscapeTexturePatch::CanGroupRenderLayerWith(TScriptInterface<ILandscapeEditLayerRenderer> InOtherRenderer) const
+{
+	UObject* OtherRenderer = InOtherRenderer.GetObject();
+	check(OtherRenderer != nullptr);
+	// Texture patches are compatible with one another (blend mode is handled per-pixel): 
+	return OtherRenderer->IsA<ULandscapeTexturePatch>();
+}
+
+bool ULandscapeTexturePatch::RenderLayer(FRenderParams& RenderParams, UE::Landscape::FRDGBuilderRecorder& RDGBuilderRecorder)
 {
 	using namespace UE::Landscape;
 	using namespace UE::Landscape::PatchUtil;
 	using namespace UE::Landscape::EditLayers;
 
-	FTransform LandscapeHeightmapToWorld = GetHeightmapToWorld(InRenderParams.RenderAreaWorldTransform);
+	FTransform LandscapeHeightmapToWorld = GetHeightmapToWorld(RenderParams.RenderAreaWorldTransform);
 
-	ULandscapeScratchRenderTarget* LandscapeScratchRT = InRenderParams.MergeRenderContext->GetBlendRenderTargetWrite();
+	ULandscapeScratchRenderTarget* LandscapeScratchRT = RenderParams.MergeRenderContext->GetBlendRenderTargetWrite();
 
-	const bool bIsHeightmapTarget = InRenderParams.MergeRenderContext->IsHeightmapMerge();
+	const bool bIsHeightmapTarget = RenderParams.MergeRenderContext->IsHeightmapMerge();
 	if (bIsHeightmapTarget)
 	{
 		UTextureRenderTarget2D* CurrentData = LandscapeScratchRT->TryGetRenderTarget2D();
 		if (!ensure(CurrentData))
 		{
-			return;
+			return false;
 		}
 
 		if (bReinitializeHeightOnNextRender)
@@ -167,13 +177,14 @@ void ULandscapeTexturePatch::RenderLayer(FRenderParams& InRenderParams, UE::Land
 			bReinitializeHeightOnNextRender = false;
 			checkf(!RDGBuilderRecorder.IsRecording(), TEXT("We should be using ERenderFlags::RenderMode_Immediate when reinitializing height"));
 			ReinitializeHeight(CurrentData, LandscapeHeightmapToWorld);
-			return;
+			return true;
 		}
 		else
 		{
 			checkf(RDGBuilderRecorder.IsRecording(), TEXT("We should be using ERenderFlags::RenderMode_Recorded in the typical case"));
-			ApplyToHeightmap(RDGBuilderRecorder, CurrentData, GetHeightmapToWorld(InRenderParams.RenderAreaWorldTransform));
-			return;
+			bool bDidRenderSomething = false;
+			ApplyToHeightmap(/*bInPerformBlending = */false, &RenderParams, RDGBuilderRecorder, CurrentData, GetHeightmapToWorld(RenderParams.RenderAreaWorldTransform), bDidRenderSomething);
+			return bDidRenderSomething;
 		}
 	}
 
@@ -182,20 +193,22 @@ void ULandscapeTexturePatch::RenderLayer(FRenderParams& InRenderParams, UE::Land
 	UTextureRenderTarget2DArray* TextureArray = LandscapeScratchRT->TryGetRenderTarget2DArray();
 	check(TextureArray != nullptr);
 	
-	int32 NumTargetLayersInGroup = InRenderParams.TargetLayerGroupLayerNames.Num();
+	int32 NumTargetLayersInGroup = RenderParams.TargetLayerGroupLayerNames.Num();
 	check(LandscapeScratchRT->GetEffectiveNumSlices() == NumTargetLayersInGroup);
 
+	bool bDidRenderSomething = false;
+
 	for (int32 TargetLayerIndexInGroup = 0; TargetLayerIndexInGroup < NumTargetLayersInGroup; ++TargetLayerIndexInGroup)
 	{
-		bool bIsVisibilityLayer = ensure(TargetLayerIndexInGroup < InRenderParams.TargetLayerGroupLayerInfos.Num())
-			&& UE::Landscape::IsVisibilityLayer(InRenderParams.TargetLayerGroupLayerInfos[TargetLayerIndexInGroup]);
+		bool bIsVisibilityLayer = ensure(TargetLayerIndexInGroup < RenderParams.TargetLayerGroupLayerInfos.Num())
+			&& UE::Landscape::IsVisibilityLayer(RenderParams.TargetLayerGroupLayerInfos[TargetLayerIndexInGroup]);
 		
 		// Try to find the weight patch
 		ULandscapeWeightPatchTextureInfo* WeightPatchInfo = nullptr;
 		for (const TObjectPtr<ULandscapeWeightPatchTextureInfo>& WeightPatchEntry : WeightPatches)
 		{
 			if ((bIsVisibilityLayer && WeightPatchEntry->bEditVisibilityLayer)
-				|| (WeightPatchEntry->WeightmapLayerName == InRenderParams.TargetLayerGroupLayerNames[TargetLayerIndexInGroup]))
+				|| (WeightPatchEntry->WeightmapLayerName == RenderParams.TargetLayerGroupLayerNames[TargetLayerIndexInGroup]))
 			{
 				WeightPatchInfo = WeightPatchEntry;
 				break;
@@ -214,19 +227,51 @@ void ULandscapeTexturePatch::RenderLayer(FRenderParams& InRenderParams, UE::Land
 			checkf(!RDGBuilderRecorder.IsRecording(), TEXT("We should be using ERenderFlags::RenderMode_Immediate when reinitializing weight"));
 			ReinitializeWeightPatch(WeightPatchInfo, TextureArray->GetResource(),
 				FIntPoint(TextureArray->SizeX, TextureArray->SizeY), TargetLayerIndexInGroup, LandscapeHeightmapToWorld);
+
+			bDidRenderSomething = true;
 		}
 		else
 		{
 			checkf(RDGBuilderRecorder.IsRecording(), TEXT("We should be using ERenderFlags::RenderMode_Recorded in the typical case"));
-			ApplyToWeightmap(RDGBuilderRecorder, WeightPatchInfo,
+			bool bLayerDidRenderSomething = false;
+			ApplyToWeightmap(/*bInPerformBlending = */false, &RenderParams, RDGBuilderRecorder, WeightPatchInfo,
 				TextureArray->GetResource(),
 				TargetLayerIndexInGroup,
-				InRenderParams.RenderAreaSectionRect.Size(), 
-				GetHeightmapToWorld(InRenderParams.RenderAreaWorldTransform));
+				RenderParams.RenderAreaSectionRect.Size(), 
+				GetHeightmapToWorld(RenderParams.RenderAreaWorldTransform), 
+				bLayerDidRenderSomething);
+
+			bDidRenderSomething |= bLayerDidRenderSomething;
 		}
 	}//end for each layer index
+
+	return bDidRenderSomething;
 }//end ULandscapeTexturePatch::RenderLayer
 
+void ULandscapeTexturePatch::BlendLayer(FRenderParams& RenderParams, UE::Landscape::FRDGBuilderRecorder& RDGBuilderRecorder)
+{
+	using namespace UE::Landscape::EditLayers;
+
+	// Prepare the generic blend params based on the patch's data : 
+	FBlendParams BlendParams;
+	if (RenderParams.MergeRenderContext->IsHeightmapMerge())
+	{
+		BlendParams.HeightmapBlendParams.BlendMode = EHeightmapBlendMode::AlphaBlend;
+	}
+	else 
+	{
+		BlendParams.WeightmapBlendParams.Reserve(RenderParams.TargetLayerGroupLayerNames.Num());
+		for (FName TargetLayerName : RenderParams.TargetLayerGroupLayerNames)
+		{
+			// only blend the layers involved in this step (the others are using EWeightmapBlendMode::Passthrough): 
+			FWeightmapBlendParams& TargetLayerBlendParams = BlendParams.WeightmapBlendParams.Emplace(TargetLayerName, EWeightmapBlendMode::AlphaBlend);
+		}
+	}
+
+	// Then perform the generic blend : 
+	RenderParams.MergeRenderContext->GenericBlendLayer(BlendParams, RDGBuilderRecorder);
+}
+
 // Legacy path, which gets the entire heightmap.
 UTextureRenderTarget2D* ULandscapeTexturePatch::RenderLayer_Native(const FLandscapeBrushParameters& InParameters, 
 	const FTransform& LandscapeHeightmapToWorld)
@@ -254,7 +299,8 @@ UTextureRenderTarget2D* ULandscapeTexturePatch::RenderLayer_Native(const FLandsc
 		{
 			// We render in immediate mode for the legacy path so we use an immediate recorder :  
 			FRDGBuilderRecorder RDGBuilderRecorderImmediate;
-			return ApplyToHeightmap(RDGBuilderRecorderImmediate, InParameters.CombinedResult, LandscapeHeightmapToWorld, /*OutputAccess = */ERHIAccess::SRVMask); // in the legacy path the RT leaves as a SRV
+			bool bDidRenderSomething = false;
+			return ApplyToHeightmap(/*bInPerformBlending = */true, /*RenderParams = */nullptr, RDGBuilderRecorderImmediate, InParameters.CombinedResult, LandscapeHeightmapToWorld, bDidRenderSomething, /*OutputAccess = */ERHIAccess::SRVMask); // in the legacy path the RT leaves as a SRV
 		}
 	}
 	else
@@ -298,21 +344,24 @@ UTextureRenderTarget2D* ULandscapeTexturePatch::RenderLayer_Native(const FLandsc
 		{
 			// We render in immediate mode for the legacy path so we use an immediate recorder :  
 			FRDGBuilderRecorder RDGBuilderRecorderImmediate;
-			ApplyToWeightmap(RDGBuilderRecorderImmediate, WeightPatchInfo,
+			bool bLayerDidRenderSomething = false;
+			ApplyToWeightmap(/*bInPerformBlending = */true, /*RenderParams = */nullptr, RDGBuilderRecorderImmediate, WeightPatchInfo,
 				InParameters.CombinedResult->GetResource(),
 				0, // Slice index
 				FIntPoint(InParameters.CombinedResult->SizeX, InParameters.CombinedResult->SizeY),
 				LandscapeHeightmapToWorld, 
+				bLayerDidRenderSomething,
 				/*OutputAccess = */ERHIAccess::SRVMask); // in the legacy path the RT leaves as a SRV
 			return InParameters.CombinedResult;
 		}
 	}
 }
 
-UTextureRenderTarget2D* ULandscapeTexturePatch::ApplyToHeightmap(UE::Landscape::FRDGBuilderRecorder& RDGBuilderRecorder, UTextureRenderTarget2D* InCombinedResult,
-	const FTransform& LandscapeHeightmapToWorld, ERHIAccess OutputAccess)
+UTextureRenderTarget2D* ULandscapeTexturePatch::ApplyToHeightmap(bool bInPerformBlending, FRenderParams* RenderParams, UE::Landscape::FRDGBuilderRecorder& RDGBuilderRecorder, UTextureRenderTarget2D* InCombinedResult,
+	const FTransform& LandscapeHeightmapToWorld, bool& bOutDidRenderSomething, ERHIAccess OutputAccess)
 {
 	using namespace UE::Landscape;
+	bOutDidRenderSomething = false;
 
 	// Get the source of our height patch
 	UTexture* PatchUObject = nullptr;
@@ -344,8 +393,8 @@ UTextureRenderTarget2D* ULandscapeTexturePatch::ApplyToHeightmap(UE::Landscape::
 		return InCombinedResult;
 	}
 
-	FTextureResource* Patch = PatchUObject->GetResource();
-	if (!Patch)
+	FTextureResource* PatchTextureResource = PatchUObject->GetResource();
+	if (!PatchTextureResource)
 	{
 		return InCombinedResult;
 	}
@@ -354,7 +403,7 @@ UTextureRenderTarget2D* ULandscapeTexturePatch::ApplyToHeightmap(UE::Landscape::
 	// individually in the lambda below.
 	FApplyLandscapeTextureHeightPatchPSParameters ShaderParamsToCopy;
 	FIntRect DestinationBounds;
-	GetHeightShaderParams(LandscapeHeightmapToWorld, FIntPoint(Patch->GetSizeX(), Patch->GetSizeY()), FIntPoint(InCombinedResult->SizeX, InCombinedResult->SizeY), ShaderParamsToCopy, DestinationBounds);
+	GetHeightShaderParams(LandscapeHeightmapToWorld, FIntPoint(PatchTextureResource->GetSizeX(), PatchTextureResource->GetSizeY()), FIntPoint(InCombinedResult->SizeX, InCombinedResult->SizeY), ShaderParamsToCopy, DestinationBounds);
 
 	if (DestinationBounds.Area() <= 0)
 	{
@@ -362,58 +411,100 @@ UTextureRenderTarget2D* ULandscapeTexturePatch::ApplyToHeightmap(UE::Landscape::
 		return InCombinedResult;
 	}
 
+	FTextureResource* OutputResource = InCombinedResult->GetResource();
+	FString OutputResourceName = OutputResource->GetResourceName().ToString();
+	if (!bInPerformBlending)
+	{
+		check(RenderParams != nullptr);
+		ULandscapeScratchRenderTarget* WriteRT = RenderParams->MergeRenderContext->GetBlendRenderTargetWrite();
+		// After this point, the render cannot fail so if we're the first in our render layer group to render, we can cycle the blend render targets and 
+		//  start rendering in the write one : 
+		if (RenderParams->NumSuccessfulRenderLayerStepsUntilBlendLayerStep == 0)
+		{
+			RenderParams->MergeRenderContext->CycleBlendRenderTargets(RDGBuilderRecorder);
+			WriteRT = RenderParams->MergeRenderContext->GetBlendRenderTargetWrite();
+			WriteRT->Clear(RDGBuilderRecorder);
+			check(WriteRT->GetCurrentState() == ERHIAccess::RTV);
+		}
+		OutputResource = WriteRT->GetRenderTarget()->GetResource();
+		OutputResourceName = WriteRT->GetDebugName();
+	}
+
 	auto RDGCommand =
-		[InCombinedResult, ShaderParamsToCopy, Patch, DestinationBounds](FRDGBuilder& GraphBuilder)
+		[ OutputResource
+		, OutputResourceName
+		, ShaderParamsToCopy
+		, PatchTextureResource
+		, DestinationBounds
+		, bInPerformBlending](FRDGBuilder& GraphBuilder)
 		{
 			TRACE_CPUPROFILER_EVENT_SCOPE(LandscapeTextureHeightPatch_Render);
 
-			TRefCountPtr<IPooledRenderTarget> DestinationRenderTarget = CreateRenderTarget(InCombinedResult->GetResource()->GetTexture2DRHI(), TEXT("LandscapeTextureHeightPatchOutput"));
+			TRefCountPtr<IPooledRenderTarget> DestinationRenderTarget = CreateRenderTarget(OutputResource->GetTexture2DRHI(), TEXT("LandscapeTextureHeightPatchOutput"));
 			FRDGTextureRef DestinationTexture = GraphBuilder.RegisterExternalTexture(DestinationRenderTarget);
 
-			// Make a copy of the portion of our heightmap input that we're writing to so that we can 
-			// read and write at the same time (needed for blending)
-			FRDGTextureDesc InputCopyDescription = DestinationTexture->Desc;
-			InputCopyDescription.NumMips = 1;
-			InputCopyDescription.Extent = DestinationBounds.Size();
-			FRDGTextureRef InputCopy = GraphBuilder.CreateTexture(InputCopyDescription, TEXT("LandscapeTextureHeightPatchInputCopy"));
-
-			FRHICopyTextureInfo CopyTextureInfo;
-			CopyTextureInfo.SourceMipIndex = 0;
-			CopyTextureInfo.NumMips = 1;
-			CopyTextureInfo.SourcePosition = FIntVector(DestinationBounds.Min.X, DestinationBounds.Min.Y, 0);
-			CopyTextureInfo.Size = FIntVector(InputCopyDescription.Extent.X, InputCopyDescription.Extent.Y, 0);
-
-			AddCopyTexturePass(GraphBuilder, DestinationTexture, InputCopy, CopyTextureInfo);
+			FRDGTextureSRVRef InputCopySRV = nullptr;
+			if (bInPerformBlending)
+			{
+				// Make a copy of the portion of our heightmap input that we're writing to so that we can 
+				// read and write at the same time (needed for blending)
+				FRDGTextureDesc InputCopyDescription = DestinationTexture->Desc;
+				InputCopyDescription.NumMips = 1;
+				InputCopyDescription.Extent = DestinationBounds.Size();
+				FRDGTextureRef InputCopy = GraphBuilder.CreateTexture(InputCopyDescription, TEXT("LandscapeTextureHeightPatchInputCopy"));
+				InputCopySRV = GraphBuilder.CreateSRV(FRDGTextureSRVDesc::CreateForMipLevel(InputCopy, 0));
+
+				FRHICopyTextureInfo CopyTextureInfo;
+				CopyTextureInfo.SourceMipIndex = 0;
+				CopyTextureInfo.NumMips = 1;
+				CopyTextureInfo.SourcePosition = FIntVector(DestinationBounds.Min.X, DestinationBounds.Min.Y, 0);
+				CopyTextureInfo.Size = FIntVector(InputCopyDescription.Extent.X, InputCopyDescription.Extent.Y, 0);
+
+				AddCopyTexturePass(GraphBuilder, DestinationTexture, InputCopy, CopyTextureInfo);
+			}
 
-			FApplyLandscapeTextureHeightPatchPS::FParameters* ShaderParams =
-				GraphBuilder.AllocParameters<FApplyLandscapeTextureHeightPatchPS::FParameters>();
+			FApplyLandscapeTextureHeightPatchPSBase::FParameters* ShaderParams =
+				GraphBuilder.AllocParameters<FApplyLandscapeTextureHeightPatchPSBase::FParameters>();
 			*ShaderParams = ShaderParamsToCopy;
 
-			TRefCountPtr<IPooledRenderTarget> PatchRenderTarget = CreateRenderTarget(Patch->GetTexture2DRHI(), TEXT("LandscapeTextureHeightPatch"));
+			TRefCountPtr<IPooledRenderTarget> PatchRenderTarget = CreateRenderTarget(PatchTextureResource->GetTexture2DRHI(), TEXT("LandscapeTextureHeightPatch"));
 			FRDGTextureRef PatchTexture = GraphBuilder.RegisterExternalTexture(PatchRenderTarget);
 			FRDGTextureSRVRef PatchSRV = GraphBuilder.CreateSRV(FRDGTextureSRVDesc::CreateForMipLevel(PatchTexture, 0));
 			ShaderParams->InHeightPatch = PatchSRV;
 			ShaderParams->InHeightPatchSampler = TStaticSamplerState<SF_Bilinear, AM_Clamp, AM_Clamp>::GetRHI();
 
-			FRDGTextureSRVRef InputCopySRV = GraphBuilder.CreateSRV(FRDGTextureSRVDesc::CreateForMipLevel(InputCopy, 0));
-			ShaderParams->InSourceHeightmap = InputCopySRV;
-			ShaderParams->InSourceHeightmapOffset = DestinationBounds.Min;
+			ShaderParams->InSourceHeightmap = GraphBuilder.CreateSRV(FRDGTextureSRVDesc::Create(GSystemTextures.GetBlackDummy(GraphBuilder)));
+			ShaderParams->InSourceHeightmapOffset = FIntPoint(ForceInit);
+			if (bInPerformBlending)
+			{
+				ShaderParams->InSourceHeightmap = InputCopySRV;
+				ShaderParams->InSourceHeightmapOffset = DestinationBounds.Min;
+			}
 
 			ShaderParams->RenderTargets[0] = FRenderTargetBinding(DestinationTexture, ERenderTargetLoadAction::ENoAction, /*InMipIndex = */0);
 
-			FApplyLandscapeTextureHeightPatchPS::AddToRenderGraph(GraphBuilder, ShaderParams, DestinationBounds);
+			if (bInPerformBlending)
+			{
+				FApplyLandscapeTextureHeightPatchPS</*bPerformBlending = */true>::AddToRenderGraph(RDG_EVENT_NAME("RenderTextureHeightPatch -> %s", *OutputResourceName), GraphBuilder, ShaderParams, DestinationBounds);
+			}
+			else
+			{
+				FApplyLandscapeTextureHeightPatchPS</*bPerformBlending = */false>::AddToRenderGraph(RDG_EVENT_NAME("RenderTextureHeightPatch -> %s", *OutputResourceName), GraphBuilder, ShaderParams, DestinationBounds);
+			}
 		};
 	// We need to specify the final state of the external texture to prevent the graph builder from transitioning it to SRVMask :
-	RDGBuilderRecorder.EnqueueRDGCommand(RDGCommand, { { InCombinedResult->GetResource(), (OutputAccess == ERHIAccess::None) ? ERHIAccess::RTV : OutputAccess} });
+	RDGBuilderRecorder.EnqueueRDGCommand(RDGCommand, { { OutputResource, (OutputAccess == ERHIAccess::None) ? ERHIAccess::RTV : OutputAccess} });
 
+	bOutDidRenderSomething = true;
 	return InCombinedResult;
 }
 
-void ULandscapeTexturePatch::ApplyToWeightmap(UE::Landscape::FRDGBuilderRecorder& RDGBuilderRecorder, ULandscapeWeightPatchTextureInfo* PatchInfo,
+void ULandscapeTexturePatch::ApplyToWeightmap(bool bInPerformBlending, FRenderParams* RenderParams, UE::Landscape::FRDGBuilderRecorder& RDGBuilderRecorder, ULandscapeWeightPatchTextureInfo* PatchInfo,
 	FTextureResource* InMergedLandscapeTextureResource, int32 LandscapeTextureSliceIndex,
-	const FIntPoint& LandscapeTextureResolution, const FTransform& LandscapeHeightmapToWorld, ERHIAccess OutputAccess)
+	const FIntPoint& LandscapeTextureResolution, const FTransform& LandscapeHeightmapToWorld, bool& bOutDidRenderSomething, ERHIAccess OutputAccess)
 {
 	using namespace UE::Landscape;
+	bOutDidRenderSomething = false;
 
 	if (!PatchInfo)
 	{
@@ -449,8 +540,8 @@ void ULandscapeTexturePatch::ApplyToWeightmap(UE::Landscape::FRDGBuilderRecorder
 		return;
 	}
 
-	FTextureResource* Patch = PatchUObject->GetResource();
-	if (!Patch)
+	FTextureResource* PatchTextureResource = PatchUObject->GetResource();
+	if (!PatchTextureResource)
 	{
 		return;
 	}
@@ -460,7 +551,7 @@ void ULandscapeTexturePatch::ApplyToWeightmap(UE::Landscape::FRDGBuilderRecorder
 	FApplyLandscapeTextureWeightPatchPSParameters ShaderParamsToCopy;
 	FIntRect DestinationBounds;
 
-	GetWeightShaderParams(LandscapeHeightmapToWorld, FIntPoint(Patch->GetSizeX(), Patch->GetSizeY()), 
+	GetWeightShaderParams(LandscapeHeightmapToWorld, FIntPoint(PatchTextureResource->GetSizeX(), PatchTextureResource->GetSizeY()),
 		LandscapeTextureResolution, PatchInfo, ShaderParamsToCopy, DestinationBounds);
 
 	if (DestinationBounds.Area() <= 0)
@@ -469,54 +560,98 @@ void ULandscapeTexturePatch::ApplyToWeightmap(UE::Landscape::FRDGBuilderRecorder
 		return;
 	}
 
+	FTextureResource* OutputResource = InMergedLandscapeTextureResource;
+	FString OutputResourceName = OutputResource->GetResourceName().ToString();
+	if (!bInPerformBlending)
+	{
+		check(RenderParams != nullptr);
+		ULandscapeScratchRenderTarget* WriteRT = RenderParams->MergeRenderContext->GetBlendRenderTargetWrite();
+		// After this point, the render cannot fail so if we're the first in our render layer group to render, we can cycle the blend render targets and 
+		//  start rendering in the write one : 
+		if (RenderParams->NumSuccessfulRenderLayerStepsUntilBlendLayerStep == 0)
+		{
+			RenderParams->MergeRenderContext->CycleBlendRenderTargets(RDGBuilderRecorder);
+			WriteRT = RenderParams->MergeRenderContext->GetBlendRenderTargetWrite();
+			WriteRT->Clear(RDGBuilderRecorder);
+			check(WriteRT->GetCurrentState() == ERHIAccess::RTV);
+		}
+		OutputResource = WriteRT->GetRenderTarget()->GetResource();
+		OutputResourceName = WriteRT->GetDebugName();
+	}
+
 	auto RDGCommand =
-		[InMergedLandscapeTextureResource, LandscapeTextureSliceIndex, ShaderParamsToCopy, Patch, DestinationBounds](FRDGBuilder& GraphBuilder)
-	{
-		TRACE_CPUPROFILER_EVENT_SCOPE(LandscapeTextureWeightPatch_Render);
-
-			TRefCountPtr<IPooledRenderTarget> DestinationRenderTarget = CreateRenderTarget(InMergedLandscapeTextureResource->GetTextureRHI(), TEXT("LandscapeTextureWeightPatchOutput"));
-		FRDGTextureRef DestinationTexture = GraphBuilder.RegisterExternalTexture(DestinationRenderTarget);
-
-		// Make a copy of the portion of our weightmap input that we're writing to so that we can 
-		// read and write at the same time (needed for blending)
-		FRDGTextureDesc InputCopyDescription = DestinationTexture->Desc;
-		InputCopyDescription.Dimension = ETextureDimension::Texture2D;
-		InputCopyDescription.ArraySize = 1;
-		InputCopyDescription.NumMips = 1;
-		InputCopyDescription.Extent = DestinationBounds.Size();
-		FRDGTextureRef InputCopy = GraphBuilder.CreateTexture(InputCopyDescription, TEXT("LandscapeTextureWeightPatchInputCopy"));
-
-		FRHICopyTextureInfo CopyTextureInfo;
-		CopyTextureInfo.SourceMipIndex = 0;
-		CopyTextureInfo.NumMips = 1;
-		CopyTextureInfo.SourceSliceIndex = LandscapeTextureSliceIndex;
-		CopyTextureInfo.NumSlices = 1;
-		CopyTextureInfo.SourcePosition = FIntVector(DestinationBounds.Min.X, DestinationBounds.Min.Y, 0);
-		CopyTextureInfo.Size = FIntVector(InputCopyDescription.Extent.X, InputCopyDescription.Extent.Y, 0);
-		AddCopyTexturePass(GraphBuilder, DestinationTexture, InputCopy, CopyTextureInfo);
-
-		FApplyLandscapeTextureWeightPatchPS::FParameters* ShaderParams =
-			GraphBuilder.AllocParameters<FApplyLandscapeTextureWeightPatchPS::FParameters>();
-		*ShaderParams = ShaderParamsToCopy;
-
-		TRefCountPtr<IPooledRenderTarget> PatchRenderTarget = CreateRenderTarget(Patch->GetTexture2DRHI(), TEXT("LandscapeTextureWeightPatch"));
-		FRDGTextureRef PatchTexture = GraphBuilder.RegisterExternalTexture(PatchRenderTarget);
-		FRDGTextureSRVRef PatchSRV = GraphBuilder.CreateSRV(FRDGTextureSRVDesc::CreateForMipLevel(PatchTexture, 0));
-		ShaderParams->InWeightPatch = PatchSRV;
-		ShaderParams->InWeightPatchSampler = TStaticSamplerState<SF_Bilinear, AM_Clamp, AM_Clamp>::GetRHI();
-
-		FRDGTextureSRVRef InputCopySRV = GraphBuilder.CreateSRV(FRDGTextureSRVDesc::CreateForMipLevel(InputCopy, 0));
-		ShaderParams->InSourceWeightmap = InputCopySRV;
-		ShaderParams->InSourceWeightmapCoordOffset = DestinationBounds.Min;
-
-		ShaderParams->RenderTargets[0] = FRenderTargetBinding(DestinationTexture, ERenderTargetLoadAction::ENoAction,
-			/*InMipIndex = */0, LandscapeTextureSliceIndex);
-
-		FApplyLandscapeTextureWeightPatchPS::AddToRenderGraph(GraphBuilder, ShaderParams, DestinationBounds);
+		[ OutputResource
+		, OutputResourceName
+		, LandscapeTextureSliceIndex
+		, ShaderParamsToCopy
+		, PatchTextureResource
+		, DestinationBounds
+		, bInPerformBlending](FRDGBuilder& GraphBuilder)
+		{
+			TRACE_CPUPROFILER_EVENT_SCOPE(LandscapeTextureWeightPatch_Render);
+
+			TRefCountPtr<IPooledRenderTarget> DestinationRenderTarget = CreateRenderTarget(OutputResource->GetTextureRHI(), TEXT("LandscapeTextureWeightPatchOutput"));
+			FRDGTextureRef DestinationTexture = GraphBuilder.RegisterExternalTexture(DestinationRenderTarget);
+
+			FRDGTextureSRVRef InputCopySRV = nullptr;
+			if (bInPerformBlending)
+			{
+				// Make a copy of the portion of our weightmap input that we're writing to so that we can 
+				// read and write at the same time (needed for blending)
+				FRDGTextureDesc InputCopyDescription = DestinationTexture->Desc;
+				InputCopyDescription.Dimension = ETextureDimension::Texture2D;
+				InputCopyDescription.ArraySize = 1;
+				InputCopyDescription.NumMips = 1;
+				InputCopyDescription.Extent = DestinationBounds.Size();
+				FRDGTextureRef InputCopy = GraphBuilder.CreateTexture(InputCopyDescription, TEXT("LandscapeTextureWeightPatchInputCopy"));
+
+				FRHICopyTextureInfo CopyTextureInfo;
+				CopyTextureInfo.SourceMipIndex = 0;
+				CopyTextureInfo.NumMips = 1;
+				CopyTextureInfo.SourceSliceIndex = LandscapeTextureSliceIndex;
+				CopyTextureInfo.NumSlices = 1;
+				CopyTextureInfo.SourcePosition = FIntVector(DestinationBounds.Min.X, DestinationBounds.Min.Y, 0);
+				CopyTextureInfo.Size = FIntVector(InputCopyDescription.Extent.X, InputCopyDescription.Extent.Y, 0);
+				AddCopyTexturePass(GraphBuilder, DestinationTexture, InputCopy, CopyTextureInfo);
+
+				InputCopySRV = GraphBuilder.CreateSRV(FRDGTextureSRVDesc::CreateForMipLevel(InputCopy, 0));
+			}
+
+			FApplyLandscapeTextureWeightPatchPSBase::FParameters* ShaderParams =
+				GraphBuilder.AllocParameters<FApplyLandscapeTextureWeightPatchPSBase::FParameters>();
+			*ShaderParams = ShaderParamsToCopy;
+
+			TRefCountPtr<IPooledRenderTarget> PatchRenderTarget = CreateRenderTarget(PatchTextureResource->GetTexture2DRHI(), TEXT("LandscapeTextureWeightPatch"));
+			FRDGTextureRef PatchTexture = GraphBuilder.RegisterExternalTexture(PatchRenderTarget);
+			FRDGTextureSRVRef PatchSRV = GraphBuilder.CreateSRV(FRDGTextureSRVDesc::CreateForMipLevel(PatchTexture, 0));
+			ShaderParams->InWeightPatch = PatchSRV;
+			ShaderParams->InWeightPatchSampler = TStaticSamplerState<SF_Bilinear, AM_Clamp, AM_Clamp>::GetRHI();
+
+			ShaderParams->InSourceWeightmap = GraphBuilder.CreateSRV(FRDGTextureSRVDesc::Create(GSystemTextures.GetBlackDummy(GraphBuilder)));
+			ShaderParams->InSourceWeightmapCoordOffset = FIntPoint(ForceInit);
+			if (bInPerformBlending)
+			{
+				ShaderParams->InSourceWeightmap = InputCopySRV;
+				ShaderParams->InSourceWeightmapCoordOffset = DestinationBounds.Min;
+			}
+
+			ShaderParams->RenderTargets[0] = FRenderTargetBinding(DestinationTexture, ERenderTargetLoadAction::ENoAction,
+				/*InMipIndex = */0, LandscapeTextureSliceIndex);
+
+			if (bInPerformBlending)
+			{
+				FApplyLandscapeTextureWeightPatchPS</*bPerformBlending = */true>::AddToRenderGraph(RDG_EVENT_NAME("RenderTextureWeightPatch -> %s", *OutputResourceName), GraphBuilder, ShaderParams, DestinationBounds);
+			}
+			else
+			{
+				FApplyLandscapeTextureWeightPatchPS</*bPerformBlending = */false>::AddToRenderGraph(RDG_EVENT_NAME("RenderTextureWeightPatch -> %s", *OutputResourceName), GraphBuilder, ShaderParams, DestinationBounds);
+			}
 		};
 
 	// We need to specify the final state of the external texture to prevent the graph builder from transitioning it to SRVMask :
 	RDGBuilderRecorder.EnqueueRDGCommand(RDGCommand, { { InMergedLandscapeTextureResource, (OutputAccess == ERHIAccess::None) ? ERHIAccess::RTV : OutputAccess} });
+	
+	bOutDidRenderSomething = true;
 }
 
 void ULandscapeTexturePatch::GetCommonShaderParams(const FTransform& LandscapeHeightmapToWorldIn,
@@ -632,7 +767,7 @@ void ULandscapeTexturePatch::GetHeightShaderParams(const FTransform& LandscapeHe
 	ParamsOut.InBlendMode = static_cast<uint32>(BlendMode);
 
 	// Pack our booleans into a bitfield
-	using EShaderFlags = FApplyLandscapeTextureHeightPatchPS::EFlags;
+	using EShaderFlags = FApplyLandscapeTextureHeightPatchPSBase::EFlags;
 	EShaderFlags Flags = EShaderFlags::None;
 
 	Flags |= (FalloffMode == ELandscapeTexturePatchFalloffMode::RoundedRectangle) ?
@@ -664,7 +799,7 @@ void ULandscapeTexturePatch::GetWeightShaderParams(
 	ParamsOut.InBlendMode = static_cast<uint32>(WeightPatchInfo->bOverrideBlendMode ? WeightPatchInfo->OverrideBlendMode : BlendMode);
 
 	// Pack our booleans into a bitfield
-	using EShaderFlags = FApplyLandscapeTextureHeightPatchPS::EFlags;
+	using EShaderFlags = FApplyLandscapeTextureHeightPatchPSBase::EFlags;
 	EShaderFlags Flags = EShaderFlags::None;
 
 	Flags |= (FalloffMode == ELandscapeTexturePatchFalloffMode::RoundedRectangle) ?
diff --git a/Engine/Plugins/Experimental/LandscapePatch/Source/LandscapePatch/Private/LandscapeTexturePatchPS.cpp b/Engine/Plugins/Experimental/LandscapePatch/Source/LandscapePatch/Private/LandscapeTexturePatchPS.cpp
--- a/Engine/Plugins/Experimental/LandscapePatch/Source/LandscapePatch/Private/LandscapeTexturePatchPS.cpp
+++ b/Engine/Plugins/Experimental/LandscapePatch/Source/LandscapePatch/Private/LandscapeTexturePatchPS.cpp
@@ -9,19 +9,19 @@
 namespace UE::Landscape
 {
 
-bool FApplyLandscapeTextureHeightPatchPS::ShouldCompilePermutation(const FGlobalShaderPermutationParameters& Parameters)
+bool FApplyLandscapeTextureHeightPatchPSBase::ShouldCompilePermutation(const FGlobalShaderPermutationParameters& Parameters)
 {
 	return DoesPlatformSupportEditLayers(Parameters.Platform);
 }
 
-void FApplyLandscapeTextureHeightPatchPS::ModifyCompilationEnvironment(const FGlobalShaderPermutationParameters& Parameters, FShaderCompilerEnvironment& OutEnvironment)
+void FApplyLandscapeTextureHeightPatchPSBase::ModifyCompilationEnvironment(const FGlobalShaderPermutationParameters& Parameters, FShaderCompilerEnvironment& OutEnvironment)
 {
 	OutEnvironment.SetDefine(TEXT("APPLY_HEIGHT_PATCH"), 1);
 
 	// Make our flag choices match in the shader.
-	OutEnvironment.SetDefine(TEXT("RECTANGULAR_FALLOFF_FLAG"), static_cast<uint8>(FApplyLandscapeTextureHeightPatchPS::EFlags::RectangularFalloff));
-	OutEnvironment.SetDefine(TEXT("APPLY_PATCH_ALPHA_FLAG"), static_cast<uint8>(FApplyLandscapeTextureHeightPatchPS::EFlags::ApplyPatchAlpha));
-	OutEnvironment.SetDefine(TEXT("INPUT_IS_PACKED_HEIGHT_FLAG"), static_cast<uint8>(FApplyLandscapeTextureHeightPatchPS::EFlags::InputIsPackedHeight));
+	OutEnvironment.SetDefine(TEXT("RECTANGULAR_FALLOFF_FLAG"), static_cast<uint8>(EFlags::RectangularFalloff));
+	OutEnvironment.SetDefine(TEXT("APPLY_PATCH_ALPHA_FLAG"), static_cast<uint8>(EFlags::ApplyPatchAlpha));
+	OutEnvironment.SetDefine(TEXT("INPUT_IS_PACKED_HEIGHT_FLAG"), static_cast<uint8>(EFlags::InputIsPackedHeight));
 
 	OutEnvironment.SetDefine(TEXT("ALPHA_BLEND_MODE"), static_cast<uint8>(ELandscapeTexturePatchBlendMode::AlphaBlend));
 	OutEnvironment.SetDefine(TEXT("ADDITIVE_MODE"), static_cast<uint8>(ELandscapeTexturePatchBlendMode::Additive));
@@ -29,20 +29,6 @@ void FApplyLandscapeTextureHeightPatchPS::ModifyCompilationEnvironment(const FGl
 	OutEnvironment.SetDefine(TEXT("MAX_MODE"), static_cast<uint8>(ELandscapeTexturePatchBlendMode::Max));
 }
 
-void FApplyLandscapeTextureHeightPatchPS::AddToRenderGraph(FRDGBuilder& GraphBuilder, FParameters* InParameters, const FIntRect& DestinationBounds)
-{
-	FGlobalShaderMap* ShaderMap = GetGlobalShaderMap(GMaxRHIFeatureLevel);
-	TShaderMapRef<FApplyLandscapeTextureHeightPatchPS> PixelShader(ShaderMap);
-
-	FPixelShaderUtils::AddFullscreenPass(
-		GraphBuilder,
-		ShaderMap,
-		RDG_EVENT_NAME("LandscapeTextureHeightPatch"),
-		PixelShader,
-		InParameters,
-		DestinationBounds);
-}
-
 bool FOffsetHeightmapPS::ShouldCompilePermutation(const FGlobalShaderPermutationParameters& Parameters)
 {
 	return DoesPlatformSupportEditLayers(Parameters.Platform);
@@ -172,18 +158,18 @@ void FConvertBackFromNativeLandscapePatchPS::AddToRenderGraph(FRDGBuilder& Graph
 		FIntRect(0, 0, DestinationSize.X, DestinationSize.Y));
 }
 
-bool FApplyLandscapeTextureWeightPatchPS::ShouldCompilePermutation(const FGlobalShaderPermutationParameters& Parameters)
+bool FApplyLandscapeTextureWeightPatchPSBase::ShouldCompilePermutation(const FGlobalShaderPermutationParameters& Parameters)
 {
 	return DoesPlatformSupportEditLayers(Parameters.Platform);
 }
 
-void FApplyLandscapeTextureWeightPatchPS::ModifyCompilationEnvironment(const FGlobalShaderPermutationParameters& Parameters, FShaderCompilerEnvironment& OutEnvironment)
+void FApplyLandscapeTextureWeightPatchPSBase::ModifyCompilationEnvironment(const FGlobalShaderPermutationParameters& Parameters, FShaderCompilerEnvironment& OutEnvironment)
 {
 	OutEnvironment.SetDefine(TEXT("APPLY_WEIGHT_PATCH"), 1);
 
 	// Make our flag choices match in the shader.
-	OutEnvironment.SetDefine(TEXT("RECTANGULAR_FALLOFF_FLAG"), static_cast<uint8>(FApplyLandscapeTextureWeightPatchPS::EFlags::RectangularFalloff));
-	OutEnvironment.SetDefine(TEXT("APPLY_PATCH_ALPHA_FLAG"), static_cast<uint8>(FApplyLandscapeTextureWeightPatchPS::EFlags::ApplyPatchAlpha));
+	OutEnvironment.SetDefine(TEXT("RECTANGULAR_FALLOFF_FLAG"), static_cast<uint8>(EFlags::RectangularFalloff));
+	OutEnvironment.SetDefine(TEXT("APPLY_PATCH_ALPHA_FLAG"), static_cast<uint8>(EFlags::ApplyPatchAlpha));
 
 	OutEnvironment.SetDefine(TEXT("ALPHA_BLEND_MODE"), static_cast<uint8>(ELandscapeTexturePatchBlendMode::AlphaBlend));
 	OutEnvironment.SetDefine(TEXT("ADDITIVE_MODE"), static_cast<uint8>(ELandscapeTexturePatchBlendMode::Additive));
@@ -191,20 +177,6 @@ void FApplyLandscapeTextureWeightPatchPS::ModifyCompilationEnvironment(const FGl
 	OutEnvironment.SetDefine(TEXT("MAX_MODE"), static_cast<uint8>(ELandscapeTexturePatchBlendMode::Max));
 }
 
-void FApplyLandscapeTextureWeightPatchPS::AddToRenderGraph(FRDGBuilder& GraphBuilder, FParameters* InParameters, const FIntRect& DestinationBounds)
-{
-	FGlobalShaderMap* ShaderMap = GetGlobalShaderMap(GMaxRHIFeatureLevel);
-	TShaderMapRef<FApplyLandscapeTextureWeightPatchPS> PixelShader(ShaderMap);
-
-	FPixelShaderUtils::AddFullscreenPass(
-		GraphBuilder,
-		ShaderMap,
-		RDG_EVENT_NAME("LandscapeTextureWeightPatch"),
-		PixelShader,
-		InParameters,
-		DestinationBounds);
-}
-
 bool FReinitializeLandscapePatchPS::ShouldCompilePermutation(const FGlobalShaderPermutationParameters& Parameters)
 {
 	return DoesPlatformSupportEditLayers(Parameters.Platform);
@@ -236,11 +208,13 @@ void FReinitializeLandscapePatchPS::AddToRenderGraph(FRDGBuilder& GraphBuilder,
 
 }//end UE::Landscape
 
-IMPLEMENT_GLOBAL_SHADER(UE::Landscape::FApplyLandscapeTextureHeightPatchPS, "/Plugin/LandscapePatch/Private/LandscapeTexturePatchPS.usf", "ApplyLandscapeTextureHeightPatch", SF_Pixel);
+IMPLEMENT_SHADER_TYPE(template<>, UE::Landscape::FApplyLandscapeTextureHeightPatchPS</*bPerformBlending = */true>, TEXT("/Plugin/LandscapePatch/Private/LandscapeTexturePatchPS.usf"), TEXT("ApplyLandscapeTextureHeightPatch"), SF_Pixel);
+IMPLEMENT_SHADER_TYPE(template<>, UE::Landscape::FApplyLandscapeTextureHeightPatchPS</*bPerformBlending = */false>, TEXT("/Plugin/LandscapePatch/Private/LandscapeTexturePatchPS.usf"), TEXT("ApplyLandscapeTextureHeightPatch"), SF_Pixel);
 IMPLEMENT_GLOBAL_SHADER(UE::Landscape::FOffsetHeightmapPS, "/Plugin/LandscapePatch/Private/LandscapeTexturePatchPS.usf", "ApplyOffsetToHeightmap", SF_Pixel);
 IMPLEMENT_GLOBAL_SHADER(UE::Landscape::FSimpleTextureCopyPS, "/Plugin/LandscapePatch/Private/LandscapeTexturePatchPS.usf", "SimpleTextureCopy", SF_Pixel);
 IMPLEMENT_GLOBAL_SHADER(UE::Landscape::FConvertToNativeLandscapePatchPS, "/Plugin/LandscapePatch/Private/LandscapeTexturePatchPS.usf", "ConvertToNativeLandscapePatch", SF_Pixel);
 IMPLEMENT_GLOBAL_SHADER(UE::Landscape::FConvertBackFromNativeLandscapePatchPS, "/Plugin/LandscapePatch/Private/LandscapeTexturePatchPS.usf", "ConvertBackFromNativeLandscapePatch", SF_Pixel);
-IMPLEMENT_GLOBAL_SHADER(UE::Landscape::FApplyLandscapeTextureWeightPatchPS, "/Plugin/LandscapePatch/Private/LandscapeTexturePatchPS.usf", "ApplyLandscapeTextureWeightPatch", SF_Pixel);
+IMPLEMENT_SHADER_TYPE(template<>, UE::Landscape::FApplyLandscapeTextureWeightPatchPS</*bPerformBlending = */true>, TEXT("/Plugin/LandscapePatch/Private/LandscapeTexturePatchPS.usf"), TEXT("ApplyLandscapeTextureWeightPatch"), SF_Pixel);
+IMPLEMENT_SHADER_TYPE(template<>, UE::Landscape::FApplyLandscapeTextureWeightPatchPS</*bPerformBlending = */false>, TEXT("/Plugin/LandscapePatch/Private/LandscapeTexturePatchPS.usf"), TEXT("ApplyLandscapeTextureWeightPatch"), SF_Pixel);
 IMPLEMENT_GLOBAL_SHADER(UE::Landscape::FReinitializeLandscapePatchPS, "/Plugin/LandscapePatch/Private/LandscapeTexturePatchPS.usf", "ReinitializePatch", SF_Pixel);
 
diff --git a/Engine/Plugins/Experimental/LandscapePatch/Source/LandscapePatch/Private/LandscapeTexturePatchPS.h b/Engine/Plugins/Experimental/LandscapePatch/Source/LandscapePatch/Private/LandscapeTexturePatchPS.h
--- a/Engine/Plugins/Experimental/LandscapePatch/Source/LandscapePatch/Private/LandscapeTexturePatchPS.h
+++ b/Engine/Plugins/Experimental/LandscapePatch/Source/LandscapePatch/Private/LandscapeTexturePatchPS.h
@@ -6,6 +6,7 @@
 #include "GlobalShader.h"
 #include "ShaderParameterStruct.h"
 #include "LandscapeTexturePatch.h"
+#include "PixelShaderUtils.h"
 
 
 struct FLandscapeHeightPatchConvertToNativeParams;
@@ -16,13 +17,14 @@ namespace UE::Landscape
 /**
  * Shader that applies a texture-based height patch to a landscape heightmap.
  */
-class FApplyLandscapeTextureHeightPatchPS : public FGlobalShader
+class FApplyLandscapeTextureHeightPatchPSBase : public FGlobalShader
 {
-	DECLARE_GLOBAL_SHADER(FApplyLandscapeTextureHeightPatchPS);
-	SHADER_USE_PARAMETER_STRUCT(FApplyLandscapeTextureHeightPatchPS, FGlobalShader);
+	// TODO: We could consider exposing an additional global alpha setting that we can use to pass in the given
+	// edit layer alpha value... On the other hand, we currently don't bother doing this in any existing blueprint
+	// brushes, and it would be hard to support in a way that doesn't require each blueprint brush to respect it
+	// individually... Not clear whether this is something worth doing yet.
 
 public:
-
 	// Flags that get packed into a bitfield because we're not allowed to use bool shader parameters:
 	enum class EFlags : uint8
 	{
@@ -39,11 +41,6 @@ public:
 		InputIsPackedHeight = 1 << 2
 	};
 
-	// TODO: We could consider exposing an additional global alpha setting that we can use to pass in the given
-	// edit layer alpha value... On the other hand, we currently don't bother doing this in any existing blueprint
-	// brushes, and it would be hard to support in a way that doesn't require each blueprint brush to respect it
-	// individually... Not clear whether this is something worth doing yet.
-
 	BEGIN_SHADER_PARAMETER_STRUCT(FParameters, )
 		SHADER_PARAMETER_RDG_TEXTURE_SRV(Texture2D<float4>, InSourceHeightmap)
 		// Offset of the source heightmap relative to the 0,0 location in the destination heightmap, because
@@ -74,16 +71,55 @@ public:
 		RENDER_TARGET_BINDING_SLOTS() // Holds our output
 	END_SHADER_PARAMETER_STRUCT()
 
+	FApplyLandscapeTextureHeightPatchPSBase() {}
+	FApplyLandscapeTextureHeightPatchPSBase(const ShaderMetaType::CompiledShaderInitializerType& Initializer)
+		: FGlobalShader(Initializer)
+	{}
+
+public:
 	static bool ShouldCompilePermutation(const FGlobalShaderPermutationParameters& Parameters);
 	static void ModifyCompilationEnvironment(const FGlobalShaderPermutationParameters& InParameters, FShaderCompilerEnvironment& OutEnvironment);
 
-	static void AddToRenderGraph(FRDGBuilder& GraphBuilder, FParameters* InParameters, const FIntRect& DestinationBounds);
 };
 
+ENUM_CLASS_FLAGS(FApplyLandscapeTextureHeightPatchPSBase::EFlags);
+
 // We need an alias for FApplyLandscapeTextureHeightPatchPS::FParameters so that it can be forward-declared in another header file (cannot forward-declare internal classes)
-class FApplyLandscapeTextureHeightPatchPSParameters : public FApplyLandscapeTextureHeightPatchPS::FParameters {};
+class FApplyLandscapeTextureHeightPatchPSParameters : public FApplyLandscapeTextureHeightPatchPSBase::FParameters {};
+
+template <bool bPerformBlending = false>
+class FApplyLandscapeTextureHeightPatchPS : public FApplyLandscapeTextureHeightPatchPSBase
+{
+	using Super = FApplyLandscapeTextureHeightPatchPSBase;
+	DECLARE_GLOBAL_SHADER(FApplyLandscapeTextureHeightPatchPS);
+	SHADER_USE_PARAMETER_STRUCT(FApplyLandscapeTextureHeightPatchPS, FApplyLandscapeTextureHeightPatchPSBase);
+
+public:
 
-ENUM_CLASS_FLAGS(FApplyLandscapeTextureHeightPatchPS::EFlags);
+	static void ModifyCompilationEnvironment(const FGlobalShaderPermutationParameters& InParameters, FShaderCompilerEnvironment& OutEnvironment)
+	{
+		Super::ModifyCompilationEnvironment(InParameters, OutEnvironment);
+
+		if (bPerformBlending)
+		{
+			OutEnvironment.SetDefine(TEXT("PERFORM_BLENDING"), 1);
+		}
+	}
+
+	static void AddToRenderGraph(FRDGEventName&& InRDGEventName, FRDGBuilder& GraphBuilder, FParameters* InParameters, const FIntRect& DestinationBounds)
+	{
+		FGlobalShaderMap* ShaderMap = GetGlobalShaderMap(GMaxRHIFeatureLevel);
+		TShaderMapRef<FApplyLandscapeTextureHeightPatchPS<bPerformBlending>> PixelShader(ShaderMap);
+
+		FPixelShaderUtils::AddFullscreenPass(
+			GraphBuilder,
+			ShaderMap,
+			MoveTemp(InRDGEventName),
+			PixelShader,
+			InParameters,
+			DestinationBounds);
+	}
+};
 
 /**
  * Simple shader that just offsets each height value in a height patch by a constant.
@@ -161,13 +197,14 @@ public:
 		FRDGTextureRef DestinationTexture, const FLandscapeHeightPatchConvertToNativeParams& Params);
 };
 
-class FApplyLandscapeTextureWeightPatchPS : public FGlobalShader
+class FApplyLandscapeTextureWeightPatchPSBase : public FGlobalShader
 {
-	DECLARE_GLOBAL_SHADER(FApplyLandscapeTextureWeightPatchPS);
-	SHADER_USE_PARAMETER_STRUCT(FApplyLandscapeTextureWeightPatchPS, FGlobalShader);
+	// TODO: We could consider exposing an additional global alpha setting that we can use to pass in the given
+	// edit layer alpha value... On the other hand, we currently don't bother doing this in any existing blueprint
+	// brushes, and it would be hard to support in a way that doesn't require each blueprint brush to respect it
+	// individually... Not clear whether this is something worth doing yet.
 
 public:
-
 	// Flags that get packed into a bitfield because we're not allowed to use bool shader parameters:
 	enum class EFlags : uint8
 	{
@@ -180,11 +217,6 @@ public:
 		ApplyPatchAlpha = 1 << 1,
 	};
 
-	// TODO: We could consider exposing an additional global alpha setting that we can use to pass in the given
-	// edit layer alpha value... On the other hand, we currently don't bother doing this in any existing blueprint
-	// brushes, and it would be hard to support in a way that doesn't require each blueprint brush to respect it
-	// individually... Not clear whether this is something worth doing yet.
-
 	BEGIN_SHADER_PARAMETER_STRUCT(FParameters, )
 		SHADER_PARAMETER_RDG_TEXTURE_SRV(Texture2D<float4>, InSourceWeightmap)
 		// Offset of the source weightmap relative to the 0,0 location in the destination weightmap, because
@@ -209,16 +241,54 @@ public:
 		RENDER_TARGET_BINDING_SLOTS() // Holds our output
 	END_SHADER_PARAMETER_STRUCT()
 
+	FApplyLandscapeTextureWeightPatchPSBase() {}
+	FApplyLandscapeTextureWeightPatchPSBase(const ShaderMetaType::CompiledShaderInitializerType& Initializer)
+		: FGlobalShader(Initializer)
+	{ }
+
+public:
 	static bool ShouldCompilePermutation(const FGlobalShaderPermutationParameters& Parameters);
 	static void ModifyCompilationEnvironment(const FGlobalShaderPermutationParameters& InParameters, FShaderCompilerEnvironment& OutEnvironment);
-
-	static void AddToRenderGraph(FRDGBuilder& GraphBuilder, FParameters* InParameters, const FIntRect& DestinationBounds);
 };
 
+ENUM_CLASS_FLAGS(FApplyLandscapeTextureWeightPatchPSBase::EFlags);
+
 // We need an alias for FApplyLandscapeTextureWeightPatchPSParameters::FParameters so that it can be forward-declared in another header file (cannot forward-declare internal classes)
-class FApplyLandscapeTextureWeightPatchPSParameters : public FApplyLandscapeTextureWeightPatchPS::FParameters {};
+class FApplyLandscapeTextureWeightPatchPSParameters : public FApplyLandscapeTextureWeightPatchPSBase::FParameters {};
+
+template <bool bPerformBlending = false>
+class FApplyLandscapeTextureWeightPatchPS : public FApplyLandscapeTextureWeightPatchPSBase
+{
+	using Super = FApplyLandscapeTextureWeightPatchPSBase;
+	DECLARE_GLOBAL_SHADER(FApplyLandscapeTextureWeightPatchPS);
+	SHADER_USE_PARAMETER_STRUCT(FApplyLandscapeTextureWeightPatchPS, FApplyLandscapeTextureWeightPatchPSBase);
+
+public:
+
+	static void ModifyCompilationEnvironment(const FGlobalShaderPermutationParameters& InParameters, FShaderCompilerEnvironment& OutEnvironment)
+	{
+		Super::ModifyCompilationEnvironment(InParameters, OutEnvironment);
 
-ENUM_CLASS_FLAGS(FApplyLandscapeTextureWeightPatchPS::EFlags);
+		if (bPerformBlending)
+		{
+			OutEnvironment.SetDefine(TEXT("PERFORM_BLENDING"), 1);
+		}
+	}
+
+	static void AddToRenderGraph(FRDGEventName&& InRDGEventName, FRDGBuilder& GraphBuilder, FParameters* InParameters, const FIntRect& DestinationBounds)
+	{
+		FGlobalShaderMap* ShaderMap = GetGlobalShaderMap(GMaxRHIFeatureLevel);
+		TShaderMapRef<FApplyLandscapeTextureWeightPatchPS<bPerformBlending>> PixelShader(ShaderMap);
+
+		FPixelShaderUtils::AddFullscreenPass(
+			GraphBuilder,
+			ShaderMap,
+			MoveTemp(InRDGEventName),
+			PixelShader,
+			InParameters,
+			DestinationBounds);
+	}
+};
 
 class FReinitializeLandscapePatchPS : public FGlobalShader
 {
diff --git a/Engine/Plugins/Experimental/LandscapePatch/Source/LandscapePatch/Public/LandscapeCircleHeightPatch.h b/Engine/Plugins/Experimental/LandscapePatch/Source/LandscapePatch/Public/LandscapeCircleHeightPatch.h
--- a/Engine/Plugins/Experimental/LandscapePatch/Source/LandscapePatch/Public/LandscapeCircleHeightPatch.h
+++ b/Engine/Plugins/Experimental/LandscapePatch/Source/LandscapePatch/Public/LandscapeCircleHeightPatch.h
@@ -27,9 +27,10 @@ class ULandscapeCircleHeightPatch : public ULandscapePatchComponent
 
 public:
 
-	virtual UTextureRenderTarget2D* RenderLayer_Native(const FLandscapeBrushParameters& InParameters, const FTransform& HeightmapToWorld) override;
 
 #if WITH_EDITOR
+	virtual UTextureRenderTarget2D* RenderLayer_Native(const FLandscapeBrushParameters& InParameters, const FTransform& HeightmapToWorld) override;
+
 	using FEditLayerTargetTypeState = UE::Landscape::EditLayers::FEditLayerTargetTypeState;
 	using FEditLayerRenderItem = UE::Landscape::EditLayers::FEditLayerRenderItem;
 
@@ -40,8 +41,10 @@ public:
 		TArray<TBitArray<>>& OutTargetLayerGroups) const override;
 	virtual FString GetEditLayerRendererDebugName() const override;
 	virtual TArray<FEditLayerRenderItem> GetRenderItems(const UE::Landscape::EditLayers::FMergeContext* InMergeContext) const override;
-	virtual UE::Landscape::EditLayers::ERenderFlags GetRenderFlags(const UE::Landscape::EditLayers::FMergeContext* InMergeContext) const override { return UE::Landscape::EditLayers::ERenderFlags::RenderMode_Recorded; }
-	virtual void RenderLayer(FRenderParams& InRenderParams, UE::Landscape::FRDGBuilderRecorder& RDGBuilderRecorder) override;
+	virtual UE::Landscape::EditLayers::ERenderFlags GetRenderFlags(const UE::Landscape::EditLayers::FMergeContext* InMergeContext) const override;
+	virtual bool CanGroupRenderLayerWith(TScriptInterface<ILandscapeEditLayerRenderer> InOtherRenderer) const override;
+	virtual bool RenderLayer(FRenderParams& RenderParams, UE::Landscape::FRDGBuilderRecorder& RDGBuilderRecorder) override;
+	virtual void BlendLayer(FRenderParams& RenderParams, UE::Landscape::FRDGBuilderRecorder& RDGBuilderRecorder) override;
 #endif
 
 	virtual bool CanAffectHeightmap() const override { return !bEditVisibility; }
@@ -69,7 +72,9 @@ protected:
 	bool bExclusiveRadius = false;
 
 private:
-	void ApplyCirclePatch(UE::Landscape::FRDGBuilderRecorder& RDGBuilderRecorder, bool bIsVisibilityLayer,
+#if WITH_EDITOR
+	bool ApplyCirclePatch(bool bInPerformBlending, FRenderParams* RenderParams, UE::Landscape::FRDGBuilderRecorder& RDGBuilderRecorder, bool bIsVisibilityLayer,
 		FTextureResource* InMergedLandscapeTextureResource, int32 LandscapeTextureSliceIndex,
 		const FIntPoint& DestinationResolution, const FTransform& HeightmapCoordsToWorld, ERHIAccess OutputAccess = ERHIAccess::None);
+#endif // WITH_EDITOR
 };
diff --git a/Engine/Plugins/Experimental/LandscapePatch/Source/LandscapePatch/Public/LandscapePatchManager.h b/Engine/Plugins/Experimental/LandscapePatch/Source/LandscapePatch/Public/LandscapePatchManager.h
--- a/Engine/Plugins/Experimental/LandscapePatch/Source/LandscapePatch/Public/LandscapePatchManager.h
+++ b/Engine/Plugins/Experimental/LandscapePatch/Source/LandscapePatch/Public/LandscapePatchManager.h
@@ -48,7 +48,7 @@ public:
 		UE::Landscape::EditLayers::FEditLayerTargetTypeState& OutEnabledTargetTypeState,
 		TArray<TBitArray<>>& OutTargetLayerGroups) const override {}
 	virtual TArray<UE::Landscape::EditLayers::FEditLayerRenderItem> GetRenderItems(const UE::Landscape::EditLayers::FMergeContext* InMergeContext) const override { return {}; }
-	virtual void RenderLayer(ILandscapeEditLayerRenderer::FRenderParams& InRenderParams, UE::Landscape::FRDGBuilderRecorder& RDGBuilderRecorder) override { };
+	virtual bool RenderLayer(ILandscapeEditLayerRenderer::FRenderParams& RenderParams, UE::Landscape::FRDGBuilderRecorder& RDGBuilderRecorder) override { return false; };
 #endif // WITH_EDITOR
 
 	// Adds the brush to the given landscape, removing it from any previous one. This differs from SetOwningLandscape
diff --git a/Engine/Plugins/Experimental/LandscapePatch/Source/LandscapePatch/Public/LandscapeTexturePatch.h b/Engine/Plugins/Experimental/LandscapePatch/Source/LandscapePatch/Public/LandscapeTexturePatch.h
--- a/Engine/Plugins/Experimental/LandscapePatch/Source/LandscapePatch/Public/LandscapeTexturePatch.h
+++ b/Engine/Plugins/Experimental/LandscapePatch/Source/LandscapePatch/Public/LandscapeTexturePatch.h
@@ -238,9 +238,10 @@ public:
 		TArray<TBitArray<>>& OutTargetLayerGroups) const override;
 	virtual FString GetEditLayerRendererDebugName() const override;
 	virtual TArray<FEditLayerRenderItem> GetRenderItems(const UE::Landscape::EditLayers::FMergeContext* InMergeContext) const override;
-
 	virtual UE::Landscape::EditLayers::ERenderFlags GetRenderFlags(const UE::Landscape::EditLayers::FMergeContext* InMergeContext) const override;
-	virtual void RenderLayer(FRenderParams& RenderParams, UE::Landscape::FRDGBuilderRecorder& RDGBuilderRecorder) override;
+	virtual bool CanGroupRenderLayerWith(TScriptInterface<ILandscapeEditLayerRenderer> InOtherRenderer) const override;
+	virtual bool RenderLayer(FRenderParams& RenderParams, UE::Landscape::FRDGBuilderRecorder& RDGBuilderRecorder) override;
+	virtual void BlendLayer(FRenderParams& RenderParams, UE::Landscape::FRDGBuilderRecorder& RDGBuilderRecorder) override;
 
 	// ULandscapePatchComponent
 	UTextureRenderTarget2D* RenderLayer_Native(const FLandscapeBrushParameters& InParameters, const FTransform& HeightmapToWorld) override;
@@ -351,7 +352,6 @@ public:
 		}
 	}
 
-
 	// Height related functions:
 
 	UFUNCTION(BlueprintCallable, Category = LandscapePatch)
@@ -660,11 +660,13 @@ private:
 #if WITH_EDITOR
 	void TransitionHeightSourceModeInternal(ELandscapeTexturePatchSourceMode OldMode, ELandscapeTexturePatchSourceMode NewMode);
 	FLandscapeHeightPatchConvertToNativeParams GetHeightConvertToNativeParams() const;
-	UTextureRenderTarget2D* ApplyToHeightmap(UE::Landscape::FRDGBuilderRecorder& RDGBuilderRecorder, UTextureRenderTarget2D* InCombinedResult, const FTransform& LandscapeHeightmapToWorld, 
-		ERHIAccess OutputAccess = ERHIAccess::None);
-	void ApplyToWeightmap(UE::Landscape::FRDGBuilderRecorder& RDGBuilderRecorder, ULandscapeWeightPatchTextureInfo* PatchInfo,
-		FTextureResource* InMergedLandscapeTextureResource, int32 LandscapeTextureSliceIndex,
-		const FIntPoint& LandscapeTextureResolution, const FTransform& LandscapeHeightmapToWorld, ERHIAccess OutputAccess = ERHIAccess::None);
+
+	UTextureRenderTarget2D* ApplyToHeightmap(bool bInPerformBlending, FRenderParams* RenderParams, UE::Landscape::FRDGBuilderRecorder& RDGBuilderRecorder, UTextureRenderTarget2D* InCombinedResult,
+		const FTransform& LandscapeHeightmapToWorld, bool& bOutDidRenderSomething, ERHIAccess OutputAccess = ERHIAccess::None);
+
+	void ApplyToWeightmap(bool bInPerformBlending, FRenderParams* RenderParams, UE::Landscape::FRDGBuilderRecorder& RDGBuilderRecorder, ULandscapeWeightPatchTextureInfo* PatchInfo,
+		FTextureResource* InMergedLandscapeTextureResource, int32 LandscapeTextureSliceIndex, const FIntPoint& LandscapeTextureResolution, 
+		const FTransform& LandscapeHeightmapToWorld, bool& bOutDidRenderSomething, ERHIAccess OutputAccess = ERHIAccess::None);
 
 	void GetCommonShaderParams(const FTransform& LandscapeHeightmapToWorldIn, 
 		const FIntPoint& SourceResolutionIn, const FIntPoint& DestinationResolutionIn,
diff --git a/Engine/Shaders/Private/Landscape/LandscapeCommon.ush b/Engine/Shaders/Private/Landscape/LandscapeCommon.ush
--- a/Engine/Shaders/Private/Landscape/LandscapeCommon.ush
+++ b/Engine/Shaders/Private/Landscape/LandscapeCommon.ush
@@ -2,11 +2,30 @@
 
 #pragma once
 
+/** EHeightmapAlphaFlags enum */
+#define EHEIGHTMAPALPHAFLAGS_NONE 0
+#define EHEIGHTMAPALPHAFLAGS_ADDITIVE 0
+#define EHEIGHTMAPALPHAFLAGS_MIN (1 << 0)
+#define EHEIGHTMAPALPHAFLAGS_MAX (1 << 1)
+#define EHEIGHTMAPALPHAFLAGS_ALPHABLEND (EHEIGHTMAPALPHAFLAGS_MIN | EHEIGHTMAPALPHAFLAGS_MAX)
+
+/** EWeightmapAlphaFlags enum */
+#define EWEIGHTMAPALPHAFLAGS_NONE 0
+#define EWEIGHTMAPALPHAFLAGS_ADDITIVE 0
+#define EWEIGHTMAPALPHAFLAGS_MIN (1 << 0)
+#define EWEIGHTMAPALPHAFLAGS_MAX (1 << 1)
+#define EWEIGHTMAPALPHAFLAGS_ALPHABLEND (EHEIGHTMAPALPHAFLAGS_MIN | EHEIGHTMAPALPHAFLAGS_MAX)
+
+// Must match LandscapeDataAccess::MaxValue : 
+static const float LANDSCAPE_MAX_VALUE = 65535.0f;
+// Must match LandscapeDataAccess::MidValue : 
+static const float LANDSCAPE_MID_VALUE = 32768.0f;
+
 // Height is packed from a [0, 2^16-1] value to two bytes, each of which we need to provide as a 0-1 float for writing 
 float2 PackHeight(float InHeight)
 {
-	uint Height = (uint) clamp(round(InHeight), 0.0f, 65535.0f);
-	return float2((float) (Height >> 8) / 255.0, (float) (Height & 0xFF) / 255.0);
+	uint PackedHeight = (uint) clamp(round(InHeight), 0.0f, 65535.0f);
+	return float2((float)(PackedHeight >> 8) / 255.0, (float)(PackedHeight & 0xFF) / 255.0);
 }
 
 float UnpackHeight(float2 InPackedHeight)
@@ -19,6 +38,51 @@ float UnpackHeight(float2 InPackedHeight)
 	return (InPackedHeight.r * 256 + InPackedHeight.g) * 255;
 }
 
+// Height alpha is 14 bits-wide + 2 bits are left for flags. InAlpha is expected in the [0,1] range
+float2 PackHeightAlpha(float InAlpha, uint InAlphaFlags)
+{
+	uint PackedAlpha = (uint)clamp(InAlpha * 65532.0f, 0.0f, 65532.0f); // First 14 bits are used to store alpha (0xFFFC in float)
+	uint PackedFlags = InAlphaFlags & 0x3; // Last 2 bits are used to store flags
+	return float2((float)(PackedAlpha >> 8) / 255.0f, (float)((PackedAlpha & 0xFF) | PackedFlags) / 255.0);
+}
+
+void UnpackHeightAlpha(float2 InPackedAlpha, out float OutAlpha, out uint OutFlags)
+{
+	uint AlphaWithFlags = ((uint)round(InPackedAlpha.r * 255.0f) << 8) | (uint)round(InPackedAlpha.g * 255.0f);
+	uint Alpha = AlphaWithFlags & 0xFFFC; 
+	OutAlpha = float(Alpha) / 65532.0f; // First 14 bits are used to store alpha (0xFFFC in float)
+	OutFlags = AlphaWithFlags & 0x3; // Last 2 bits are used to store flags
+}
+
+// It's a bit overkill but we have 2 channels available for weight, so we can use 16 bits for storing it
+// Weight is packed from a [0, 2^16-1] value to two bytes, each of which we need to provide as a 0-1 float for writing 
+float2 PackWeight(float InWeight)
+{
+	uint PackedWeight = (uint) clamp(round(InWeight * 65535.0f), 0.0f, 65535.0f);
+	return float2((float)(PackedWeight >> 8) / 255.0f, (float)(PackedWeight & 0xFF) / 255.0f);
+}
+
+float UnpackWeight(float2 InPackedWeight)
+{
+	return ((InPackedWeight.r * 256 + InPackedWeight.g) * 255) / 65535.0f;
+}
+
+// Weight Alpha is 14 bits-wide + 2 bits are left for flags. InWeight is expected in the [0,1] range
+float2 PackWeightAlpha(float InAlpha, uint InAlphaFlags)
+{
+	uint PackedAlpha = (uint)clamp(InAlpha * 65532.0f, 0.0f, 65532.0f); // First 14 bits are used to store alpha (0xFFFC in float)
+	uint PackedFlags = InAlphaFlags & 0x3; // Last 2 bits are used to store flags
+	return float2((float)(PackedAlpha >> 8) / 255.0, (float)((PackedAlpha & 0xFF) | PackedFlags) / 255.0);
+}
+
+void UnpackWeightAlpha(float2 InPackedAlpha, out float OutAlpha, out uint OutFlags)
+{
+	uint AlphaWithFlags = ((uint)round(InPackedAlpha.r * 255.0f) << 8) | (uint)round(InPackedAlpha.g * 255.0f);
+	uint Alpha = AlphaWithFlags & 0xFFFC; 
+	OutAlpha = float(Alpha) / 65532.0f; // First 14 bits are used to store alpha (0xFFFC in float)
+	OutFlags = AlphaWithFlags & 0x3; // Last 2 bits are used to store flags
+}
+
 // Compute the normal of the triangle formed by the 3 points (in winding order).
 // .xyz of each point is its world space position
 // Additionally, the .w component of each point indicates its validity (valid: .w == 1, invalid: .w == 0) 
diff --git a/Engine/Shaders/Private/Landscape/LandscapeEditLayersHeightmaps.usf b/Engine/Shaders/Private/Landscape/LandscapeEditLayersHeightmaps.usf
--- a/Engine/Shaders/Private/Landscape/LandscapeEditLayersHeightmaps.usf
+++ b/Engine/Shaders/Private/Landscape/LandscapeEditLayersHeightmaps.usf
@@ -18,11 +18,12 @@
 #if MERGE_EDIT_LAYER
 // MergeEditLayerPS inputs/outputs :
 
-/** EEditLayerHeightmapBlendMode enum */
-#define EEDITLAYERHEIGHTMAPBLENDMODE_ADDITIVE 0
-#define EEDITLAYERHEIGHTMAPBLENDMODE_ALPHABLEND 1
+/** EHeightmapBlendMode enum */
+#define EHEIGHTMAPBLENDMODE_ADDITIVE 0
+#define EHEIGHTMAPBLENDMODE_LEGACYALPHABLEND 1
+#define EHEIGHTMAPBLENDMODE_ALPHABLEND 2
 
-uint InEditLayerBlendMode; // See EEditLayerHeightmapBlendMode
+uint InEditLayerBlendMode; // See EHeightmapBlendMode
 float InEditLayerAlpha; // Alpha value to be used in the blend
 
 Texture2D<float4> InCurrentEditLayerHeightmap; // The height of the current edit layer to merge 
@@ -67,7 +68,7 @@ float ExtractAlpha(float2 InPackedAlpha)
 #if GENERATE_NORMALS
 float ScaleHeight(float InHeight)
 {
-	return (InHeight - 32768.0f) * InLandscapeGridScale.z * TERRAIN_ZSCALE;
+	return (InHeight - LANDSCAPE_MID_VALUE) * InLandscapeGridScale.z * TERRAIN_ZSCALE;
 }
 
 float ExtractAndScaleHeight(float2 InPackedHeight)
@@ -77,7 +78,7 @@ float ExtractAndScaleHeight(float2 InPackedHeight)
 
 float UnscaleHeight(float InScaledHeight)
 {
-	return (InScaledHeight / InLandscapeGridScale.z / TERRAIN_ZSCALE) + 32768.0f;
+	return (InScaledHeight / InLandscapeGridScale.z / TERRAIN_ZSCALE) + LANDSCAPE_MID_VALUE;
 }
 
 float4 ExtractSampleInfo(int2 InTextureCoordinates, float2 InPackedHeight)
@@ -94,22 +95,21 @@ float4 ExtractSampleInfo(int2 InTextureCoordinates, float2 InPackedHeight)
 // Pixel shaders : 
 
 #if MERGE_EDIT_LAYER
-void MergeEditLayerPS(in float4 SVPos : SV_POSITION, out float4 OutColor : SV_Target0)
+void MergeEditLayerPS(in float4 SVPos : SV_POSITION, out float4 OutPackedHeight : SV_Target0)
 {
 	uint2 TextureCoordinates = floor(SVPos.xy);
 		
-	// TODO [jonathan.bard] Add a const for 32768.0f or something...
 	float4 CurrentLayerSample = InCurrentEditLayerHeightmap.Load(int3(TextureCoordinates, 0));
-	float CurrentLayerHeight = UnpackHeight(CurrentLayerSample.xy) - 32768.0f; // xy = relative coordinates, z = mip level
+	float CurrentLayerHeight = UnpackHeight(CurrentLayerSample.xy) - LANDSCAPE_MID_VALUE; // xy = relative coordinates, z = mip level
 	float4 PreviousLayersSample = InPreviousEditLayersHeightmap.Load(int3(TextureCoordinates, 0));
-	float PreviousLayersHeight = UnpackHeight(PreviousLayersSample.xy) - 32768.0f; // xy = relative coordinates, z = mip level
+	float PreviousLayersHeight = UnpackHeight(PreviousLayersSample.xy) - LANDSCAPE_MID_VALUE; // xy = relative coordinates, z = mip level
 	float FinalHeight = PreviousLayersHeight;
 	
-	if (InEditLayerBlendMode == EEDITLAYERHEIGHTMAPBLENDMODE_ADDITIVE)
+	if (InEditLayerBlendMode == EHEIGHTMAPBLENDMODE_ADDITIVE)
 	{
 		FinalHeight += InEditLayerAlpha * CurrentLayerHeight;
 	}
-	else if (InEditLayerBlendMode == EEDITLAYERHEIGHTMAPBLENDMODE_ALPHABLEND)
+	else if (InEditLayerBlendMode == EHEIGHTMAPBLENDMODE_LEGACYALPHABLEND)
 	{
 		float LayerHeightAlphaBlend = ExtractAlpha(CurrentLayerSample.ba);
 		float NewHeight = InEditLayerAlpha * CurrentLayerHeight + (LayerHeightAlphaBlend * PreviousLayersHeight);
@@ -125,10 +125,35 @@ void MergeEditLayerPS(in float4 SVPos : SV_POSITION, out float4 OutColor : SV_Ta
 			FinalHeight = PreviousLayersHeight;
 		}
 	}
+	else if (InEditLayerBlendMode == EHEIGHTMAPBLENDMODE_ALPHABLEND)
+	{
+		float LayerHeightAlphaBlend = 0.0f;
+		uint HeightFlags = 0;
+		UnpackHeightAlpha(CurrentLayerSample.ba, LayerHeightAlphaBlend, HeightFlags);
+		float FinalAlpha = InEditLayerAlpha * LayerHeightAlphaBlend;
+		if (HeightFlags == EHEIGHTMAPALPHAFLAGS_ALPHABLEND)
+		{
+			FinalHeight = lerp(PreviousLayersHeight, CurrentLayerHeight, FinalAlpha);
+		}
+		else if (HeightFlags & EHEIGHTMAPALPHAFLAGS_MIN)
+		{
+			FinalHeight = lerp(PreviousLayersHeight, min(PreviousLayersHeight, CurrentLayerHeight), FinalAlpha);
+		}
+		else if (HeightFlags & EHEIGHTMAPALPHAFLAGS_MAX)
+		{
+			FinalHeight = lerp(PreviousLayersHeight, max(PreviousLayersHeight, CurrentLayerHeight), FinalAlpha);
+		}
+		// Additive case (== no flag) :
+		else if (HeightFlags == EHEIGHTMAPALPHAFLAGS_ADDITIVE)
+		{
+			FinalHeight += CurrentLayerHeight * FinalAlpha;
+		}
+	}
 	
-	FinalHeight = clamp(FinalHeight, -32768.0, 32767.0);
-	FinalHeight += 32768.0f;
-	OutColor = float4(PackHeight(FinalHeight), 0.0f, 0.0f);
+	FinalHeight += LANDSCAPE_MID_VALUE;
+	FinalHeight = clamp(FinalHeight, 0.0f, LANDSCAPE_MAX_VALUE);
+
+	OutPackedHeight = float4(PackHeight(FinalHeight), 0.0f, 0.0f);
 }
 #endif // MERGE_EDIT_LAYER
 
diff --git a/Engine/Shaders/Private/Landscape/LandscapeEditLayersUtils.usf b/Engine/Shaders/Private/Landscape/LandscapeEditLayersUtils.usf
--- a/Engine/Shaders/Private/Landscape/LandscapeEditLayersUtils.usf
+++ b/Engine/Shaders/Private/Landscape/LandscapeEditLayersUtils.usf
@@ -10,7 +10,8 @@
 //#define MARK_VALIDITY 1
 //#define COPY_QUADS 1
 //#define COPY_QUADS_MULTISOURCE 1
-//#define USE_SWIZZLE_MASK 1
+//#define COPY_WEIGHTMAP 1
+//#define RESOLVE_LAYER_DATA 1
 #endif // defined (__INTELLISENSE__)
 
 
@@ -91,7 +92,7 @@ Texture2D<float4> InSourceTexture;
 	op(61) \
 	op(62) 
 
-Buffer<uint2> InQuadInfos; // .x = source texture index, .y = swizzle mask
+Buffer<uint2> InQuadInfos; // .x = source texture index, .y = source texture channel (for weightmaps)
 
 // Define N input textures :
 #define DECLARE_SOURCE_TEXTURE(n) Texture2D<float4> InSourceTexture_##n;
@@ -99,6 +100,10 @@ FOREACH_SOURCETEXTURE(DECLARE_SOURCE_TEXTURE)
 #undef DECLARE_SOURCE_TEXTURE
 #endif // COPY_QUADS_MULTISOURCE
 
+#if RESOLVE_LAYER_DATA
+Texture2D<float4> InSourceTexture;
+#endif // RESOLVE_LAYER_DATA
+
 // ----------------------------------------------------------------------------------
 // Pixel shaders : 
 
@@ -130,15 +135,14 @@ void CopyQuadsPS(float4 SvPosition : SV_POSITION, float2 QuadUV : TEXCOORD0, out
 
 void CopyQuadsMultiSourcePS(float4 SvPosition : SV_POSITION, float2 SourceUV : TEXCOORD0, float2 UnusedQuadUV : TEXCOORD1, float InRectIndex : RECT_INDEX, out float4 OutColor : SV_Target0)
 {
-	OutColor = 0.0f;
-
 	uint QuadIndex = (uint)InRectIndex;
 	int SourceTextureIndex = InQuadInfos[QuadIndex].x;
+	float4 ColorSample = 0.0f;
 
 	BRANCH
 	switch (SourceTextureIndex)
 	{
-#define IMPLEMENT_SAMPLE_SWITCH_CASE(n) case n: OutColor = InSourceTexture_##n.SampleLevel(GlobalPointClampedSampler, SourceUV, 0); break; 
+#define IMPLEMENT_SAMPLE_SWITCH_CASE(n) case n: ColorSample = InSourceTexture_##n.SampleLevel(GlobalPointClampedSampler, SourceUV, 0); break; 
 	FOREACH_SOURCETEXTURE(IMPLEMENT_SAMPLE_SWITCH_CASE);
 #undef IMPLEMENT_SAMPLE_SWITCH_CASE
 
@@ -146,18 +150,40 @@ void CopyQuadsMultiSourcePS(float4 SvPosition : SV_POSITION, float2 SourceUV : T
 		break;
 	}
 
-	// Swizzle channels as requested by ChannelSwizzleMask:
-#if USE_SWIZZLE_MASK
-	uint ChannelSwizzleMask = (uint)InQuadInfos[QuadIndex].y;
-	float4 SwizzledColor = 0.0f;
-	[unroll]
-	for (uint DestinationChannelIndex = 0; DestinationChannelIndex < 4; ++DestinationChannelIndex)
-	{
-		uint SourceChannelIndex = ((3u << (DestinationChannelIndex * 2u)) & ChannelSwizzleMask) >> (DestinationChannelIndex * 2u);
-		SwizzledColor[DestinationChannelIndex] = OutColor[SourceChannelIndex];
-	}
-	OutColor = SwizzledColor;
-#endif // USE_SWIZZLE_MASK
+	OutColor = 0.0f;
+
+#if COPY_WEIGHTMAP
+	// For weightmaps, copy from the appropriate channel and pack : 
+	uint SourceChannelIndex = (uint)InQuadInfos[QuadIndex].y;
+	float Weight = ColorSample[SourceChannelIndex];
+	OutColor = float4(PackWeight(Weight), PackWeightAlpha(1.0f, EWEIGHTMAPALPHAFLAGS_NONE));
+#else // COPY_WEIGHTMAP
+	OutColor = ColorSample;
+#endif // !COPY_WEIGHTMAP
+}
+
+#endif // COPY_QUADS_MULTISOURCE
+
+// ----------------------------------------------------------------------------------
+
+#if RESOLVE_LAYER_DATA
+
+void ResolveLayerData(in float4 SVPos : SV_POSITION, out float4 OutColor : SV_Target0)
+{
+	OutColor = 0;
+
+	float4 ColorSample = InSourceTexture.Load(int3(SVPos.xy, 0)); // xy = coordinates, z = mip level
+
+	// TODO [jonathan.bard] : since this step is not a plain copy anymore, we should be able to support different format types (e.g. instead 
+	//  of expecting heightmaps to be compressed, we could read it in a different way) and support different input resolutions (e.g. BP brushes 
+	//  wouldn't have the necessity to use the same resolution as the merge RT). We could also support an alpha channel and/or other alpha flags : 
+#if RESOLVE_WEIGHTMAP
+	float Weight = ColorSample.x;
+	OutColor = float4(PackWeight(Weight), PackWeightAlpha(1.0f, EWEIGHTMAPALPHAFLAGS_NONE));
+#else // RESOLVE_WEIGHTMAP
+	// Currently, it's just a plain copy for heightmaps : 
+	OutColor = ColorSample;
+#endif // !RESOLVE_WEIGHTMAP
 }
 
-#endif // COPY_QUADS_MULTISOURCE
\ No newline at end of file
+#endif // RESOLVE_LAYER_DATA
\ No newline at end of file
diff --git a/Engine/Shaders/Private/Landscape/LandscapeEditLayersWeightmaps.usf b/Engine/Shaders/Private/Landscape/LandscapeEditLayersWeightmaps.usf
--- a/Engine/Shaders/Private/Landscape/LandscapeEditLayersWeightmaps.usf
+++ b/Engine/Shaders/Private/Landscape/LandscapeEditLayersWeightmaps.usf
@@ -1,6 +1,7 @@
 // Copyright Epic Games, Inc. All Rights Reserved.
 
 #include "../Common.ush"
+#include "LandscapeCommon.ush"
 
 // ----------------------------------------------------------------------------------
 
@@ -8,47 +9,49 @@
 // Uncomment the appropriate define for enabling syntax highlighting with HLSL Tools for Visual Studio : 
 //#define MERGE_EDIT_LAYER 1
 //#define PERFORM_LEGACY_WEIGHT_BLENDING 1
-#define PACK_WEIGHTMAP 1
+//#define PACK_WEIGHTMAP 1
 //#define GENERATE_MIPS 1
 #endif // defined (__INTELLISENSE__)
 
 
 // ----------------------------------------------------------------------------------
 
-/** EEditLayerWeightmapBlendMode enum */
-#define EEDITLAYERWEIGHTMAPBLENDMODE_ADDITIVE 0
-#define EEDITLAYERWEIGHTMAPBLENDMODE_SUBTRACTIVE 1
+/** EWeightmapBlendMode enum */
+#define EWEIGHTMAPBLENDMODE_ADDITIVE 0
+#define EWEIGHTMAPBLENDMODE_SUBTRACTIVE 1
+#define EWEIGHTMAPBLENDMODE_PASSTHROUGH 2
+#define EWEIGHTMAPBLENDMODE_ALPHABLEND 3
 
-/** EWeightmapPaintLayerFlags enum */
-#define EWEIGHTMAPPAINTLAYERFLAGS_ISVISIBILITYLAYER (1 << 0)
-#define EWEIGHTMAPPAINTLAYERFLAGS_ISWEIGHTBLENDED (1 << 1)
+/** EWeightmapTargetLayerFlags enum */
+#define EWEIGHTMAPTARGETLAYERFLAGS_ISVISIBILITYLAYER (1 << 0)
+#define EWEIGHTMAPTARGETLAYERFLAGS_ISWEIGHTBLENDED (1 << 1)
 
 #if MERGE_EDIT_LAYER
 // MergeEditLayerPS inputs/outputs :
 
-uint InPaintLayerIndex; // Index of the paint layer being processed
-uint InEditLayerPaintLayerBlendMode; // See EEditLayerWeightmapBlendMode. Paint layer's blend mode, which is per-edit layer and per-paint layer and can therefore differ from one edit layer to another
+uint InTargetLayerIndex; // Index of the target layer being processed
+uint InEditLayerTargetLayerBlendMode; // See EWeightmapBlendMode. Target layer's blend mode, which is per-edit layer and per-target layer and can therefore differ from one edit layer to another
 float InEditLayerAlpha; // Global alpha value of the edit layer currently being merged
 
-Texture2DArray<float> InCurrentEditLayerWeightmaps; // The paint layers of the current edit layer to merge (1 slice per weightmap in the target layer group)
-Texture2DArray<float> InPreviousEditLayersWeightmaps; // The result from the merge of all prior edit layers (1 slice per weightmap in the target layer group)
-// Array that contains per-edit layer / per-paint layer information (e.g. the paint layer's blend mode, which is per-edit layer and can therefore differ from one edit layer to another)
+Texture2DArray<float4> InCurrentEditLayerWeightmaps; // The target layers of the current edit layer to merge (1 slice per weightmap in the target layer group)
+Texture2DArray<float4> InPreviousEditLayersWeightmaps; // The result from the merge of all prior edit layers (1 slice per weightmap in the target layer group)
+// Array that contains per-edit layer / per-target layer information (e.g. the target layer's blend mode, which is per-edit layer and can therefore differ from one edit layer to another)
 #endif // MERGE_EDIT_LAYER
 
 #if PERFORM_LEGACY_WEIGHT_BLENDING
 // PerformLegacyWeightBlendingPS inputs/outputs :
 
-// Per-paint layer information : 
-struct FWeightmapPaintLayerInfo
+// Per-target layer information : 
+struct FWeightmapTargetLayerInfo
 {
-	uint Flags; // See EWeightmapPaintLayerFlags
+	uint Flags; // See EWeightmapTargetLayerFlags
 };
 
-uint InPaintLayerIndex; // Index of the paint layer being processed
-uint InNumPaintLayers; // Number of paint layers (i.e. == InPaintLayerInfos.Num())
-// Array that contains per-paint layer information (e.g. paint layer's flags)
-StructuredBuffer<FWeightmapPaintLayerInfo> InPaintLayerInfos;
-Texture2DArray<float> InCurrentEditLayerWeightmaps; // The texture to horizontally blend (1 slice per weightmap in the target layer group)
+uint InTargetLayerIndex; // Index of the target layer being processed
+uint InNumTargetLayers; // Number of target layers (i.e. == InTargetLayerInfos.Num())
+// Array that contains per-target layer information (e.g. target layer's flags)
+StructuredBuffer<FWeightmapTargetLayerInfo> InTargetLayerInfos;
+Texture2DArray<float4> InCurrentEditLayerWeightmaps; // The texture to horizontally blend (1 slice per weightmap in the target layer group)
 #endif // PERFORM_LEGACY_WEIGHT_BLENDING
 
 #if PACK_WEIGHTMAP
@@ -59,7 +62,7 @@ uint4 InSourcePixelOffsets[4];  // For each channel, offset to add to the pixel'
 uint2 InSubsectionPixelOffset; // Offset of the subsection currently being rendered
 uint InIsAdditive; // = 1 if some channels of this texture have been packed in a previous draw, 0 otherwise
 
-Texture2DArray<float> InSourceWeightmaps; // Source, single-channel, texture to pack
+Texture2DArray<float2> InSourceWeightmaps; // Source, single-channel (we don't care about the alpha flags at this point), texture to pack
 Texture2D<float4> InWeightmapBeingPacked; // The weightmap being packed, in case the packing operation is additive, i.e. it occurs across multiple draws (contains the channels that have previously been packed)
 #endif // PACK_WEIGHTMAP
 
@@ -76,11 +79,11 @@ Texture2D<float4> InSourceWeightmap; // Source weightmap (containing the current
 // Util functions : 
 
 #if PERFORM_LEGACY_WEIGHT_BLENDING
-bool IsWeightBlendedPaintLayer(FWeightmapPaintLayerInfo InPaintLayerInfo)
+bool IsWeightBlendedTargetLayer(FWeightmapTargetLayerInfo InTargetLayerInfo)
 {
 	// For the visibility layer, deactivate weight blending altogether : 
-	return (((InPaintLayerInfo.Flags & EWEIGHTMAPPAINTLAYERFLAGS_ISVISIBILITYLAYER) == 0)
-		&& (InPaintLayerInfo.Flags & EWEIGHTMAPPAINTLAYERFLAGS_ISWEIGHTBLENDED));
+	return (((InTargetLayerInfo.Flags & EWEIGHTMAPTARGETLAYERFLAGS_ISVISIBILITYLAYER) == 0)
+		&& (InTargetLayerInfo.Flags & EWEIGHTMAPTARGETLAYERFLAGS_ISWEIGHTBLENDED));
 }
 #endif // PERFORM_LEGACY_WEIGHT_BLENDING
 
@@ -89,61 +92,98 @@ bool IsWeightBlendedPaintLayer(FWeightmapPaintLayerInfo InPaintLayerInfo)
 // Pixel shaders : 
 
 #if MERGE_EDIT_LAYER
-void MergeEditLayerPS(in float4 SVPos : SV_POSITION, out float OutWeight : SV_Target0)
+void MergeEditLayerPS(in float4 SVPos : SV_POSITION, out float4 OutPackedWeight : SV_Target0)
 {
 	uint2 TextureCoordinates = floor(SVPos.xy);
 		
-	float CurrentLayerWeight = InCurrentEditLayerWeightmaps.Load(int4(TextureCoordinates, InPaintLayerIndex, 0)); // xy = relative coordinates, z = index in texture array, w = mip level
-	float PreviousLayersWeight = InPreviousEditLayersWeightmaps.Load(int4(TextureCoordinates, InPaintLayerIndex, 0)); // xy = relative coordinates, z = index in texture array, w = mip level
-	OutWeight = PreviousLayersWeight;
+	float4 CurrentLayerWeightSample = 0.0f; 
+	if (InEditLayerTargetLayerBlendMode != EWEIGHTMAPBLENDMODE_PASSTHROUGH)
+	{
+		CurrentLayerWeightSample = InCurrentEditLayerWeightmaps.Load(int4(TextureCoordinates, InTargetLayerIndex, 0)); // xy = relative coordinates, z = index in texture array, w = mip level
+	}
+	float CurrentLayerWeight = UnpackWeight(CurrentLayerWeightSample.xy);
+	float CurrentLayerWeightAlpha = 1.0f;
+	uint CurrentLayerWeightFlags = EWEIGHTMAPALPHAFLAGS_NONE;
+	UnpackWeightAlpha(CurrentLayerWeightSample.zw, CurrentLayerWeightAlpha, CurrentLayerWeightFlags);
+
+	float2 PreviousLayersWeightSample = InPreviousEditLayersWeightmaps.Load(int4(TextureCoordinates, InTargetLayerIndex, 0)).xy; // xy = relative coordinates, z = index in texture array, w = mip level
+	float PreviousLayersWeight = UnpackWeight(PreviousLayersWeightSample);
+
+	float FinalWeight = PreviousLayersWeight;
 	
-	if (InEditLayerPaintLayerBlendMode == EEDITLAYERWEIGHTMAPBLENDMODE_ADDITIVE)
+	if (InEditLayerTargetLayerBlendMode == EWEIGHTMAPBLENDMODE_ADDITIVE)
 	{
-		OutWeight += CurrentLayerWeight * InEditLayerAlpha;
+		FinalWeight += CurrentLayerWeight * InEditLayerAlpha;
 	}
-	else if (InEditLayerPaintLayerBlendMode == EEDITLAYERWEIGHTMAPBLENDMODE_SUBTRACTIVE)
+	else if (InEditLayerTargetLayerBlendMode == EWEIGHTMAPBLENDMODE_SUBTRACTIVE)
 	{
-		OutWeight -= CurrentLayerWeight * InEditLayerAlpha;
+		FinalWeight -= CurrentLayerWeight * InEditLayerAlpha;
 	}
+	else if (InEditLayerTargetLayerBlendMode == EWEIGHTMAPBLENDMODE_ALPHABLEND)
+	{
+		float FinalAlpha = InEditLayerAlpha * CurrentLayerWeightAlpha;
+		if (CurrentLayerWeightFlags == EWEIGHTMAPALPHAFLAGS_ALPHABLEND)
+		{
+			FinalWeight = lerp(PreviousLayersWeight, CurrentLayerWeight, FinalAlpha);
+		}
+		else if (CurrentLayerWeightFlags & EWEIGHTMAPALPHAFLAGS_MIN)
+		{
+			FinalWeight = lerp(PreviousLayersWeight, min(PreviousLayersWeight, CurrentLayerWeight), FinalAlpha);
+		}
+		else if (CurrentLayerWeightFlags & EWEIGHTMAPALPHAFLAGS_MAX)
+		{
+			FinalWeight = lerp(PreviousLayersWeight, max(PreviousLayersWeight, CurrentLayerWeight), FinalAlpha);
+		}
+		// Additive case (== no flag) :
+		// TODO [jonathan.bard] : remove InEditLayerTargetLayerBlendMode == EWEIGHTMAPBLENDMODE_SUBTRACTIVE and add the EWEIGHTMAPBLENDMODE_SUBTRACTIVE flag
+		else if (CurrentLayerWeightFlags == EWEIGHTMAPALPHAFLAGS_ADDITIVE)
+		{
+			FinalWeight += CurrentLayerWeight * FinalAlpha;
+		}
+	}
+
+	OutPackedWeight = float4(PackWeight(FinalWeight), 0.0f, 0.0f);
 }
 #endif // MERGE_EDIT_LAYER
 
 #if PERFORM_LEGACY_WEIGHT_BLENDING
-void PerformLegacyWeightBlendingPS(in float4 SVPos : SV_POSITION, out float OutColor : SV_Target0)
+void PerformLegacyWeightBlendingPS(in float4 SVPos : SV_POSITION, out float4 OutColor : SV_Target0)
 {
 	uint2 TextureCoordinates = floor(SVPos.xy);
-	OutColor = 0.0f;
+	float FinalWeight = 0.0f;
 	
-	FWeightmapPaintLayerInfo ActivePaintLayerInfo = InPaintLayerInfos[InPaintLayerIndex];
-	if (!IsWeightBlendedPaintLayer(ActivePaintLayerInfo))
+	FWeightmapTargetLayerInfo ActiveTargetLayerInfo = InTargetLayerInfos[InTargetLayerIndex];
+	if (!IsWeightBlendedTargetLayer(ActiveTargetLayerInfo))
 	{
-		OutColor = InCurrentEditLayerWeightmaps.Load(int4(TextureCoordinates, InPaintLayerIndex, 0)); // xy = relative coordinates, z = index in texture array, w = mip level
+		FinalWeight = UnpackWeight(InCurrentEditLayerWeightmaps.Load(int4(TextureCoordinates, InTargetLayerIndex, 0)).xy); // xy = relative coordinates, z = index in texture array, w = mip level
 	}
 	else
 	{
-		float ActivePaintLayerWeight = 0.0f;
+		float ActiveTargetLayerWeight = 0.0f;
 		float BlendedWeightsSum = 0.0f;
 
 		LOOP
-		for (int i = 0; i < InNumPaintLayers; ++i)
+		for (uint i = 0; i < InNumTargetLayers; ++i)
 		{
-			bool bIsOutputPaintLayer = (i == InPaintLayerIndex);
-			FWeightmapPaintLayerInfo PaintLayerInfo = InPaintLayerInfos[i];
+			bool bIsOutputTargetLayer = (i == InTargetLayerIndex);
+			FWeightmapTargetLayerInfo TargetLayerInfo = InTargetLayerInfos[i];
 
-			// Only weight blended (and non-visibility) paint layers participate to weight blending : 
-			if (IsWeightBlendedPaintLayer(PaintLayerInfo))
+			// Only weight blended (and non-visibility) target layers participate to weight blending : 
+			if (IsWeightBlendedTargetLayer(TargetLayerInfo))
 			{
-				float Weight = InCurrentEditLayerWeightmaps.Load(int4(TextureCoordinates, i, 0)); // xy = relative coordinates, z = index in texture array, w = mip level
-				if (bIsOutputPaintLayer)
+				float Weight = UnpackWeight(InCurrentEditLayerWeightmaps.Load(int4(TextureCoordinates, i, 0)).xy); // xy = relative coordinates, z = index in texture array, w = mip level
+				if (bIsOutputTargetLayer)
 				{
-					ActivePaintLayerWeight = Weight;
+					ActiveTargetLayerWeight = Weight;
 				}
 				BlendedWeightsSum += Weight;
 			}
 		}
 
-		OutColor = (BlendedWeightsSum > 0.0f) ? saturate(ActivePaintLayerWeight / BlendedWeightsSum) : ActivePaintLayerWeight;
+		FinalWeight = (BlendedWeightsSum > 0.0f) ? saturate(ActiveTargetLayerWeight / BlendedWeightsSum) : ActiveTargetLayerWeight;
 	}
+
+	OutColor = float4(PackWeight(FinalWeight), 0.0f, 0.0f);
 }
 #endif // PERFORM_LEGACY_WEIGHT_BLENDING
 
@@ -167,7 +207,7 @@ void PackWeightmapPS(in float4 SVPos : SV_POSITION, out float4 OutColor : SV_Tar
 		if (InSourceSliceIndices[i] >= 0)
 		{
 			uint2 TextureCoordinates = SubsectionRelativeTextureCoordinates + InSourcePixelOffsets[i].xy;
-			OutColor[i] = InSourceWeightmaps.Load(int4(TextureCoordinates, InSourceSliceIndices[i], 0)); // xy = relative coordinates, z = index in texture array, w = mip level
+			OutColor[i] = UnpackWeight(InSourceWeightmaps.Load(int4(TextureCoordinates, InSourceSliceIndices[i], 0)).xy); // xy = relative coordinates, z = index in texture array, w = mip level
 		}
 	}
 }
diff --git a/Engine/Shaders/Private/Landscape/LandscapeLayersWeightmapsPS.usf b/Engine/Shaders/Private/Landscape/LandscapeLayersWeightmapsPS.usf
--- a/Engine/Shaders/Private/Landscape/LandscapeLayersWeightmapsPS.usf
+++ b/Engine/Shaders/Private/Landscape/LandscapeLayersWeightmapsPS.usf
@@ -42,7 +42,7 @@ StructuredBuffer<FEditLayerWeightmapMergeInfo> InEditLayersMergeInfos; // Descri
 #define EWEIGHTMAPPAINTLAYERFLAGS_ISWEIGHTBLENDED (1 << 1)
 
 // Contains the information related to a given edit layers and how it should merge into the final weightmap
-struct FWeightmapPaintLayerInfo
+struct FEditLayerWeightmapPaintLayerInfo
 {
 	uint Flags; // See EWeightmapPaintLayerFlags
 };
@@ -70,7 +70,7 @@ Texture2DArray<float> InPerChannelPaintLayerWeightmaps_3;
 StructuredBuffer<uint> InPaintLayerInfoIndices;
 
 // Array that contains per-paint layer information (e.g. paint layer's blend mode)
-StructuredBuffer<FWeightmapPaintLayerInfo> InPaintLayerInfos; 
+StructuredBuffer<FEditLayerWeightmapPaintLayerInfo> InPaintLayerInfos; 
 #endif // FINALIZE_WEIGHTMAP
 
 
@@ -118,7 +118,7 @@ uint LoadPaintLayerInfoIndexForChannel(uint InChannelIndex, uint InPaintLayerInd
 	return InPaintLayerInfoIndices[PaintLayerInfoIndexInBigBuffer];
 }
 
-bool IsWeightBlendedPaintLayer(FWeightmapPaintLayerInfo InPaintLayerInfo)
+bool IsWeightBlendedPaintLayer(FEditLayerWeightmapPaintLayerInfo InPaintLayerInfo)
 {
 	// For the visibility layer, deactivate weight blending altogether : 
 	return (((InPaintLayerInfo.Flags & EWEIGHTMAPPAINTLAYERFLAGS_ISVISIBILITYLAYER) == 0)
@@ -135,7 +135,7 @@ float FinalizeWeightmapChannel(uint2 InSampleCoordinates, uint InChannelIndex)
 		uint NumPaintLayers = InPerChannelNumPaintLayers[InChannelIndex];
 		
 		uint ActivePaintLayerInfoIndex = LoadPaintLayerInfoIndexForChannel(InChannelIndex, PaintLayerIndexInWeightmaps);
-		FWeightmapPaintLayerInfo ActivePaintLayerInfo = InPaintLayerInfos[ActivePaintLayerInfoIndex];
+		FEditLayerWeightmapPaintLayerInfo ActivePaintLayerInfo = InPaintLayerInfos[ActivePaintLayerInfoIndex];
 
 		if (!IsWeightBlendedPaintLayer(ActivePaintLayerInfo))
 		{
@@ -150,7 +150,7 @@ float FinalizeWeightmapChannel(uint2 InSampleCoordinates, uint InChannelIndex)
 		{
 			bool bIsOutputPaintLayer = (i == PaintLayerIndexInWeightmaps);
 			uint PaintLayerInfoIndex = LoadPaintLayerInfoIndexForChannel(InChannelIndex, i);
-			FWeightmapPaintLayerInfo PaintLayerInfo = InPaintLayerInfos[PaintLayerInfoIndex];
+			FEditLayerWeightmapPaintLayerInfo PaintLayerInfo = InPaintLayerInfos[PaintLayerInfoIndex];
 
 			// Only weight blended (and non-visibility) paint layers participate to weight blending : 
 			if (IsWeightBlendedPaintLayer(PaintLayerInfo)) 
diff --git a/Engine/Shaders/Private/Landscape/LandscapeMergeTexturesPS.usf b/Engine/Shaders/Private/Landscape/LandscapeMergeTexturesPS.usf
--- a/Engine/Shaders/Private/Landscape/LandscapeMergeTexturesPS.usf
+++ b/Engine/Shaders/Private/Landscape/LandscapeMergeTexturesPS.usf
@@ -42,7 +42,6 @@ Texture2D<float> InSourceTextures_2; // Source, single-channel, texture to pack
 Texture2D<float> InSourceTextures_3; // Source, single-channel, texture to pack
 #endif // PACK_RGBA_CHANNELS
 
-static const float LANDSCAPE_MAX_VALUE = 65535.0f;
 
 // ----------------------------------------------------------------------------------
 // Util functions : 
diff --git a/Engine/Source/Runtime/Landscape/Classes/LandscapeEditLayer.h b/Engine/Source/Runtime/Landscape/Classes/LandscapeEditLayer.h
--- a/Engine/Source/Runtime/Landscape/Classes/LandscapeEditLayer.h
+++ b/Engine/Source/Runtime/Landscape/Classes/LandscapeEditLayer.h
@@ -197,11 +197,13 @@ public:
 	LANDSCAPE_API virtual void GetRendererStateInfo(const UE::Landscape::EditLayers::FMergeContext* InMergeContext,
 		UE::Landscape::EditLayers::FEditLayerTargetTypeState& OutSupportedTargetTypeState, UE::Landscape::EditLayers::FEditLayerTargetTypeState& OutEnabledTargetTypeState, TArray<TBitArray<>>& OutTargetLayerGroups) const override;
 	LANDSCAPE_EDIT_LAYERS_BATCHED_MERGE_EXPERIMENTAL
-	virtual UE::Landscape::EditLayers::ERenderFlags GetRenderFlags(const UE::Landscape::EditLayers::FMergeContext* InMergeContext) const override { return UE::Landscape::EditLayers::ERenderFlags::RenderMode_Recorded; }
+	virtual UE::Landscape::EditLayers::ERenderFlags GetRenderFlags(const UE::Landscape::EditLayers::FMergeContext* InMergeContext) const override;
 	LANDSCAPE_EDIT_LAYERS_BATCHED_MERGE_EXPERIMENTAL
 	LANDSCAPE_API virtual TArray<UE::Landscape::EditLayers::FEditLayerRenderItem> GetRenderItems(const UE::Landscape::EditLayers::FMergeContext* InMergeContext) const override;
 	LANDSCAPE_EDIT_LAYERS_BATCHED_MERGE_EXPERIMENTAL
-	LANDSCAPE_API virtual void RenderLayer(FRenderParams& RenderParams, UE::Landscape::FRDGBuilderRecorder& RDGBuilderRecorder) override;
+	LANDSCAPE_API virtual bool RenderLayer(FRenderParams& RenderParams, UE::Landscape::FRDGBuilderRecorder& RDGBuilderRecorder) override;
+	LANDSCAPE_EDIT_LAYERS_BATCHED_MERGE_EXPERIMENTAL
+	LANDSCAPE_API virtual void BlendLayer(FRenderParams& RenderParams, UE::Landscape::FRDGBuilderRecorder& RDGBuilderRecorder) override;
 	LANDSCAPE_EDIT_LAYERS_BATCHED_MERGE_EXPERIMENTAL
 	LANDSCAPE_API virtual FString GetEditLayerRendererDebugName() const override;
 	//~ End ILandscapeEditLayerRenderer implementation
diff --git a/Engine/Source/Runtime/Landscape/Private/LandscapeBlueprintBrushBase.cpp b/Engine/Source/Runtime/Landscape/Private/LandscapeBlueprintBrushBase.cpp
--- a/Engine/Source/Runtime/Landscape/Private/LandscapeBlueprintBrushBase.cpp
+++ b/Engine/Source/Runtime/Landscape/Private/LandscapeBlueprintBrushBase.cpp
@@ -2,6 +2,8 @@
 
 #include "LandscapeBlueprintBrushBase.h"
 #include "CoreMinimal.h"
+#include "GlobalShader.h"
+#include "Engine/Engine.h"
 #include "LandscapeProxy.h"
 #include "Landscape.h"
 #include "LandscapeEditResourcesSubsystem.h"
@@ -9,15 +11,16 @@
 #include "LandscapeInfo.h"
 #include "LandscapeLayerInfoObject.h"
 #include "LandscapePrivate.h"
-#include "ProfilingDebugging/RealtimeGPUProfiler.h"
-#include "Misc/MapErrors.h"
-#include "Misc/UObjectToken.h"
 #include "Logging/MessageLog.h"
 #include "Logging/TokenizedMessage.h"
-#include "Engine/Engine.h"
-#include "TextureResource.h"
-#include "RHITransition.h"
 #include "Materials/MaterialExpressionLandscapeVisibilityMask.h"
+#include "Misc/MapErrors.h"
+#include "Misc/UObjectToken.h"
+#include "ProfilingDebugging/RealtimeGPUProfiler.h"
+#include "PixelShaderUtils.h"
+#include "RHITransition.h"
+#include "ShaderParameterStruct.h"
+#include "TextureResource.h"
 
 #include UE_INLINE_GENERATED_CPP_BY_NAME(LandscapeBlueprintBrushBase)
 
@@ -32,6 +35,58 @@ static TAutoConsoleVariable<int32> CVarLandscapeBrushPadding(
 	TEXT("The number of frames to wait before pushing a full Landscape update when a brush is calling RequestLandscapeUpdate"));
 #endif
 
+// ----------------------------------------------------------------------------------
+
+class FLandscapeEditLayersResolveLayerDataPS : public FGlobalShader
+{
+	DECLARE_GLOBAL_SHADER(FLandscapeEditLayersResolveLayerDataPS);
+	SHADER_USE_PARAMETER_STRUCT(FLandscapeEditLayersResolveLayerDataPS, FGlobalShader);
+
+public:
+	BEGIN_SHADER_PARAMETER_STRUCT(FParameters, )
+		SHADER_PARAMETER_RDG_TEXTURE_SRV(Texture2D<float4>, InSourceTexture)
+		RENDER_TARGET_BINDING_SLOTS()
+	END_SHADER_PARAMETER_STRUCT()
+
+	class FResolveWeightmap : SHADER_PERMUTATION_BOOL("RESOLVE_WEIGHTMAP");
+	using FPermutationDomain = TShaderPermutationDomain<FResolveWeightmap>;
+
+	static FPermutationDomain GetPermutationVector(bool bResolveWeightmap)
+	{
+		FPermutationDomain PermutationVector;
+		PermutationVector.Set<FResolveWeightmap>(bResolveWeightmap);
+		return PermutationVector;
+	}
+
+	static bool ShouldCompilePermutation(const FGlobalShaderPermutationParameters& InParameters)
+	{
+		return UE::Landscape::DoesPlatformSupportEditLayers(InParameters.Platform);
+	}
+
+	static void ModifyCompilationEnvironment(const FGlobalShaderPermutationParameters& InParameters, FShaderCompilerEnvironment& OutEnvironment)
+	{
+		OutEnvironment.SetDefine(TEXT("RESOLVE_LAYER_DATA"), 1);
+	}
+
+	static void ResolveLayerData(FRDGBuilder& GraphBuilder, FParameters* InParameters, const FIntPoint& InTextureSize, bool bResolveWeightmap)
+	{
+		FGlobalShaderMap* ShaderMap = GetGlobalShaderMap(GMaxRHIFeatureLevel);
+
+		const FLandscapeEditLayersResolveLayerDataPS::FPermutationDomain PixelPermutationVector = FLandscapeEditLayersResolveLayerDataPS::GetPermutationVector(bResolveWeightmap);
+		TShaderMapRef<FLandscapeEditLayersResolveLayerDataPS> PixelShader(ShaderMap, PixelPermutationVector);
+
+		FPixelShaderUtils::AddFullscreenPass(
+			GraphBuilder,
+			ShaderMap,
+			RDG_EVENT_NAME("LandscapeLayers_ResolveLayerData"),
+			PixelShader,
+			InParameters,
+			FIntRect(0, 0, InTextureSize.X, InTextureSize.Y));
+	}
+};
+
+IMPLEMENT_GLOBAL_SHADER(FLandscapeEditLayersResolveLayerDataPS, "/Engine/Private/Landscape/LandscapeEditLayersUtils.usf", "ResolveLayerData", SF_Pixel);
+
 
 // ----------------------------------------------------------------------------------
 
@@ -291,7 +346,7 @@ TArray<UE::Landscape::EditLayers::FEditLayerRenderItem> ALandscapeBlueprintBrush
 	return { FEditLayerRenderItem(EnabledTargetTypeState, InputWorldArea, OutputWorldArea, /*bModifyExistingWeightmapsOnly = */false)};
 }
 
-void ALandscapeBlueprintBrushBase::RenderLayer(ILandscapeEditLayerRenderer::FRenderParams& RenderParams, UE::Landscape::FRDGBuilderRecorder& RDGBuilderRecorder)
+bool ALandscapeBlueprintBrushBase::RenderLayer(ILandscapeEditLayerRenderer::FRenderParams& RenderParams, UE::Landscape::FRDGBuilderRecorder& RDGBuilderRecorder)
 {
 	using namespace UE::Landscape;
 
@@ -303,9 +358,8 @@ void ALandscapeBlueprintBrushBase::RenderLayer(ILandscapeEditLayerRenderer::FRen
 	RenderParams.MergeRenderContext->CycleBlendRenderTargets(RDGBuilderRecorder);
 	ULandscapeScratchRenderTarget* WriteRT = RenderParams.MergeRenderContext->GetBlendRenderTargetWrite();
 	ULandscapeScratchRenderTarget* CurrentLayerReadRT = RenderParams.MergeRenderContext->GetBlendRenderTargetRead();
-	
-	// The write render target will be accessed as ERHIAccess::CopyDest all along : 
-	WriteRT->TransitionTo(ERHIAccess::CopyDest, RDGBuilderRecorder);
+
+	WriteRT->TransitionTo(ERHIAccess::RTV, RDGBuilderRecorder);
 	CurrentLayerReadRT->TransitionTo(ERHIAccess::SRVMask, RDGBuilderRecorder);
 
 	const bool bIsHeightmapMerge = RenderParams.MergeRenderContext->IsHeightmapMerge();
@@ -342,9 +396,8 @@ void ALandscapeBlueprintBrushBase::RenderLayer(ILandscapeEditLayerRenderer::FRen
 		RHI_BREADCRUMB_EVENT_GAMETHREAD("Render %s", TargetLayerName);
 
 		// If necessary, copy from the source render target (slice) to the scratch render target 2D : 
-		RHI_BREADCRUMB_EVENT_GAMETHREAD("Copy Source (slice %i) -> %s", TargetLayerIndexInGroup, CurrentLayerReadRT2D->GetDebugName());
-
 		{
+			RHI_BREADCRUMB_EVENT_GAMETHREAD("Copy Source (slice %i) -> %s", TargetLayerIndexInGroup, CurrentLayerReadRT2D->GetDebugName());
 			ULandscapeScratchRenderTarget::FCopyFromScratchRenderTargetParams CopyParams(CurrentLayerReadRT);
 			// It's important to copy only the needed size, because CurrentLayerReadRT2D comes from ScratchTexture, whose size may be larger than RenderParams.RenderAreaSectionRect)
 			CopyParams.CopySize = RenderParams.RenderAreaSectionRect.Size();
@@ -355,7 +408,6 @@ void ALandscapeBlueprintBrushBase::RenderLayer(ILandscapeEditLayerRenderer::FRen
 		CurrentLayerReadRT2D->TransitionTo(ERHIAccess::SRVMask, RDGBuilderRecorder);
 
 		check(CurrentLayerReadRT2D->GetCurrentState() == ERHIAccess::SRVMask);
-		check(WriteRT->GetCurrentState() == ERHIAccess::CopyDest);
 
 		UTextureRenderTarget2D* ReadRT2D = CurrentLayerReadRT2D->GetRenderTarget2D();
 		// If the BP brush failed to render, we still need to passthrough from the read RT to the write RT in order not to lose what has been merged so far : 
@@ -379,25 +431,30 @@ void ALandscapeBlueprintBrushBase::RenderLayer(ILandscapeEditLayerRenderer::FRen
 
 		// Resolve back to the write RT 
 		{
-			RHI_BREADCRUMB_EVENT_GAMETHREAD("Copy BP Render Result -> %s (slice %i)", WriteRT->GetDebugName(), TargetLayerIndexInGroup);
-
-			// The RT returned by the brush is in SRV state so we need a transition: 
-			ENQUEUE_RENDER_COMMAND(TransitionToCopySrc)([Resource = OutputRT2D->GetResource()](FRHICommandListImmediate& InRHICmdList) mutable
+			RHI_BREADCRUMB_EVENT_GAMETHREAD("Resolve BP Render Result -> %s (slice %i)", WriteRT->GetDebugName(), TargetLayerIndexInGroup);
+
+			ENQUEUE_RENDER_COMMAND(ResolveLayerData)(
+				[ InputResource = OutputRT2D->GetResource()
+				, OutputResource = WriteRT->GetRenderTarget()->GetResource()
+				, SliceIndex = TargetLayerIndexInGroup
+				, bIsHeightmapMerge
+				// It's important to copy only the needed size, because OutputRT2D is currently sized after ReadRT2D, which can be a render target larger than it needs to be (because 
+				//  it can come from ScratchTexture, whose size may be larger than RenderParams.RenderAreaSectionRect)
+				, ResolveTextureSize = RenderParams.RenderAreaSectionRect.Size()](FRHICommandListImmediate& InRHICmdList) mutable
 			{
-				InRHICmdList.Transition(FRHITransitionInfo(Resource->TextureRHI, ERHIAccess::SRVMask, ERHIAccess::CopySrc));
-			});
+					FRDGBuilder GraphBuilder(InRHICmdList, RDG_EVENT_NAME("ResolveLayerData"));
+					FRDGTextureRef InputTexture = GraphBuilder.RegisterExternalTexture(CreateRenderTarget(InputResource->GetTextureRHI(), TEXT("InputTexture")));
+					FRDGTextureRef OutputTexture = GraphBuilder.RegisterExternalTexture(CreateRenderTarget(OutputResource->GetTextureRHI(), TEXT("OutputTexture")));
 
-			ULandscapeScratchRenderTarget::FCopyFromTextureParams CopyParams(OutputRT2D);
-			// It's important to copy only the needed size, because OutputRT2D is currently sized after ReadRT2D, which can be a render target larger than it needs to be (because 
-			//  it can come from ScratchTexture, whose size may be larger than RenderParams.RenderAreaSectionRect)
-			CopyParams.CopySize = RenderParams.RenderAreaSectionRect.Size();
-			CopyParams.DestSliceIndex = TargetLayerIndexInGroup;
-			WriteRT->CopyFrom(CopyParams, RDGBuilderRecorder);
+					FLandscapeEditLayersResolveLayerDataPS::FParameters* PassParameters = GraphBuilder.AllocParameters<FLandscapeEditLayersResolveLayerDataPS::FParameters>();
+					PassParameters->RenderTargets[0] = FRenderTargetBinding(OutputTexture, ERenderTargetLoadAction::ELoad, /*InMipIndex = */0, SliceIndex);
+					PassParameters->InSourceTexture = GraphBuilder.CreateSRV(FRDGTextureSRVDesc::Create(InputTexture));
 
-			// It's also expected we leave the RT return by the BP as SRV :
-			ENQUEUE_RENDER_COMMAND(TransitionToSRV)([Resource = OutputRT2D->GetResource()](FRHICommandListImmediate& InRHICmdList) mutable
-			{
-				InRHICmdList.Transition(FRHITransitionInfo(Resource->TextureRHI, ERHIAccess::CopySrc, ERHIAccess::SRVMask));
+					FLandscapeEditLayersResolveLayerDataPS::ResolveLayerData(GraphBuilder, PassParameters, ResolveTextureSize, /*bResolveWeightmap = */!bIsHeightmapMerge);
+
+					// Don't let the graph builder transition the output texture back to SRVMask
+					GraphBuilder.SetTextureAccessFinal(OutputTexture, ERHIAccess::RTV);
+					GraphBuilder.Execute();
 			});
 		}
 	}
@@ -405,6 +462,8 @@ void ALandscapeBlueprintBrushBase::RenderLayer(ILandscapeEditLayerRenderer::FRen
 	// Leave the render targets in the state they're expected to be in: 
 	WriteRT->TransitionTo(ERHIAccess::RTV, RDGBuilderRecorder);
 	CurrentLayerReadRT->TransitionTo(ERHIAccess::SRVMask, RDGBuilderRecorder);
+
+	return true;
 }
 
 FString ALandscapeBlueprintBrushBase::GetEditLayerRendererDebugName() const
diff --git a/Engine/Source/Runtime/Landscape/Private/LandscapeEditLayerRenderer.cpp b/Engine/Source/Runtime/Landscape/Private/LandscapeEditLayerRenderer.cpp
--- a/Engine/Source/Runtime/Landscape/Private/LandscapeEditLayerRenderer.cpp
+++ b/Engine/Source/Runtime/Landscape/Private/LandscapeEditLayerRenderer.cpp
@@ -5,6 +5,8 @@
 #include "Algo/Transform.h"
 #include "Algo/Count.h"
 #include "Engine/Engine.h"
+#include "Engine/TextureRenderTarget2D.h"
+#include "Engine/TextureRenderTarget2DArray.h"
 #include "EngineModule.h"
 #include "Landscape.h"
 #include "LandscapeComponent.h"
@@ -99,6 +101,99 @@ BEGIN_SHADER_PARAMETER_STRUCT(FCopyQuadsPSParameters, )
 END_SHADER_PARAMETER_STRUCT()
 
 
+// ----------------------------------------------------------------------------------
+// /Engine/Private/Landscape/LandscapeEditLayersHeightmaps.usf shaders :
+
+class FLandscapeEditLayersHeightmapsMergeEditLayerPS : public FGlobalShader
+{
+	DECLARE_GLOBAL_SHADER(FLandscapeEditLayersHeightmapsMergeEditLayerPS);
+	SHADER_USE_PARAMETER_STRUCT(FLandscapeEditLayersHeightmapsMergeEditLayerPS, FGlobalShader);
+
+public:
+	BEGIN_SHADER_PARAMETER_STRUCT(FParameters, )
+		SHADER_PARAMETER(uint32, InEditLayerBlendMode)
+		SHADER_PARAMETER(float, InEditLayerAlpha)
+		SHADER_PARAMETER_RDG_TEXTURE_SRV(Texture2D<float4>, InCurrentEditLayerHeightmap)
+		SHADER_PARAMETER_RDG_TEXTURE_SRV(Texture2D<float4>, InPreviousEditLayersHeightmap)
+		RENDER_TARGET_BINDING_SLOTS()
+	END_SHADER_PARAMETER_STRUCT()
+
+	static bool ShouldCompilePermutation(const FGlobalShaderPermutationParameters& InParameters)
+	{
+		return UE::Landscape::DoesPlatformSupportEditLayers(InParameters.Platform);
+	}
+
+	static void ModifyCompilationEnvironment(const FGlobalShaderPermutationParameters& InParameters, FShaderCompilerEnvironment& OutEnvironment)
+	{
+		OutEnvironment.SetDefine(TEXT("MERGE_EDIT_LAYER"), 1);
+	}
+
+	static void MergeEditLayerPS(FRDGEventName&& InRDGEventName, FRDGBuilder& GraphBuilder, FParameters* InParameters, const FIntPoint& InTextureSize)
+	{
+		FGlobalShaderMap* ShaderMap = GetGlobalShaderMap(GMaxRHIFeatureLevel);
+		TShaderMapRef<FLandscapeEditLayersHeightmapsMergeEditLayerPS> PixelShader(ShaderMap);
+
+		FPixelShaderUtils::AddFullscreenPass(
+			GraphBuilder,
+			ShaderMap,
+			MoveTemp(InRDGEventName),
+			PixelShader,
+			InParameters,
+			FIntRect(0, 0, InTextureSize.X, InTextureSize.Y),
+			TStaticBlendStateWriteMask<CW_RG>::GetRHI());
+	}
+};
+
+IMPLEMENT_GLOBAL_SHADER(FLandscapeEditLayersHeightmapsMergeEditLayerPS, "/Engine/Private/Landscape/LandscapeEditLayersHeightmaps.usf", "MergeEditLayerPS", SF_Pixel);
+
+
+// ----------------------------------------------------------------------------------
+// /Engine/Private/Landscape/LandscapeEditLayersWeightmaps.usf shaders :
+
+class FLandscapeEditLayersWeightmapsMergeEditLayerPS : public FGlobalShader
+{
+	DECLARE_GLOBAL_SHADER(FLandscapeEditLayersWeightmapsMergeEditLayerPS);
+	SHADER_USE_PARAMETER_STRUCT(FLandscapeEditLayersWeightmapsMergeEditLayerPS, FGlobalShader);
+
+public:
+	BEGIN_SHADER_PARAMETER_STRUCT(FParameters, )
+		SHADER_PARAMETER(uint32, InTargetLayerIndex)
+		SHADER_PARAMETER(uint32, InEditLayerTargetLayerBlendMode)
+		SHADER_PARAMETER(float, InEditLayerAlpha)
+		SHADER_PARAMETER_RDG_TEXTURE_SRV(Texture2DArray<float4>, InCurrentEditLayerWeightmaps)
+		SHADER_PARAMETER_RDG_TEXTURE_SRV(Texture2DArray<float4>, InPreviousEditLayersWeightmaps)
+		RENDER_TARGET_BINDING_SLOTS()
+	END_SHADER_PARAMETER_STRUCT()
+
+	static bool ShouldCompilePermutation(const FGlobalShaderPermutationParameters& InParameters)
+	{
+		return UE::Landscape::DoesPlatformSupportEditLayers(InParameters.Platform);
+	}
+
+	static void ModifyCompilationEnvironment(const FGlobalShaderPermutationParameters& InParameters, FShaderCompilerEnvironment& OutEnvironment)
+	{
+		OutEnvironment.SetDefine(TEXT("MERGE_EDIT_LAYER"), 1);
+	}
+
+	static void MergeEditLayerPS(FRDGEventName&& InRDGEventName, FRDGBuilder& GraphBuilder, FParameters* InParameters, const FIntPoint& InTextureSize)
+	{
+		FGlobalShaderMap* ShaderMap = GetGlobalShaderMap(GMaxRHIFeatureLevel);
+		TShaderMapRef<FLandscapeEditLayersWeightmapsMergeEditLayerPS> PixelShader(ShaderMap);
+
+		FPixelShaderUtils::AddFullscreenPass(
+			GraphBuilder,
+			ShaderMap,
+			MoveTemp(InRDGEventName),
+			PixelShader,
+			InParameters,
+			FIntRect(0, 0, InTextureSize.X, InTextureSize.Y),
+			TStaticBlendStateWriteMask<CW_RG>::GetRHI());
+	}
+};
+
+IMPLEMENT_GLOBAL_SHADER(FLandscapeEditLayersWeightmapsMergeEditLayerPS, "/Engine/Private/Landscape/LandscapeEditLayersWeightmaps.usf", "MergeEditLayerPS", SF_Pixel);
+
+
 // ----------------------------------------------------------------------------------
 
 FString ConvertTargetLayerNamesToString(const TArrayView<const FName>& InTargetLayerNames)
@@ -323,6 +418,12 @@ TBitArray<> FMergeContext::BuildTargetLayerBitIndices(bool bInBitValue) const
 
 // ----------------------------------------------------------------------------------
 
+const FEditLayerTargetTypeState& FEditLayerTargetTypeState::GetDummyTargetTypeState()
+{
+	static FEditLayerTargetTypeState DummyTargetTypeState;
+	return DummyTargetTypeState;
+}
+
 FEditLayerTargetTypeState::FEditLayerTargetTypeState(const FMergeContext* InMergeContext)
 	: FEditLayerTargetTypeState(InMergeContext, ELandscapeToolTargetTypeFlags::None)
 {}
@@ -532,6 +633,12 @@ FString FEditLayerTargetTypeState::ToString() const
 
 // ----------------------------------------------------------------------------------
 
+const FEditLayerRendererState& FEditLayerRendererState::GetDummyRendererState()
+{
+	static FEditLayerRendererState DummyRendererState;
+	return DummyRendererState;
+}
+
 FEditLayerRendererState::FEditLayerRendererState(const FMergeContext* InMergeContext)
 	: MergeContext(InMergeContext)
 	, SupportedTargetTypeState(InMergeContext)
@@ -675,6 +782,7 @@ void FEditLayerRendererState::UpdateActiveTargetTypeState()
 	ActiveTargetTypeState = SupportedTargetTypeState.Intersect(EnabledTargetTypeState);
 }
 
+
 // ----------------------------------------------------------------------------------
 
 bool FMergeRenderBatch::operator<(const FMergeRenderBatch& InOther) const
@@ -831,6 +939,8 @@ void FMergeRenderContext::AllocateResources()
 	else
 	{
 		checkf(MaxNeededNumSlices > 0, TEXT("Weightmaps should have at least 1 slice"));
+		// We use extra channels for weightmaps for storing alpha / alpha flags :
+		RenderTargetFormat = ETextureRenderTargetFormat::RTF_RGBA8;
 	}
 
 	ULandscapeEditResourcesSubsystem* LandscapeEditResourcesSubsystem = GEngine->GetEngineSubsystem<ULandscapeEditResourcesSubsystem>();
@@ -969,10 +1079,13 @@ void FMergeRenderContext::ResetVisualLogOffset()
 	CurrentVisualLogOffset = FVector(ForceInitToZero);
 }
 
-void FMergeRenderContext::RenderValidityRenderTargets(const FMergeRenderBatch& InRenderBatch, UE::Landscape::FRDGBuilderRecorder& RDGBuilderRecorder)
+void FMergeRenderContext::RenderValidityRenderTargets(UE::Landscape::FRDGBuilderRecorder& RDGBuilderRecorder)
 {
 	using namespace UE::Landscape;
 
+	const FMergeRenderBatch* RenderBatch = GetCurrentRenderBatch();
+	check(RenderBatch != nullptr);
+
 	struct FTextureAndRects
 	{
 		FTextureAndRects(const FName& InTargetLayerName, const FString& InTextureDebugName, FTextureResource* InTextureResource)
@@ -988,9 +1101,9 @@ void FMergeRenderContext::RenderValidityRenderTargets(const FMergeRenderBatch& I
 	};
 
 	TArray<FTextureAndRects> TexturesAndRects;
-	TexturesAndRects.Reserve(InRenderBatch.TargetLayerBitIndices.CountSetBits());
-	ForEachTargetLayer(InRenderBatch.TargetLayerBitIndices,
-		[this, &RDGBuilderRecorder, &TexturesAndRects, &InRenderBatch](int32 InTargetLayerIndex, const FName& InTargetLayerName, ULandscapeLayerInfoObject* InWeightmapLayerInfo)
+	TexturesAndRects.Reserve(RenderBatch->TargetLayerBitIndices.CountSetBits());
+	ForEachTargetLayer(RenderBatch->TargetLayerBitIndices,
+		[this, &RDGBuilderRecorder, &TexturesAndRects, RenderBatch](int32 InTargetLayerIndex, const FName& InTargetLayerName, ULandscapeLayerInfoObject* InWeightmapLayerInfo)
 		{
 			ULandscapeScratchRenderTarget* ScratchRenderTarget = PerTargetLayerValidityRenderTargets.FindChecked(InTargetLayerName);
 			check(ScratchRenderTarget != nullptr);
@@ -1001,11 +1114,11 @@ void FMergeRenderContext::RenderValidityRenderTargets(const FMergeRenderBatch& I
 			FTextureAndRects& TextureAndRects = TexturesAndRects.Emplace_GetRef(InTargetLayerName, ScratchRenderTarget->GetDebugName(), ScratchRenderTarget->GetRenderTarget2D()->GetResource());
 
 			// Then build a list of quads for marking where the components are valid for this target layer on this batch:
-			const TSet<ULandscapeComponent*>& Components = InRenderBatch.TargetLayersToComponents[InTargetLayerIndex];
+			const TSet<ULandscapeComponent*>& Components = RenderBatch->TargetLayersToComponents[InTargetLayerIndex];
 			TextureAndRects.Rects.Reserve(Components.Num());
 			for (ULandscapeComponent* Component : Components)
 			{
-				FIntRect ComponentRect = InRenderBatch.ComputeSectionRect(Component, /*bInWithDuplicateBorders = */false);
+				FIntRect ComponentRect = RenderBatch->ComputeSectionRect(Component, /*bInWithDuplicateBorders = */false);
 				TextureAndRects.Rects.Add(FUintVector4(ComponentRect.Min.X, ComponentRect.Min.Y, ComponentRect.Max.X + 1, ComponentRect.Max.Y + 1));
 			}
 
@@ -1052,14 +1165,17 @@ void FMergeRenderContext::RenderValidityRenderTargets(const FMergeRenderBatch& I
 	RDGBuilderRecorder.EnqueueRDGCommand(RDGCommand, RDGExternalTextureAccessFinalList);
 }
 
-void FMergeRenderContext::RenderExpandedRenderTarget(const FMergeRenderBatch& InRenderBatch, UE::Landscape::FRDGBuilderRecorder& RDGBuilderRecorder)
+void FMergeRenderContext::RenderExpandedRenderTarget(UE::Landscape::FRDGBuilderRecorder& RDGBuilderRecorder)
 {
 	using namespace UE::Landscape;
 
+	const FMergeRenderBatch* RenderBatch = GetCurrentRenderBatch();
+	check(RenderBatch != nullptr);
+
 	TArray<FUintVector4> SourceRects, DestinationRects;
 	{
 		TArray<FIntRect> SourceInclusiveRects, DestinationInclusiveRects;
-		InRenderBatch.ComputeAllSubsectionRects(SourceInclusiveRects, DestinationInclusiveRects);
+		RenderBatch->ComputeAllSubsectionRects(SourceInclusiveRects, DestinationInclusiveRects);
 		// ComputeAllSubsectionRects returns inclusive bounds while AddRasterizeToRectsPass requires exclusive bounds : 
 		Algo::Transform(SourceInclusiveRects, SourceRects, [](const FIntRect& InRect) { return FUintVector4(InRect.Min.X, InRect.Min.Y, InRect.Max.X + 1, InRect.Max.Y + 1); });
 		Algo::Transform(DestinationInclusiveRects, DestinationRects, [](const FIntRect& InRect) { return FUintVector4(InRect.Min.X, InRect.Min.Y, InRect.Max.X + 1, InRect.Max.Y + 1); });
@@ -1132,6 +1248,126 @@ void FMergeRenderContext::RenderExpandedRenderTarget(const FMergeRenderBatch& In
 	RDGBuilderRecorder.EnqueueRDGCommand(RDGCommand, RDGExternalTextureAccessFinalList);
 }
 
+void FMergeRenderContext::GenericBlendLayer(const FBlendParams& InBlendParams, UE::Landscape::FRDGBuilderRecorder& RDGBuilderRecorder)
+{
+	using namespace UE::Landscape;
+	using namespace UE::Landscape::EditLayers;
+
+	const FMergeRenderBatch* RenderBatch = GetCurrentRenderBatch();
+	check(RenderBatch != nullptr);
+
+	// In this step, we perform the blend : merge the edit layer with the previous layers in the stack
+	CycleBlendRenderTargets(RDGBuilderRecorder);
+	ULandscapeScratchRenderTarget* WriteRT = GetBlendRenderTargetWrite();
+	ULandscapeScratchRenderTarget* CurrentLayerReadRT = GetBlendRenderTargetRead();
+	ULandscapeScratchRenderTarget* PreviousLayersReadRT = GetBlendRenderTargetReadPrevious();
+
+	WriteRT->TransitionTo(ERHIAccess::RTV, RDGBuilderRecorder);
+	CurrentLayerReadRT->TransitionTo(ERHIAccess::SRVMask, RDGBuilderRecorder);
+	PreviousLayersReadRT->TransitionTo(ERHIAccess::SRVMask, RDGBuilderRecorder);
+
+	if (IsHeightmapMerge())
+	{
+		auto RDGCommand =
+			[ HeightmapBlendParams = InBlendParams.HeightmapBlendParams
+			, OutputResource = WriteRT->GetRenderTarget2D()->GetResource()
+			, OutputResourceName = WriteRT->GetDebugName()
+			, CurrentEditLayerResource = CurrentLayerReadRT->GetRenderTarget2D()->GetResource()
+			, PreviousEditLayersResource = PreviousLayersReadRT->GetRenderTarget2D()->GetResource()
+			, EffectiveTextureSize = RenderBatch->GetRenderTargetResolution(/*bInWithDuplicateBorders = */false)]
+			(FRDGBuilder& GraphBuilder)
+			{
+				FRDGTextureRef OutputTextureRef = GraphBuilder.RegisterExternalTexture(CreateRenderTarget(OutputResource->TextureRHI, TEXT("OutputTexture")));
+				FRDGTextureRef CurrentEditLayerTextureRef = GraphBuilder.RegisterExternalTexture(CreateRenderTarget(CurrentEditLayerResource->TextureRHI, TEXT("CurrentEditLayerTexture")));
+				FRDGTextureRef PreviousEditLayersTextureRef = GraphBuilder.RegisterExternalTexture(CreateRenderTarget(PreviousEditLayersResource->TextureRHI, TEXT("PreviousEditLayersTexture")));
+
+				FLandscapeEditLayersHeightmapsMergeEditLayerPS::FParameters* PSParams = GraphBuilder.AllocParameters<FLandscapeEditLayersHeightmapsMergeEditLayerPS::FParameters>();
+				PSParams->RenderTargets[0] = FRenderTargetBinding(OutputTextureRef, ERenderTargetLoadAction::ENoAction);
+				PSParams->InEditLayerBlendMode = static_cast<uint32>(HeightmapBlendParams.BlendMode);
+				PSParams->InEditLayerAlpha = HeightmapBlendParams.Alpha;
+				PSParams->InCurrentEditLayerHeightmap = GraphBuilder.CreateSRV(FRDGTextureSRVDesc::Create(CurrentEditLayerTextureRef));
+				PSParams->InPreviousEditLayersHeightmap = GraphBuilder.CreateSRV(FRDGTextureSRVDesc::Create(PreviousEditLayersTextureRef));
+
+				FLandscapeEditLayersHeightmapsMergeEditLayerPS::MergeEditLayerPS(
+					RDG_EVENT_NAME("MergeEditLayer(Height) -> %s", *OutputResourceName), GraphBuilder, PSParams, EffectiveTextureSize);
+			};
+
+
+		// We need to specify the final state of the external textures to prevent the graph builder from transitioning them to SRVMask (even those that end up as SRVMask at the end of this command,
+		//  because they will likely be part of another RDGCommand down the line so we need to maintain an accurate picture of every external texture ever involved in the recorded command so that 
+		//  we can set a proper access when the recorder is flushed (and the FRDGBuilder, executed) :
+		TArray<FRDGBuilderRecorder::FRDGExternalTextureAccessFinal> RDGExternalTextureAccessFinalList =
+		{
+			{ WriteRT->GetRenderTarget()->GetResource(), ERHIAccess::RTV },
+			{ CurrentLayerReadRT->GetRenderTarget()->GetResource(), ERHIAccess::SRVMask },
+			{ PreviousLayersReadRT->GetRenderTarget()->GetResource(), ERHIAccess::SRVMask },
+		};
+		RDGBuilderRecorder.EnqueueRDGCommand(RDGCommand, RDGExternalTextureAccessFinalList);
+	}
+	else
+	{
+		ForEachTargetLayer(RenderBatch->TargetLayerBitIndices,
+			[WriteRT, CurrentLayerReadRT, PreviousLayersReadRT, &RenderBatch, &InBlendParams, &RDGBuilderRecorder](int32 InTargetLayerIndex, const FName& InTargetLayerName, ULandscapeLayerInfoObject* InWeightmapLayerInfo)
+			{
+				// By default, use passthrough mode so that each layer gets at least copied into the WriteRT :
+				const FWeightmapBlendParams* TargetLayerBlendParams = &FWeightmapBlendParams::GetDefaultPassthroughBlendParams();
+				// Special case for visibility which is always "Additive" :
+				if (InTargetLayerName == UMaterialExpressionLandscapeVisibilityMask::ParameterName)
+				{
+					TargetLayerBlendParams = &FWeightmapBlendParams::GetDefaultAdditiveBlendParams();
+				}
+				else if (const FWeightmapBlendParams* FoundTargetLayerBlendParams = InBlendParams.WeightmapBlendParams.Find(InTargetLayerName))
+				{
+					TargetLayerBlendParams = FoundTargetLayerBlendParams;
+				}
+
+				// TODO [jonathan.bard] : we could render several layers at once via MRT (up to MaxSimultaneousRenderTargets)
+				auto RDGCommand =
+					[ TargetLayerBlendParams = *TargetLayerBlendParams
+					, InTargetLayerIndex
+					, InTargetLayerName
+					, OutputResource = WriteRT->GetRenderTarget2DArray()->GetResource()
+					, OutputResourceName = WriteRT->GetDebugName()
+					, CurrentEditLayerResource = CurrentLayerReadRT->GetRenderTarget2DArray()->GetResource()
+					, PreviousEditLayersResource = PreviousLayersReadRT->GetRenderTarget2DArray()->GetResource()
+					, EffectiveTextureSize = RenderBatch->GetRenderTargetResolution(/*bInWithDuplicateBorders = */false)]
+					(FRDGBuilder& GraphBuilder)
+					{
+						FRDGTextureRef OutputTextureRef = GraphBuilder.RegisterExternalTexture(CreateRenderTarget(OutputResource->TextureRHI, TEXT("OutputTexture")));
+						FRDGTextureRef CurrentEditLayerTextureRef = GraphBuilder.RegisterExternalTexture(CreateRenderTarget(CurrentEditLayerResource->TextureRHI, TEXT("CurrentEditLayerTexture")));
+						FRDGTextureSRVRef CurrentEditLayerTextureSRVRef = GraphBuilder.CreateSRV(FRDGTextureSRVDesc::Create(CurrentEditLayerTextureRef));
+						FRDGTextureRef PreviousEditLayersTextureRef = GraphBuilder.RegisterExternalTexture(CreateRenderTarget(PreviousEditLayersResource->TextureRHI, TEXT("PreviousEditLayersTexture")));
+						FRDGTextureSRVRef PreviousEditLayersTextureSRVRef = GraphBuilder.CreateSRV(FRDGTextureSRVDesc::Create(PreviousEditLayersTextureRef));
+
+						FLandscapeEditLayersWeightmapsMergeEditLayerPS::FParameters* PSParams = GraphBuilder.AllocParameters<FLandscapeEditLayersWeightmapsMergeEditLayerPS::FParameters>();
+						PSParams->RenderTargets[0] = FRenderTargetBinding(OutputTextureRef, ERenderTargetLoadAction::ENoAction, /*InMipIndex = */0, /*InArraySlice = */InTargetLayerIndex);
+						PSParams->InTargetLayerIndex = InTargetLayerIndex;
+						PSParams->InEditLayerTargetLayerBlendMode = static_cast<uint32>(TargetLayerBlendParams.BlendMode);
+						PSParams->InEditLayerAlpha = TargetLayerBlendParams.Alpha;
+						PSParams->InCurrentEditLayerWeightmaps = CurrentEditLayerTextureSRVRef;
+						PSParams->InPreviousEditLayersWeightmaps = PreviousEditLayersTextureSRVRef;
+
+						FLandscapeEditLayersWeightmapsMergeEditLayerPS::MergeEditLayerPS(
+							RDG_EVENT_NAME("MergeEditLayer(%s) -> %s", *InTargetLayerName.ToString(), *OutputResourceName), GraphBuilder, PSParams, EffectiveTextureSize);
+					};
+
+				// We need to specify the final state of the external textures to prevent the graph builder from transitioning them to SRVMask (even those that end up as SRVMask at the end of this command,
+				//  because they will likely be part of another RDGCommand down the line so we need to maintain an accurate picture of every external texture ever involved in the recorded command so that 
+				//  we can set a proper access when the recorder is flushed (and the FRDGBuilder, executed) :
+				TArray<FRDGBuilderRecorder::FRDGExternalTextureAccessFinal> RDGExternalTextureAccessFinalList =
+				{
+					{ WriteRT->GetRenderTarget()->GetResource(), ERHIAccess::RTV },
+					{ CurrentLayerReadRT->GetRenderTarget()->GetResource(), ERHIAccess::SRVMask },
+					{ PreviousLayersReadRT->GetRenderTarget()->GetResource(), ERHIAccess::SRVMask },
+				};
+				RDGBuilderRecorder.EnqueueRDGCommand(RDGCommand, RDGExternalTextureAccessFinalList);
+
+				return true;
+			});
+
+	}
+}
+
 void FMergeRenderContext::Render(TFunction<void(const FOnRenderBatchTargetGroupDoneParams& /*InParams*/, UE::Landscape::FRDGBuilderRecorder& /*RDGBuilderRecorder*/)> OnBatchTargetGroupDone)
 {
 	using namespace UE::Landscape;
@@ -1150,7 +1386,7 @@ void FMergeRenderContext::Render(TFunction<void(const FOnRenderBatchTargetGroupD
 
 	AllocateResources();
 
-	// Command recorder when accumulating render commands for several consecutive RenderLayer calls when ERenderFlags::Recorded is used 
+	// Command recorder when accumulating render commands for several consecutive RenderLayer calls when ERenderFlags::RenderMode_Recorded is used 
 	FRDGBuilderRecorder RDGBuilderRecorder;
 
 	// Kick start the blend render targets : 
@@ -1190,17 +1426,24 @@ void FMergeRenderContext::Render(TFunction<void(const FOnRenderBatchTargetGroupD
 		int32 RenderLayerStepIndex = 0;
 		// Index of RenderLayer at which we started recording the current render command sequence
 		int32 RenderCommandStartLayerStepIndex = INDEX_NONE;
+		// Indicates how many successful and consecutive RenderLayer steps (i.e. something has been rendered) have occurred (valid until the next BlendLayer step). 
+		//  It's useful for BlendLayer steps as it allows to skip the separate blend step if nothing was rendered prior to it :
+		int32 NumSuccessfulRenderLayerStepsUntilBlendLayerStep = 0;
 		for (int32 RenderStepIndex = 0; RenderStepIndex < NumRenderSteps; ++RenderStepIndex)
 		{
 			const FMergeRenderStep& RenderStep = RenderBatch.RenderSteps[RenderStepIndex];
+			TScriptInterface<ILandscapeEditLayerRenderer> Renderer = RenderStep.RendererState.GetRenderer();
 
 			TArray<FName> TargetLayerGroupLayerNames;
 			TArray<ULandscapeLayerInfoObject*> TargetLayerGroupLayerInfos;
 			TArray<FComponentMergeRenderInfo> SortedComponentMergeRenderInfos;
 
+			FTransform RenderAreaWorldTransform;
+			FIntRect RenderAreaSectionRect;
+
 			// Compute some data for the actual render steps : 
-			if ((RenderStep.Type == FMergeRenderStep::EType::RenderLayer_Immediate)
-				|| (RenderStep.Type == FMergeRenderStep::EType::RenderLayer_Recorded)
+			if ((RenderStep.Type == FMergeRenderStep::EType::RenderLayer)
+				|| (RenderStep.Type == FMergeRenderStep::EType::BlendLayer)
 				|| (RenderStep.Type == FMergeRenderStep::EType::SignalBatchMergeGroupDone))
 			{
 				TargetLayerGroupLayerNames = ConvertTargetLayerBitIndicesToNames(RenderStep.TargetLayerGroupBitIndices);
@@ -1224,7 +1467,24 @@ void FMergeRenderContext::Render(TFunction<void(const FOnRenderBatchTargetGroupD
 			    });
 			    SortedComponentMergeRenderInfos.Sort();
 			}
-			
+
+			// Compute some additional data for the actual render steps : 
+			if ((RenderStep.Type == FMergeRenderStep::EType::BeginRenderLayerGroup)
+				|| (RenderStep.Type == FMergeRenderStep::EType::EndRenderLayerGroup)
+				|| (RenderStep.Type == FMergeRenderStep::EType::RenderLayer)
+				|| (RenderStep.Type == FMergeRenderStep::EType::BlendLayer)
+				|| (RenderStep.Type == FMergeRenderStep::EType::SignalBatchMergeGroupDone))
+			{
+				// TODO[jonathan.bard] offset the world transform to account for the half-pixel offset?
+				//RenderParams.RenderAreaWorldTransform = FTransform(LandscapeTransform.GetRotation(), LandscapeTransform.GetTranslation() + FVector(RenderBatch.SectionRect.Min) /** (double)ComponentSizeQuads - FVector(0.5, 0.5, 0)*/, LandscapeTransform.GetScale3D());
+
+				// TODO [jonathan.bard] : this is more of a Batch world transform/section rect at the moment. Shall we have a RenderAreaWorldTransform/RenderAreaSectionRect in FRenderParams and a BatchRenderAreaWorldTransform in FMergeRenderBatch?
+				//  because currently the old BP brushes work with FMergeRenderBatch data (i.e. 1 transform for the batch and a section rect for the entire batch) but eventually, renderers might be interested in just their Render step context, 
+				//  that is : 1 matrix corresponding to the bottom left corner of their list of components to render?
+				RenderAreaWorldTransform = FTransform(FVector(RenderBatch.SectionRect.Min)) * LandscapeTransform;
+				RenderAreaSectionRect = RenderBatch.SectionRect;
+			}
+
 			switch (RenderStep.Type)
 			{
 			case FMergeRenderStep::EType::BeginRenderCommandRecorder:
@@ -1255,56 +1515,92 @@ void FMergeRenderContext::Render(TFunction<void(const FOnRenderBatchTargetGroupD
 				RenderCommandStartLayerStepIndex = INDEX_NONE;
 				break;
 			}
-			case FMergeRenderStep::EType::RenderLayer_Immediate:
-			case FMergeRenderStep::EType::RenderLayer_Recorded:
+			case FMergeRenderStep::EType::BeginRenderLayerGroup:
+			case FMergeRenderStep::EType::EndRenderLayerGroup:
+			{
+				const ERenderFlags RenderFlags = Renderer->GetRenderFlags(this);
+				check((Renderer != nullptr) 
+					&& EnumHasAnyFlags(RenderFlags, ERenderFlags::RenderMode_Mask) 
+					&& EnumHasAllFlags(RenderFlags, ERenderFlags::BlendMode_SeparateBlend | ERenderFlags::RenderLayerGroup_SupportsGrouping));
+
+				if (RenderStep.Type == FMergeRenderStep::EType::BeginRenderLayerGroup)
 				{
-					TRACE_CPUPROFILER_EVENT_SCOPE(RenderLayer);
-					TScriptInterface<ILandscapeEditLayerRenderer> Renderer = RenderStep.RendererState.GetRenderer();
-					check((Renderer != nullptr) && EnumHasAnyFlags(Renderer->GetRenderFlags(this), ERenderFlags::RenderMode_Mask));
-
-					// TODO[jonathan.bard] offset the world transform to account for the half-pixel offset?
-					//RenderParams.RenderAreaWorldTransform = FTransform(LandscapeTransform.GetRotation(), LandscapeTransform.GetTranslation() + FVector(RenderBatch.SectionRect.Min) /** (double)ComponentSizeQuads - FVector(0.5, 0.5, 0)*/, LandscapeTransform.GetScale3D());
-
-					FString RenderStepProfilingEventName = FString::Format(TEXT("Step [{0}] ({1}): Render {2}"), { RenderLayerStepIndex, *ConvertTargetLayerNamesToString(TargetLayerGroupLayerNames), *Renderer->GetEditLayerRendererDebugName() });
-
-					// TODO [jonathan.bard] : this is more of a Batch world transform/section rect at the moment. Shall we have a RenderAreaWorldTransform/RenderAreaSectionRect in FRenderParams and a BatchRenderAreaWorldTransform in FMergeRenderBatch?
-					//  because currently the old BP brushes work with FMergeRenderBatch data (i.e. 1 transform for the batch and a section rect for the entire batch) but eventually, renderers might be interested in just their Render step context, 
-					//  that is : 1 matrix corresponding to the bottom left corner of their list of components to render?
-					FTransform RenderAreaWorldTransform = FTransform(FVector(RenderBatch.SectionRect.Min)) * LandscapeTransform;
-					FIntRect RenderAreaSectionRect = RenderBatch.SectionRect;
-				
-					// Drop some visual cues to help understand how each renderer is applied :
-					UE_IFVLOG(
-						if (IsVisualLogEnabled() && bVisualLogMergeProcess)
+					// We start a new render layer group, so let's start tracking the number of successful RenderLayer steps within this group :
+					NumSuccessfulRenderLayerStepsUntilBlendLayerStep = 0;
+					ILandscapeEditLayerRenderer::FRenderParams RenderParams(this, TargetLayerGroupLayerNames, TargetLayerGroupLayerInfos, RenderStep.RendererState, SortedComponentMergeRenderInfos, 
+						RenderAreaWorldTransform, RenderAreaSectionRect, NumSuccessfulRenderLayerStepsUntilBlendLayerStep);
+					Renderer->BeginRenderLayerGroup(RenderParams, RDGBuilderRecorder);
+				}
+                else
+                {
+					ILandscapeEditLayerRenderer::FRenderParams RenderParams(this, TargetLayerGroupLayerNames, TargetLayerGroupLayerInfos, RenderStep.RendererState, SortedComponentMergeRenderInfos,
+						RenderAreaWorldTransform, RenderAreaSectionRect, NumSuccessfulRenderLayerStepsUntilBlendLayerStep);
+					Renderer->EndRenderLayerGroup(RenderParams, RDGBuilderRecorder);
+                }
+				break;
+			}
+			case FMergeRenderStep::EType::RenderLayer:
+			case FMergeRenderStep::EType::BlendLayer:
+			{
+				TRACE_CPUPROFILER_EVENT_SCOPE(RenderAndBlendLayer);
+
+				check((Renderer != nullptr) && EnumHasAnyFlags(Renderer->GetRenderFlags(this), ERenderFlags::RenderMode_Mask));
+
+				const bool bIsRecordedStep = EnumHasAnyFlags(RenderStep.RenderFlags, ERenderFlags::RenderMode_Recorded);
+				const bool bIsSeparateBlendStep = (RenderStep.Type == FMergeRenderStep::EType::BlendLayer);
+
+				FString RenderStepProfilingEventName = FString::Format(TEXT("Step [{0}] ({1}): {2} {3}"), { RenderLayerStepIndex, *ConvertTargetLayerNamesToString(TargetLayerGroupLayerNames), bIsSeparateBlendStep ? TEXT("Blend") : TEXT("Render"), *Renderer->GetEditLayerRendererDebugName() });
+
+				// Drop some visual cues to help understand how each renderer is applied :
+				UE_IFVLOG(
+				if (IsVisualLogEnabled() && !bIsSeparateBlendStep && bVisualLogMergeProcess)
+					{
+						FTransform RenderAreaWorldTransformForVisLog = FTransform(FVector(RenderBatch.SectionRect.Min)) * LandscapeWorldTransformForVisLog;
+						IncrementVisualLogOffset();
+						UE_VLOG_LOCATION(Landscape, LogLandscape, Log, ComputeVisualLogTransform(RenderAreaWorldTransformForVisLog).GetTranslation(), 10.0f, FColor::Red, TEXT("%s"), *RenderStepProfilingEventName);
+						UE_VLOG_WIREOBOX(Landscape, LogLandscape, Log, FBox(FVector(RenderBatch.SectionRect.Min) - FVector(0.5, 0.5, 0.0), FVector(RenderBatch.SectionRect.Max) - FVector(0.5, 0.5, 0.0)),
+							ComputeVisualLogTransform(LandscapeWorldTransformForVisLog).ToMatrixWithScale(), FColor::White, TEXT(""));
+
+						// Draw each component's bounds rendered by this renderer : 
+						for (const FComponentMergeRenderInfo& ComponentMergeRenderInfo : SortedComponentMergeRenderInfos)
+						{
+							UE_VLOG_WIREOBOX(Landscape, LogLandscape, Log, FBox(FVector(ComponentMergeRenderInfo.ComponentRegionInRenderArea.Min), FVector(ComponentMergeRenderInfo.ComponentRegionInRenderArea.Max)),
+								ComputeVisualLogTransform(RenderAreaWorldTransformForVisLog).ToMatrixWithScale(), FColor::White, TEXT(""));
+						}
+					});
+		
+				ILandscapeEditLayerRenderer::FRenderParams RenderParams(this, TargetLayerGroupLayerNames, TargetLayerGroupLayerInfos, RenderStep.RendererState, SortedComponentMergeRenderInfos, 
+					RenderAreaWorldTransform, RenderAreaSectionRect, NumSuccessfulRenderLayerStepsUntilBlendLayerStep);
+				auto RenderOrBlend = [&RDGBuilderRecorder, &RenderStepProfilingEventName, &NumSuccessfulRenderLayerStepsUntilBlendLayerStep, &RenderParams, bIsSeparateBlendStep, Renderer]()
+					{
+						if (bIsSeparateBlendStep)
 						{
-							FTransform RenderAreaWorldTransformForVisLog = FTransform(FVector(RenderBatch.SectionRect.Min)) * LandscapeWorldTransformForVisLog;
-							IncrementVisualLogOffset();
-							UE_VLOG_LOCATION(Landscape, LogLandscape, Log, ComputeVisualLogTransform(RenderAreaWorldTransformForVisLog).GetTranslation(), 10.0f, FColor::Red, TEXT("%s"), *RenderStepProfilingEventName);
-							UE_VLOG_WIREOBOX(Landscape, LogLandscape, Log, FBox(FVector(RenderBatch.SectionRect.Min) - FVector(0.5, 0.5, 0.0), FVector(RenderBatch.SectionRect.Max) - FVector(0.5, 0.5, 0.0)),
-								ComputeVisualLogTransform(LandscapeWorldTransformForVisLog).ToMatrixWithScale(), FColor::White, TEXT(""));
-
-							// Draw each component's bounds rendered by this renderer : 
-							for (const FComponentMergeRenderInfo& ComponentMergeRenderInfo : SortedComponentMergeRenderInfos)
+							// Skip the blend if nothing was ever rendered : 
+							if (NumSuccessfulRenderLayerStepsUntilBlendLayerStep > 0)
 							{
-								UE_VLOG_WIREOBOX(Landscape, LogLandscape, Log, FBox(FVector(ComponentMergeRenderInfo.ComponentRegionInRenderArea.Min), FVector(ComponentMergeRenderInfo.ComponentRegionInRenderArea.Max)),
-									ComputeVisualLogTransform(RenderAreaWorldTransformForVisLog).ToMatrixWithScale(), FColor::White, TEXT(""));
+								Renderer->BlendLayer(RenderParams, RDGBuilderRecorder);
 							}
-						});
-			
-					ILandscapeEditLayerRenderer::FRenderParams RenderParams(this, TargetLayerGroupLayerNames, TargetLayerGroupLayerInfos, RenderStep.RendererState, SortedComponentMergeRenderInfos, RenderAreaWorldTransform, RenderAreaSectionRect);
-				    if (RenderStep.Type == FMergeRenderStep::EType::RenderLayer_Recorded)
-				    {
-					    checkf(RDGBuilderRecorder.IsRecording(), TEXT("RenderLayer_Recorded must be preceded by a BeginRenderCommandRecorder which should create a command recorder"));
+							// The blend has occurred, we can now stop tracking the number of successful RenderLayer steps :
+							NumSuccessfulRenderLayerStepsUntilBlendLayerStep = 0;
+						}
+						else
+						{
+							const bool bHasStepRenderedSomething = Renderer->RenderLayer(RenderParams, RDGBuilderRecorder);
+							NumSuccessfulRenderLayerStepsUntilBlendLayerStep += bHasStepRenderedSomething ? 1 : 0;
+						}
+					};
     
+				    if (bIsRecordedStep)
+				    {
+					    checkf(RDGBuilderRecorder.IsRecording(), TEXT("(Render/Blend)Layer_Recorded must be preceded by a BeginRenderCommandRecorder which should create a command recorder"));
 					    RDG_RENDER_COMMAND_RECORDER_BREADCRUMB_EVENT(RDGBuilderRecorder, "%s", *RenderStepProfilingEventName);
-					    Renderer->RenderLayer(RenderParams, RDGBuilderRecorder);
+						RenderOrBlend();
 				    }
 				    else
 				    {
-					    checkf(RDGBuilderRecorder.IsEmpty() && !RDGBuilderRecorder.IsRecording(), TEXT("RenderLayer_Immediate should be preceded by a EndRenderCommandRecorder which should finalize a command recorder and destroy it"));
-    
-					    RHI_BREADCRUMB_EVENT_GAMETHREAD("%s", RenderStepProfilingEventName);
-					    Renderer->RenderLayer(RenderParams, RDGBuilderRecorder);
+					    checkf(RDGBuilderRecorder.IsEmpty() && !RDGBuilderRecorder.IsRecording(), TEXT("(Render/Blend)Layer_Immediate should be preceded by a EndRenderCommandRecorder which should finalize a command recorder and destroy it"));
+					        RHI_BREADCRUMB_EVENT_GAMETHREAD("%s", RenderStepProfilingEventName);
+					    RenderOrBlend();
 				    }
     
 				    ++RenderLayerStepIndex;
@@ -1316,8 +1612,9 @@ void FMergeRenderContext::Render(TFunction<void(const FOnRenderBatchTargetGroupD
 				    RHI_BREADCRUMB_EVENT_GAMETHREAD("Step [%i] (%s) : Render Group Done", RenderLayerStepIndex, ConvertTargetLayerNamesToString(TargetLayerGroupLayerNames));
     
 				    checkf(RDGBuilderRecorder.IsEmpty() && !RDGBuilderRecorder.IsRecording(), TEXT("SignalBatchMergeGroupDone should be preceded by a EndRenderCommandRecorder which should finalize a command recorder and destroy it"));
+
 				    // The last render target we wrote to is the one containing the batch group's merge result : 
-				    FOnRenderBatchTargetGroupDoneParams Params(this, RenderBatch, TargetLayerGroupLayerNames, TargetLayerGroupLayerInfos, SortedComponentMergeRenderInfos);
+				    FOnRenderBatchTargetGroupDoneParams Params(this, TargetLayerGroupLayerNames, TargetLayerGroupLayerInfos, SortedComponentMergeRenderInfos);
 				    OnBatchTargetGroupDone(Params, RDGBuilderRecorder);
 				    break;
 			    }
@@ -1372,6 +1669,14 @@ bool FMergeRenderContext::IsVisualLogEnabled() const
 #endif // ENABLE_VISUAL_LOG
 
 
+// ----------------------------------------------------------------------------------
+
+FBox FOOBox2D::BuildAABB() const
+{
+	return FBox::BuildAABB(Transform.GetTranslation(), Transform.TransformVector(FVector(Extents, 0.0)));
+}
+
+
 // ----------------------------------------------------------------------------------
 
 FIntRect FInputWorldArea::GetLocalComponentKeys(const FIntPoint& InComponentKey) const
@@ -1409,7 +1714,7 @@ FBox FInputWorldArea::ComputeWorldAreaAABB(const FTransform& InLandscapeTransfor
 		}
 		case EType::OOBox:
 		{
-			return FBox::BuildAABB(OOBox2D.Transform.GetTranslation(), OOBox2D.Transform.TransformVector(FVector(OOBox2D.Extents, 0.0)));
+			return OOBox2D.BuildAABB();
 		}
 		default:
 			check(false);
@@ -1483,7 +1788,7 @@ FBox FOutputWorldArea::ComputeWorldAreaAABB(const FTransform& InComponentTransfo
 	}
 	case EType::OOBox:
 	{
-		return FBox::BuildAABB(OOBox2D.Transform.GetTranslation(), OOBox2D.Transform.TransformVector(FVector(OOBox2D.Extents, 0.0)));
+		return OOBox2D.BuildAABB();
 	}
 	default:
 		check(false);
diff --git a/Engine/Source/Runtime/Landscape/Private/LandscapeEditLayerRendererPrivate.cpp b/Engine/Source/Runtime/Landscape/Private/LandscapeEditLayerRendererPrivate.cpp
--- a/Engine/Source/Runtime/Landscape/Private/LandscapeEditLayerRendererPrivate.cpp
+++ b/Engine/Source/Runtime/Landscape/Private/LandscapeEditLayerRendererPrivate.cpp
@@ -43,14 +43,13 @@ FString ULandscapeDefaultEditLayerRenderer::GetEditLayerRendererDebugName() cons
 	return TEXT("Default");
 }
 
-void ULandscapeDefaultEditLayerRenderer::RenderLayer(FRenderParams& RenderParams, UE::Landscape::FRDGBuilderRecorder& RDGBuilderRecorder)
+bool ULandscapeDefaultEditLayerRenderer::RenderLayer(FRenderParams& RenderParams, UE::Landscape::FRDGBuilderRecorder& RDGBuilderRecorder)
 {
 	using namespace UE::Landscape::EditLayers;
 
 	checkf(RDGBuilderRecorder.IsRecording(), TEXT("ERenderFlags::RenderMode_Recorded means the command recorder should be recording at this point"));
 
 	FMergeRenderContext* RenderContext = RenderParams.MergeRenderContext;
-	const FMergeRenderBatch* RenderBatch = RenderContext->GetCurrentRenderBatch();
 
 	RenderContext->CycleBlendRenderTargets(RDGBuilderRecorder);
 	ULandscapeScratchRenderTarget* WriteRT = RenderContext->GetBlendRenderTargetWrite();
@@ -60,7 +59,9 @@ void ULandscapeDefaultEditLayerRenderer::RenderLayer(FRenderParams& RenderParams
 	check(WriteRT->GetCurrentState() == ERHIAccess::RTV);
 
 	// Render the components of the batch for each target layer into the "pseudo-stencil" buffer, so that it can be sampled by users as a UTextures in UMaterials and such : 
-	RenderContext->RenderValidityRenderTargets(*RenderBatch, RDGBuilderRecorder);
+	RenderContext->RenderValidityRenderTargets(RDGBuilderRecorder);
+
+	return true; 
 }
 
 
diff --git a/Engine/Source/Runtime/Landscape/Private/LandscapeEditLayerRendererPrivate.h b/Engine/Source/Runtime/Landscape/Private/LandscapeEditLayerRendererPrivate.h
--- a/Engine/Source/Runtime/Landscape/Private/LandscapeEditLayerRendererPrivate.h
+++ b/Engine/Source/Runtime/Landscape/Private/LandscapeEditLayerRendererPrivate.h
@@ -34,7 +34,7 @@ public:
 	virtual TArray<UE::Landscape::EditLayers::FEditLayerRenderItem> GetRenderItems(const UE::Landscape::EditLayers::FMergeContext* InMergeContext) const override;
 	virtual FString GetEditLayerRendererDebugName() const override;
 	virtual UE::Landscape::EditLayers::ERenderFlags GetRenderFlags(const UE::Landscape::EditLayers::FMergeContext* InMergeContext) const override { return UE::Landscape::EditLayers::ERenderFlags::RenderMode_Recorded; }
-	virtual void RenderLayer(FRenderParams& RenderParams, UE::Landscape::FRDGBuilderRecorder& RDGBuilderRecorder) override;
+	virtual bool RenderLayer(FRenderParams& RenderParams, UE::Landscape::FRDGBuilderRecorder& RDGBuilderRecorder) override;
 	//~ End ILandscapeEditLayerRenderer implementation
 #endif // WITH_EDITOR
 };
@@ -60,7 +60,7 @@ public:
 	virtual TArray<UE::Landscape::EditLayers::FEditLayerRenderItem> GetRenderItems(const UE::Landscape::EditLayers::FMergeContext* InMergeContext) const override;
 	virtual FString GetEditLayerRendererDebugName() const override;
 	virtual UE::Landscape::EditLayers::ERenderFlags GetRenderFlags(const UE::Landscape::EditLayers::FMergeContext* InMergeContext) const override { return UE::Landscape::EditLayers::ERenderFlags::RenderMode_Recorded; }
-	virtual void RenderLayer(FRenderParams& RenderParams, UE::Landscape::FRDGBuilderRecorder& RDGBuilderRecorder) override;
+	virtual bool RenderLayer(FRenderParams& RenderParams, UE::Landscape::FRDGBuilderRecorder& RDGBuilderRecorder) override;
 	//~ End ILandscapeEditLayerRenderer implementation
 #endif // WITH_EDITOR
 };
@@ -84,7 +84,7 @@ public:
 	virtual TArray<UE::Landscape::EditLayers::FEditLayerRenderItem> GetRenderItems(const UE::Landscape::EditLayers::FMergeContext* InMergeContext) const override;
 	virtual FString GetEditLayerRendererDebugName() const override;
 	virtual UE::Landscape::EditLayers::ERenderFlags GetRenderFlags(const UE::Landscape::EditLayers::FMergeContext* InMergeContext) const override { return UE::Landscape::EditLayers::ERenderFlags::RenderMode_Recorded; }
-	virtual void RenderLayer(FRenderParams& RenderParams, UE::Landscape::FRDGBuilderRecorder& RDGBuilderRecorder) override;
+	virtual bool RenderLayer(FRenderParams& RenderParams, UE::Landscape::FRDGBuilderRecorder& RDGBuilderRecorder) override;
 	//~ End ILandscapeEditLayerRenderer implementation
 
 	TBitArray<> GatherWeightBlendedWeightmapLayerBitIndices(const UE::Landscape::EditLayers::FMergeContext* InMergeContext) const;
diff --git a/Engine/Source/Runtime/Landscape/Private/LandscapeEditLayers.cpp b/Engine/Source/Runtime/Landscape/Private/LandscapeEditLayers.cpp
--- a/Engine/Source/Runtime/Landscape/Private/LandscapeEditLayers.cpp
+++ b/Engine/Source/Runtime/Landscape/Private/LandscapeEditLayers.cpp
@@ -254,6 +254,11 @@ TAutoConsoleVariable<FString> CVarLandscapeBatchedMergeVisualLogShowComponentDep
 	"",
 	TEXT("Allows to visualize all the dependency graph for the component specified : use the \"X= Y=\" format to specify the component for which to show the dependencies"));
 
+TAutoConsoleVariable<bool> CVarLandscapeBatchedMergeEnableRenderLayerGroup(
+	TEXT("landscape.BatchedMerge.EnableRenderLayerGroup"),
+	true,
+	TEXT("Allows to batch several non-overlapping successive edit layer renderers (if they support render layer grouping), such that blending is only performed once at the end of the group instead of after each renderer"));
+
 TAutoConsoleVariable<bool> CVarSilenceMergeBatchResolutionWarning(
 	TEXT("landscape.BatchedMerge.SilenceResolutionWarning"),
 	false,
@@ -464,6 +469,31 @@ enum class ELandscapeEditLayerHeightmapBlendMode : uint32
 	Num,
 };
 
+// Must match EEditLayerWeightmapBlendMode in LandscapeLayersWeightmapsPS.usf
+enum class ELandscapeEditLayerWeightmapBlendMode : uint32
+{
+	Additive = 0,
+	Subtractive,
+
+	Num,
+};
+
+// Must match EWeightmapPaintLayerFlags in LandscapeLayersWeightmapsPS.usf
+enum class ELandscapeEditLayerWeightmapPaintLayerFlags : uint32
+{
+	IsVisibilityLayer = (1 << 0), // This paint layer is the visibility layer
+	IsWeightBlended = (1 << 1), // Blend the paint layer's value with all the other paint layers weights
+
+	None = 0
+};
+ENUM_CLASS_FLAGS(ELandscapeEditLayerWeightmapPaintLayerFlags);
+
+// Must match FEditLayerWeightmapPaintLayerInfo in LandscapeLayersWeightmapsPS.usf
+struct FLandscapeEditLayerWeightmapPaintLayerInfo
+{
+	ELandscapeEditLayerWeightmapPaintLayerFlags Flags = ELandscapeEditLayerWeightmapPaintLayerFlags::None; // Additional info about this paint layer
+};
+
 #endif // WITH_EDITOR
 
 namespace UE::Landscape::Private
@@ -497,6 +527,22 @@ namespace UE::Landscape::Private
 
 		return ELandscapeEditLayerHeightmapBlendMode::Num;
 	}
+
+	static UE::Landscape::EditLayers::EHeightmapBlendMode LandscapeBlendModeToHeightmapBlendMode(ELandscapeBlendMode InLandscapeBlendMode)
+	{
+		using namespace UE::Landscape::EditLayers;
+		switch (InLandscapeBlendMode)
+		{
+		case LSBM_AdditiveBlend:
+			return EHeightmapBlendMode::Additive;
+		case LSBM_AlphaBlend:
+			return EHeightmapBlendMode::LegacyAlphaBlend; // LSBM_AlphaBlend corresponds to the landscape spline case, i.e. "legacy alpha blend"
+		default:
+			check(false);
+		}
+
+		return EHeightmapBlendMode::Num;
+	}
 #endif // WITH_EDITOR
 	
 
@@ -1580,7 +1626,7 @@ public:
 		SHADER_PARAMETER(FUintVector4, InPerChannelNumPaintLayers)
 		SHADER_PARAMETER_RDG_TEXTURE_SRV_ARRAY(Texture2D<float4>, InPerChannelPaintLayerWeightmaps, [4])
 		SHADER_PARAMETER_RDG_BUFFER_SRV(StructuredBuffer<uint>, InPaintLayerInfoIndices)
-		SHADER_PARAMETER_RDG_BUFFER_SRV(StructuredBuffer<UE::Landscape::EditLayers::FWeightmapPaintLayerInfo>, InPaintLayerInfos)
+		SHADER_PARAMETER_RDG_BUFFER_SRV(StructuredBuffer<FLandscapeEditLayerWeightmapPaintLayerInfo>, InPaintLayerInfos)
 		RENDER_TARGET_BINDING_SLOTS()
 		END_SHADER_PARAMETER_STRUCT()
 
@@ -1659,48 +1705,6 @@ IMPLEMENT_GLOBAL_SHADER(FLandscapeLayersWeightmapsGenerateMipsPS, "/Engine/Priva
 // ----------------------------------------------------------------------------------
 // /Engine/Private/Landscape/LandscapeEditLayersHeightmaps.usf shaders :
 
-class FLandscapeEditLayersHeightmapsMergeEditLayerPS : public FGlobalShader
-{
-	DECLARE_GLOBAL_SHADER(FLandscapeEditLayersHeightmapsMergeEditLayerPS);
-	SHADER_USE_PARAMETER_STRUCT(FLandscapeEditLayersHeightmapsMergeEditLayerPS, FGlobalShader);
-
-public:
-	BEGIN_SHADER_PARAMETER_STRUCT(FParameters, )
-		SHADER_PARAMETER(uint32, InEditLayerBlendMode)
-		SHADER_PARAMETER(float, InEditLayerAlpha)
-		SHADER_PARAMETER_RDG_TEXTURE_SRV(Texture2D<float4>, InCurrentEditLayerHeightmap)
-		SHADER_PARAMETER_RDG_TEXTURE_SRV(Texture2D<float4>, InPreviousEditLayersHeightmap)
-		RENDER_TARGET_BINDING_SLOTS()
-	END_SHADER_PARAMETER_STRUCT()
-
-	static bool ShouldCompilePermutation(const FGlobalShaderPermutationParameters& InParameters)
-	{
-		return UE::Landscape::DoesPlatformSupportEditLayers(InParameters.Platform);
-	}
-
-	static void ModifyCompilationEnvironment(const FGlobalShaderPermutationParameters& InParameters, FShaderCompilerEnvironment& OutEnvironment)
-	{
-		OutEnvironment.SetDefine(TEXT("MERGE_EDIT_LAYER"), 1);
-	}
-
-	static void MergeEditLayerPS(FRDGEventName&& InRDGEventName, FRDGBuilder& GraphBuilder, FParameters* InParameters, const FIntPoint& InTextureSize)
-	{
-		FGlobalShaderMap* ShaderMap = GetGlobalShaderMap(GMaxRHIFeatureLevel);
-		TShaderMapRef<FLandscapeEditLayersHeightmapsMergeEditLayerPS> PixelShader(ShaderMap);
-
-		FPixelShaderUtils::AddFullscreenPass(
-			GraphBuilder,
-			ShaderMap,
-			MoveTemp(InRDGEventName),
-			PixelShader,
-			InParameters,
-			FIntRect(0, 0, InTextureSize.X, InTextureSize.Y),
-			TStaticBlendStateWriteMask<CW_RG>::GetRHI());
-	}
-};
-
-IMPLEMENT_GLOBAL_SHADER(FLandscapeEditLayersHeightmapsMergeEditLayerPS, "/Engine/Private/Landscape/LandscapeEditLayersHeightmaps.usf", "MergeEditLayerPS", SF_Pixel);
-
 class FLandscapeEditLayersHeightmapsGenerateNormalsPS : public FGlobalShader
 {
 	DECLARE_GLOBAL_SHADER(FLandscapeEditLayersHeightmapsGenerateNormalsPS);
@@ -1789,48 +1793,6 @@ IMPLEMENT_GLOBAL_SHADER(FLandscapeEditLayersHeightmapsGenerateMipsPS, "/Engine/P
 // ----------------------------------------------------------------------------------
 // /Engine/Private/Landscape/LandscapeEditLayersWeightmaps.usf shaders :
 
-class FLandscapeEditLayersWeightmapsMergeEditLayerPS : public FGlobalShader
-{
-	DECLARE_GLOBAL_SHADER(FLandscapeEditLayersWeightmapsMergeEditLayerPS);
-	SHADER_USE_PARAMETER_STRUCT(FLandscapeEditLayersWeightmapsMergeEditLayerPS, FGlobalShader);
-
-public:
-	BEGIN_SHADER_PARAMETER_STRUCT(FParameters, )
-		SHADER_PARAMETER(uint32, InPaintLayerIndex)
-		SHADER_PARAMETER(uint32, InEditLayerPaintLayerBlendMode)
-		SHADER_PARAMETER(float, InEditLayerAlpha)
-		SHADER_PARAMETER_RDG_TEXTURE_SRV(Texture2DArray<float4>, InCurrentEditLayerWeightmaps)
-		SHADER_PARAMETER_RDG_TEXTURE_SRV(Texture2DArray<float4>, InPreviousEditLayersWeightmaps)
-		RENDER_TARGET_BINDING_SLOTS()
-	END_SHADER_PARAMETER_STRUCT()
-
-	static bool ShouldCompilePermutation(const FGlobalShaderPermutationParameters& InParameters)
-	{
-		return UE::Landscape::DoesPlatformSupportEditLayers(InParameters.Platform);
-	}
-
-	static void ModifyCompilationEnvironment(const FGlobalShaderPermutationParameters& InParameters, FShaderCompilerEnvironment& OutEnvironment)
-	{
-		OutEnvironment.SetDefine(TEXT("MERGE_EDIT_LAYER"), 1);
-	}
-
-	static void MergeEditLayerPS(FRDGEventName&& InRDGEventName, FRDGBuilder& GraphBuilder, FParameters* InParameters, const FIntPoint& InTextureSize)
-	{
-		FGlobalShaderMap* ShaderMap = GetGlobalShaderMap(GMaxRHIFeatureLevel);
-		TShaderMapRef<FLandscapeEditLayersWeightmapsMergeEditLayerPS> PixelShader(ShaderMap);
-
-		FPixelShaderUtils::AddFullscreenPass(
-			GraphBuilder,
-			ShaderMap,
-			MoveTemp(InRDGEventName),
-			PixelShader,
-			InParameters,
-			FIntRect(0, 0, InTextureSize.X, InTextureSize.Y));
-	}
-};
-
-IMPLEMENT_GLOBAL_SHADER(FLandscapeEditLayersWeightmapsMergeEditLayerPS, "/Engine/Private/Landscape/LandscapeEditLayersWeightmaps.usf", "MergeEditLayerPS", SF_Pixel);
-
 class FLandscapeEditLayersWeightmapsPerformLegacyWeightBlendingPS : public FGlobalShader
 {
 	DECLARE_GLOBAL_SHADER(FLandscapeEditLayersWeightmapsPerformLegacyWeightBlendingPS);
@@ -1838,9 +1800,9 @@ class FLandscapeEditLayersWeightmapsPerformLegacyWeightBlendingPS : public FGlob
 
 public:
 	BEGIN_SHADER_PARAMETER_STRUCT(FParameters, )
-		SHADER_PARAMETER(uint32, InPaintLayerIndex)
-		SHADER_PARAMETER(uint32, InNumPaintLayers)
-		SHADER_PARAMETER_RDG_BUFFER_SRV(StructuredBuffer<UE::Landscape::FWeightmapPaintLayerInfo>, InPaintLayerInfos)
+		SHADER_PARAMETER(uint32, InTargetLayerIndex)
+		SHADER_PARAMETER(uint32, InNumTargetLayers)
+		SHADER_PARAMETER_RDG_BUFFER_SRV(StructuredBuffer<UE::Landscape::FWeightmapTargetLayerInfo>, InTargetLayerInfos)
 		SHADER_PARAMETER_RDG_TEXTURE_SRV(Texture2DArray<float4>, InCurrentEditLayerWeightmaps)
 		RENDER_TARGET_BINDING_SLOTS()
 	END_SHADER_PARAMETER_STRUCT()
@@ -1866,7 +1828,8 @@ public:
 			MoveTemp(InRDGEventName),
 			PixelShader,
 			InParameters,
-			FIntRect(0, 0, InTextureSize.X, InTextureSize.Y));
+			FIntRect(0, 0, InTextureSize.X, InTextureSize.Y), 
+			TStaticBlendStateWriteMask<CW_RG>::GetRHI());
 	}
 };
 
@@ -1971,9 +1934,8 @@ class FCopyQuadsMultiSourcePS : public FGlobalShader
 		SHADER_PARAMETER_RDG_TEXTURE_ARRAY(Texture2D, InSourceTexture, [NumMultiSources])
 	END_SHADER_PARAMETER_STRUCT()
 
-
-	class FUseSwizleMask : SHADER_PERMUTATION_BOOL("USE_SWIZZLE_MASK");
-	using FPermutationDomain = TShaderPermutationDomain<FUseSwizleMask>;
+	class FCopyWeightmap : SHADER_PERMUTATION_BOOL("COPY_WEIGHTMAP");
+	using FPermutationDomain = TShaderPermutationDomain<FCopyWeightmap>;
 
 	static bool ShouldCompilePermutation(const FGlobalShaderPermutationParameters& InParameters)
 	{
@@ -4336,13 +4298,12 @@ struct FLandscapeEditLayerHeightmapMergeInfo
 	uint32 Padding1;
 };
 
-
 // Must match FEditLayerWeightmapMergeInfo in LandscapeLayersWeightmapsPS.usf
 struct FLandscapeEditLayerWeightmapMergeInfo
 {
 	uint32 SourceWeightmapTextureIndex = (uint32)INDEX_NONE; // The index in InPackedWeightmaps of the texture to read from for this layer
 	uint32 SourceWeightmapTextureChannel = (uint32)INDEX_NONE; // The channel of the texture to read from for this layer
-	UE::Landscape::EditLayers::EWeightmapBlendMode BlendMode = UE::Landscape::EditLayers::EWeightmapBlendMode::Num; // How this layer blends with the previous ones in the layers stack
+	ELandscapeEditLayerWeightmapBlendMode BlendMode = ELandscapeEditLayerWeightmapBlendMode::Num; // How this layer blends with the previous ones in the layers stack
 	float Alpha = 1.0f; // Alpha value to be used in the blend
 };
 
@@ -4350,7 +4311,7 @@ struct FLandscapeEditLayerWeightmapMergeInfo
 // Additional info about this paint layer on this edit layer
 struct FLandscapePerEditLayerWeightmapPaintLayerInfo
 {
-	UE::Landscape::EditLayers::EWeightmapBlendMode Flags = UE::Landscape::EditLayers::EWeightmapBlendMode::Num; // How this layer blends with the previous ones in the layers stack
+	ELandscapeEditLayerWeightmapBlendMode Flags = ELandscapeEditLayerWeightmapBlendMode::Num; // How this layer blends with the previous ones in the layers stack
 };
 
 // Struct that contains all the information relevant for the edit layers update operation (list of dirty components, heightmaps, weightmaps, etc.
@@ -5416,7 +5377,9 @@ ELandscapeToolTargetTypeFlags FLandscapeLayer::GetEnabledTargetTypeMask() const
 	ELandscapeToolTargetTypeFlags EnabledTargetTypeMask = ELandscapeToolTargetTypeFlags::None;
 	if (bVisible)
 	{
-		if (HeightmapAlpha != 0.0f)
+		// TODO [jonathan.bard] : Remove the splines type check when HeightmapAlpha is replaced by a virtual GetHeightmapAlpha() on ULandscapeEditLayer (alpha is fixed to 1.0 in this case, but 
+		//  HeightmapAlpha might still be set to a different value)
+		if (EditLayer->IsA<ULandscapeEditLayerSplines>() || (HeightmapAlpha != 0.0f))
 		{
 			EnabledTargetTypeMask |= ELandscapeToolTargetTypeFlags::Heightmap;
 		}
@@ -5460,6 +5423,12 @@ void ULandscapeEditLayerPersistent::GetRendererStateInfo(const UE::Landscape::Ed
 	OutEnabledTargetTypeState = OutSupportedTargetTypeState;
 }
 
+UE::Landscape::EditLayers::ERenderFlags ULandscapeEditLayerPersistent::GetRenderFlags(const UE::Landscape::EditLayers::FMergeContext* InMergeContext) const
+{
+	using namespace UE::Landscape::EditLayers;
+	return ERenderFlags::RenderMode_Recorded | ERenderFlags::BlendMode_SeparateBlend; // Supports the command recorder and has a separate BlendLayer function 
+}
+
 TArray<UE::Landscape::EditLayers::FEditLayerRenderItem> ULandscapeEditLayerPersistent::GetRenderItems(const UE::Landscape::EditLayers::FMergeContext* InMergeContext) const
 {
 	using namespace UE::Landscape::EditLayers;
@@ -5526,7 +5495,7 @@ TArray<UE::Landscape::EditLayers::FEditLayerRenderItem> ULandscapeEditLayerPersi
 	return RenderItems;
 }
 
-void ULandscapeEditLayerPersistent::RenderLayer(FRenderParams& RenderParams, UE::Landscape::FRDGBuilderRecorder& RDGBuilderRecorder)
+bool ULandscapeEditLayerPersistent::RenderLayer(FRenderParams& RenderParams, UE::Landscape::FRDGBuilderRecorder& RDGBuilderRecorder)
 {
 	using namespace UE::Landscape;
 	using namespace UE::Landscape::Private;
@@ -5535,11 +5504,10 @@ void ULandscapeEditLayerPersistent::RenderLayer(FRenderParams& RenderParams, UE:
 	checkf(RDGBuilderRecorder.IsRecording(), TEXT("ERenderFlags::RenderMode_Recorded means the command recorder should be recording at this point"));
 	check(!RenderParams.SortedComponentMergeRenderInfos.IsEmpty());
 
-	ULandscapeInfo* Info = RenderParams.MergeRenderContext->GetLandscape()->GetLandscapeInfo();
+	ULandscapeInfo* Info = RenderParams.MergeRenderContext->GetLandscapeInfo();
 	check(Info != nullptr);
 	ALandscape* Landscape = RenderParams.MergeRenderContext->GetLandscape();
 	FSceneInterface* SceneInterface = Landscape->GetWorld()->Scene;
-
 	const FLandscapeLayer* LayerStruct = GetOwningLayer();
 	check(LayerStruct != nullptr);
 
@@ -5550,7 +5518,6 @@ void ULandscapeEditLayerPersistent::RenderLayer(FRenderParams& RenderParams, UE:
 	//  for large landscapes. Let's first build a list of quads to render and we'll render them all as efficiently as possible thereafter : 
 
 	// For each subsection, we'll add a quad to render, using the following struct :
-	bool bUseSwizzleMask = false;
 	struct FCopyQuadParams
 	{
 		// Texture to read from :
@@ -5562,10 +5529,10 @@ void ULandscapeEditLayerPersistent::RenderLayer(FRenderParams& RenderParams, UE:
 		// In case the destination is a texture array (weightmaps), this is the slice index to write to :
 		int32 DestinationArrayIndex = 0;
 		// Weightmaps are packed in the source texture. This specifies which channel to read from in the source texture for this quad :
-		uint8 ChannelSwizzleMask = RGBAToRGBASwizzleMask;
+		uint8 SourceChannelIndex = 0;
 	};
 	TArray<FCopyQuadParams> CopyQuadsParams;
-	CopyQuadsParams.Reserve(RenderParams.SortedComponentMergeRenderInfos.Num() * NumTargetLayersInGroup* Landscape->NumSubsections * Landscape->NumSubsections);
+	CopyQuadsParams.Reserve(RenderParams.SortedComponentMergeRenderInfos.Num() * NumTargetLayersInGroup * Landscape->NumSubsections * Landscape->NumSubsections);
 
 	// This is kinda hacky, but we since we reuse AddRasterizeToRectsPass and it passes the source texture's sample coordinates from VS to PS via the UVs, dividing the source rect by a single TextureSize, 
 	//  we correct each quad's source rect by a common scale factor so that it samples at the right location in each source texture: 
@@ -5580,19 +5547,18 @@ void ULandscapeEditLayerPersistent::RenderLayer(FRenderParams& RenderParams, UE:
 			const FName TargetLayerName = RenderParams.TargetLayerGroupLayerNames[TargetLayerIndexInGroup];
 			UTexture2D* SourceTexture = nullptr;
 			FVector2D SourceTextureBias(ForceInit);
-			// Copy all channels by default :
-			uint8 CopyTextureSwizzleMask = RGBAToRGBASwizzleMask;
+			uint8 SourceChannelIndex = 0;
 
 			if (RenderParams.MergeRenderContext->IsHeightmapMerge())
 			{
 				SourceTexture = ComponentMergeRenderInfo.Component->GetHeightmap(LayerStruct->Guid);
 				SourceTextureBias = FVector2D(ComponentMergeRenderInfo.Component->HeightmapScaleBias.Z, ComponentMergeRenderInfo.Component->HeightmapScaleBias.W);
 			}
-			else 
+			else
 			{
 				const TArray<UTexture2D*>& WeightmapTextures = ComponentMergeRenderInfo.Component->GetWeightmapTextures(LayerStruct->Guid);
 				const TArray<FWeightmapLayerAllocationInfo>& AllocInfos = ComponentMergeRenderInfo.Component->GetWeightmapLayerAllocations(LayerStruct->Guid);
-				const FWeightmapLayerAllocationInfo* AllocInfo = AllocInfos.FindByPredicate([TargetLayerName](const FWeightmapLayerAllocationInfo& InAllocInfo) { return InAllocInfo.IsAllocated() && (InAllocInfo.GetLayerName() == TargetLayerName);});
+				const FWeightmapLayerAllocationInfo* AllocInfo = AllocInfos.FindByPredicate([TargetLayerName](const FWeightmapLayerAllocationInfo& InAllocInfo) { return InAllocInfo.IsAllocated() && (InAllocInfo.GetLayerName() == TargetLayerName); });
 				if (AllocInfo != nullptr)
 				{
 					SourceTexture = WeightmapTextures[AllocInfo->WeightmapTextureIndex];
@@ -5600,9 +5566,8 @@ void ULandscapeEditLayerPersistent::RenderLayer(FRenderParams& RenderParams, UE:
 					// Note : don't use WeightmapScaleBias here, it has a different meaning than HeightmapScaleBias (very conveniently!) : this is compensated by the FloorToInt32 later on, 
 					//  but still, let's set this to zero here and use the fact that there's no texture sharing on weightmaps : 
 					SourceTextureBias = FVector2D::ZeroVector;
-					// Swizzle the appropriate source channel into the R component (as the target texture array's format is R8) : 
-					CopyTextureSwizzleMask = BuildChannelSwizzleMask(static_cast<ERGBAChannel>(AllocInfo->WeightmapTextureChannel));
-					bUseSwizzleMask = true;
+					// Copy from the appropriate source channel : 
+					SourceChannelIndex = AllocInfo->WeightmapTextureChannel;
 				}
 			}
 
@@ -5613,12 +5578,12 @@ void ULandscapeEditLayerPersistent::RenderLayer(FRenderParams& RenderParams, UE:
 				// We get the overall texture size via the resource instead of direct GetSizeX/Y calls because the latter is unreliable while the texture is being built.
 				const FVector2D SourceTextureSize(SourceTextureResource->GetSizeX(), SourceTextureResource->GetSizeY());
 				const FIntPoint SourceTextureOffset(FMath::FloorToInt32(SourceTextureBias.X * SourceTextureSize.X), FMath::FloorToInt32(SourceTextureBias.Y * SourceTextureSize.Y));
-				
+
 				auto SourceTexturePixelCoordinatesToQuadCoords = [&SourceTextureSize, &CommonSourceTextureSize](const FIntPoint& InSourceCoords)
-				{ 
-					FVector2D Result = FVector2D(InSourceCoords) / SourceTextureSize * CommonSourceTextureSize;
-					return FIntPoint(FMath::FloorToInt32(Result.X), FMath::FloorToInt32(Result.Y));
-				};
+					{
+						FVector2D Result = FVector2D(InSourceCoords) / SourceTextureSize * CommonSourceTextureSize;
+						return FIntPoint(FMath::FloorToInt32(Result.X), FMath::FloorToInt32(Result.Y));
+					};
 
 				// Fill that render target subsection by subsection, in order to bypass the redundant columns/lines on the subsection edges:
 				for (int32 SubsectionY = 0; SubsectionY < Landscape->NumSubsections; ++SubsectionY)
@@ -5634,7 +5599,7 @@ void ULandscapeEditLayerPersistent::RenderLayer(FRenderParams& RenderParams, UE:
 						QuadParams.SourceRect = FIntRect(SourceTexturePixelCoordinatesToQuadCoords(SourcePosition), SourceTexturePixelCoordinatesToQuadCoords(SourcePosition + FIntPoint(SubsectionSizeVerts, SubsectionSizeVerts)));
 						QuadParams.DestinationRect = FIntRect(DestinationPosition, DestinationPosition + FIntPoint(SubsectionSizeVerts, SubsectionSizeVerts));
 						QuadParams.DestinationArrayIndex = TargetLayerIndexInGroup;
-						QuadParams.ChannelSwizzleMask = CopyTextureSwizzleMask;
+						QuadParams.SourceChannelIndex = SourceChannelIndex;
 					}
 				}
 			}
@@ -5645,7 +5610,7 @@ void ULandscapeEditLayerPersistent::RenderLayer(FRenderParams& RenderParams, UE:
 	{
 		// No need to do anything if there's nothing to be rendered (e.g. no weightmap on the rendered area), this layer will just be ineffective on this batch
 		//  and we don't need to cycle the blend render targets, that will save some processing:
-		return;
+		return false;
 	}
 
 	RenderParams.MergeRenderContext->CycleBlendRenderTargets(RDGBuilderRecorder);
@@ -5655,307 +5620,230 @@ void ULandscapeEditLayerPersistent::RenderLayer(FRenderParams& RenderParams, UE:
 	// We will write to the RT using a (bunch of) PS : 
 	check(WriteRT->GetCurrentState() == ERHIAccess::RTV);
 
-	// First step : render the edit layer's component quads
-	{
-		// Sort the list of quads by output texture first, then by source texture, because we'll coalesce several texture copies into the same pass (one output, multiple inputs)
-		CopyQuadsParams.Sort([](const FCopyQuadParams& InLHS, const FCopyQuadParams& InRHS)
+	// In this render step, we'll simply render the edit layer's component quads
+
+	// Sort the list of quads by output texture first, then by source texture, because we'll coalesce several texture copies into the same pass (one output, multiple inputs)
+	CopyQuadsParams.Sort([](const FCopyQuadParams& InLHS, const FCopyQuadParams& InRHS)
+		{
+			if (InLHS.DestinationArrayIndex == InRHS.DestinationArrayIndex)
 			{
-				if (InLHS.DestinationArrayIndex == InRHS.DestinationArrayIndex)
+				if (InLHS.SourceTextureResource == InRHS.SourceTextureResource)
 				{
-					if (InLHS.SourceTextureResource == InRHS.SourceTextureResource)
-					{
-						return (InLHS.SourceRect.Min.Y == InRHS.SourceRect.Min.Y)
-							? (InLHS.SourceRect.Min.X < InRHS.SourceRect.Min.X)
-							: (InLHS.SourceRect.Min.Y < InRHS.SourceRect.Min.Y);
-					}
-					return (InLHS.SourceTextureResource < InRHS.SourceTextureResource);
+					return (InLHS.SourceRect.Min.Y == InRHS.SourceRect.Min.Y)
+						? (InLHS.SourceRect.Min.X < InRHS.SourceRect.Min.X)
+						: (InLHS.SourceRect.Min.Y < InRHS.SourceRect.Min.Y);
 				}
+				return (InLHS.SourceTextureResource < InRHS.SourceTextureResource);
+			}
 
-				return (InLHS.DestinationArrayIndex < InRHS.DestinationArrayIndex);
-			});
-
-		// Now process this list of quads and prepare as many passes as necessary for performing all the copies :
-		auto RDGCommand =
-			[CopyQuadsParams = MoveTemp(CopyQuadsParams)
-			, OutputResource = WriteRT->GetRenderTarget()->GetResource()
-			, OutputResourceName = WriteRT->GetDebugName()
-			, SceneInterface
-			, CommonSourceTextureSize
-			, bUseSwizzleMask] (FRDGBuilder& GraphBuilder)
-			{
-				RDG_EVENT_SCOPE(GraphBuilder, "CopyEditLayer -> %s", *OutputResourceName);
-
-				FRDGTextureSRVRef BlackDummySRVRef = GraphBuilder.CreateSRV(FRDGTextureSRVDesc::Create(GSystemTextures.GetBlackDummy(GraphBuilder)));
+			return (InLHS.DestinationArrayIndex < InRHS.DestinationArrayIndex);
+		});
 
-				FCopyQuadsMultiSourcePS::FPermutationDomain PermutationVector;
-				PermutationVector.Set< FCopyQuadsMultiSourcePS::FUseSwizleMask >(bUseSwizzleMask);
-				FGlobalShaderMap* ShaderMap = GetGlobalShaderMap(GMaxRHIFeatureLevel);
-				TShaderRef<FCopyQuadsMultiSourcePS> PixelShader = ShaderMap->GetShader<FCopyQuadsMultiSourcePS>(PermutationVector);
-
-				// TODO [jonathan.bard] this is just to avoid a RHI validation error for unoptimized shaders... once validation is made to not issue those errors, we can remove this
-				// Create a SceneView to please the shader bindings, but it's unused in practice 
-				FSceneViewFamilyContext ViewFamily(FSceneViewFamily::ConstructionValues(nullptr, SceneInterface, FEngineShowFlags(ESFIM_Game)).SetTime(FGameTime::GetTimeSinceAppStart()));
-				FSceneViewInitOptions ViewInitOptions;
-				ViewInitOptions.ViewFamily = &ViewFamily;
-				ViewInitOptions.SetViewRectangle(FIntRect(0, 0, 1, 1)); // Use a dummy rect to avoid a check(slow)
-				GetRendererModule().CreateAndInitSingleView(GraphBuilder.RHICmdList, &ViewFamily, &ViewInitOptions);
-				const FSceneView* View = ViewFamily.Views[0];
-
-				// The following variables allow to accumulate quads for rendering in a single pass. Then, we flush them all out when the pass is full and we start again : 
-				int32 CurrentDestinationArrayIndex = INDEX_NONE;
-				TArray<FTextureResource*> CurrentSourceTextureResources;
-				TArray<FUintVector2> CurrentQuadInfos;
-				TArray<FUintVector4> CurrentSourceRects, CurrentDestinationRects;
-
-				FRDGTextureRef OutputTexture = GraphBuilder.RegisterExternalTexture(CreateRenderTarget(OutputResource->GetTextureRHI(), TEXT("OutputTexture")));
-
-				// When we've reached the limit of what we can render in one pass, flush all and add a render pass : 
-				auto FlushCopyQuads = [&GraphBuilder, View, OutputTexture, &CurrentDestinationArrayIndex, &CurrentSourceTextureResources, &CurrentSourceRects, &CurrentDestinationRects,
-					&CurrentQuadInfos, &CommonSourceTextureSize, BlackDummySRVRef, ShaderMap, PixelShader]()
+	// Now process this list of quads and prepare as many passes as necessary for performing all the copies :
+	auto RDGCommand =
+		[CopyQuadsParams = MoveTemp(CopyQuadsParams)
+		, OutputResource = WriteRT->GetRenderTarget()->GetResource()
+		, OutputResourceName = WriteRT->GetDebugName()
+		, SceneInterface
+		, CommonSourceTextureSize
+		, bIsWeightmapMerge = !RenderParams.MergeRenderContext->IsHeightmapMerge()] (FRDGBuilder& GraphBuilder)
+		{
+			RDG_EVENT_SCOPE(GraphBuilder, "CopyEditLayer -> %s", *OutputResourceName);
+
+			FRDGTextureSRVRef BlackDummySRVRef = GraphBuilder.CreateSRV(FRDGTextureSRVDesc::Create(GSystemTextures.GetBlackDummy(GraphBuilder)));
+
+			FCopyQuadsMultiSourcePS::FPermutationDomain PermutationVector;
+			PermutationVector.Set< FCopyQuadsMultiSourcePS::FCopyWeightmap >(bIsWeightmapMerge);
+			FGlobalShaderMap* ShaderMap = GetGlobalShaderMap(GMaxRHIFeatureLevel);
+			TShaderRef<FCopyQuadsMultiSourcePS> PixelShader = ShaderMap->GetShader<FCopyQuadsMultiSourcePS>(PermutationVector);
+
+			// TODO [jonathan.bard] this is just to avoid a RHI validation error for unoptimized shaders... once validation is made to not issue those errors, we can remove this
+			// Create a SceneView to please the shader bindings, but it's unused in practice 
+			FSceneViewFamilyContext ViewFamily(FSceneViewFamily::ConstructionValues(nullptr, SceneInterface, FEngineShowFlags(ESFIM_Game)).SetTime(FGameTime::GetTimeSinceAppStart()));
+			FSceneViewInitOptions ViewInitOptions;
+			ViewInitOptions.ViewFamily = &ViewFamily;
+			ViewInitOptions.SetViewRectangle(FIntRect(0, 0, 1, 1)); // Use a dummy rect to avoid a check(slow)
+			GetRendererModule().CreateAndInitSingleView(GraphBuilder.RHICmdList, &ViewFamily, &ViewInitOptions);
+			const FSceneView* View = ViewFamily.Views[0];
+
+			// The following variables allow to accumulate quads for rendering in a single pass. Then, we flush them all out when the pass is full and we start again : 
+			int32 CurrentDestinationArrayIndex = INDEX_NONE;
+			TArray<FTextureResource*> CurrentSourceTextureResources;
+			TArray<FUintVector2> CurrentQuadInfos;
+			TArray<FUintVector4> CurrentSourceRects, CurrentDestinationRects;
+
+			FRDGTextureRef OutputTexture = GraphBuilder.RegisterExternalTexture(CreateRenderTarget(OutputResource->GetTextureRHI(), TEXT("OutputTexture")));
+
+			// When we've reached the limit of what we can render in one pass, flush all and add a render pass : 
+			auto FlushCopyQuads = [&GraphBuilder, View, OutputTexture, &CurrentDestinationArrayIndex, &CurrentSourceTextureResources, &CurrentSourceRects, &CurrentDestinationRects,
+				&CurrentQuadInfos, &CommonSourceTextureSize, BlackDummySRVRef, ShaderMap, PixelShader]()
+				{
+					if (CurrentSourceTextureResources.IsEmpty())
 					{
-						if (CurrentSourceTextureResources.IsEmpty())
-						{
-							// Nothing to flush :
-							check(CurrentSourceRects.IsEmpty() && CurrentDestinationRects.IsEmpty() && CurrentQuadInfos.IsEmpty());
-							return;
-						}
+						// Nothing to flush :
+						check(CurrentSourceRects.IsEmpty() && CurrentDestinationRects.IsEmpty() && CurrentQuadInfos.IsEmpty());
+						return;
+					}
 
-						check(!CurrentSourceRects.IsEmpty() && (CurrentSourceRects.Num() == CurrentDestinationRects.Num()) && (CurrentSourceRects.Num() == CurrentQuadInfos.Num()));
+					check(!CurrentSourceRects.IsEmpty() && (CurrentSourceRects.Num() == CurrentDestinationRects.Num()) && (CurrentSourceRects.Num() == CurrentQuadInfos.Num()));
 
-						FRDGBufferRef RectBuffer = CreateUploadBuffer(GraphBuilder, TEXT("DestinationRects"), TConstArrayView<FUintVector4>(MakeArrayView(CurrentDestinationRects)));
-						FRDGBufferSRVRef RectBufferSRV = GraphBuilder.CreateSRV(FRDGBufferSRVDesc(RectBuffer, PF_R32G32B32A32_UINT));
+					FRDGBufferRef RectBuffer = CreateUploadBuffer(GraphBuilder, TEXT("DestinationRects"), TConstArrayView<FUintVector4>(MakeArrayView(CurrentDestinationRects)));
+					FRDGBufferSRVRef RectBufferSRV = GraphBuilder.CreateSRV(FRDGBufferSRVDesc(RectBuffer, PF_R32G32B32A32_UINT));
 
-						FRDGBufferRef RectUVBuffer = CreateUploadBuffer(GraphBuilder, TEXT("RectUVs"), TConstArrayView<FUintVector4>(MakeArrayView(CurrentSourceRects)));
-						FRDGBufferSRVRef RectUVBufferSRV = GraphBuilder.CreateSRV(FRDGBufferSRVDesc(RectUVBuffer, PF_R32G32B32A32_UINT));
+					FRDGBufferRef RectUVBuffer = CreateUploadBuffer(GraphBuilder, TEXT("RectUVs"), TConstArrayView<FUintVector4>(MakeArrayView(CurrentSourceRects)));
+					FRDGBufferSRVRef RectUVBufferSRV = GraphBuilder.CreateSRV(FRDGBufferSRVDesc(RectUVBuffer, PF_R32G32B32A32_UINT));
 
-						FRDGBufferRef QuadInfosBuffer = CreateUploadBuffer(GraphBuilder, TEXT("QuadInfos"), TConstArrayView<FUintVector2>(MakeArrayView(CurrentQuadInfos)));
-						FRDGBufferSRVRef QuadInfosBufferSRV = GraphBuilder.CreateSRV(FRDGBufferSRVDesc(QuadInfosBuffer, PF_R32G32_UINT));
+					FRDGBufferRef QuadInfosBuffer = CreateUploadBuffer(GraphBuilder, TEXT("QuadInfos"), TConstArrayView<FUintVector2>(MakeArrayView(CurrentQuadInfos)));
+					FRDGBufferSRVRef QuadInfosBufferSRV = GraphBuilder.CreateSRV(FRDGBufferSRVDesc(QuadInfosBuffer, PF_R32G32_UINT));
 
-						// We may target an array slice directly : 
-						int32 ArrayIndex = -1;
-						if (OutputTexture->Desc.IsTextureArray())
-						{
-							check(CurrentDestinationArrayIndex < OutputTexture->Desc.ArraySize);
-							ArrayIndex = CurrentDestinationArrayIndex;
-						}
+					// We may target an array slice directly : 
+					int32 ArrayIndex = -1;
+					if (OutputTexture->Desc.IsTextureArray())
+					{
+						check(CurrentDestinationArrayIndex < OutputTexture->Desc.ArraySize);
+						ArrayIndex = CurrentDestinationArrayIndex;
+					}
 
-						FCopyQuadsMultiSourcePSParameters* PassParameters = GraphBuilder.AllocParameters<FCopyQuadsMultiSourcePSParameters>();
-						PassParameters->RenderTargets[0] = FRenderTargetBinding(OutputTexture, ERenderTargetLoadAction::ELoad, /*InMipIndex = */0, ArrayIndex);
-						PassParameters->PS.View = View->ViewUniformBuffer;
-						PassParameters->PS.InQuadInfos = QuadInfosBufferSRV;
-						check(CurrentSourceTextureResources.Num() <= FCopyQuadsMultiSourcePS::NumMultiSources);
-						for (int32 TextureIndex = 0; TextureIndex < FCopyQuadsMultiSourcePS::NumMultiSources; ++TextureIndex)
+					FCopyQuadsMultiSourcePSParameters* PassParameters = GraphBuilder.AllocParameters<FCopyQuadsMultiSourcePSParameters>();
+					PassParameters->RenderTargets[0] = FRenderTargetBinding(OutputTexture, ERenderTargetLoadAction::ELoad, /*InMipIndex = */0, ArrayIndex);
+					PassParameters->PS.View = View->ViewUniformBuffer;
+					PassParameters->PS.InQuadInfos = QuadInfosBufferSRV;
+					check(CurrentSourceTextureResources.Num() <= FCopyQuadsMultiSourcePS::NumMultiSources);
+					for (int32 TextureIndex = 0; TextureIndex < FCopyQuadsMultiSourcePS::NumMultiSources; ++TextureIndex)
+					{
+						FRDGTextureRef SourceTexture = BlackDummySRVRef->GetParent();
+						if (CurrentSourceTextureResources.IsValidIndex(TextureIndex))
 						{
-							FRDGTextureRef SourceTexture = BlackDummySRVRef->GetParent();
-							if (CurrentSourceTextureResources.IsValidIndex(TextureIndex))
-							{
-								SourceTexture = GraphBuilder.RegisterExternalTexture(CreateRenderTarget(CurrentSourceTextureResources[TextureIndex]->GetTexture2DRHI(), TEXT("SourceTexture")));
-							}
-							PassParameters->PS.InSourceTexture[TextureIndex] = SourceTexture;
+							SourceTexture = GraphBuilder.RegisterExternalTexture(CreateRenderTarget(CurrentSourceTextureResources[TextureIndex]->GetTexture2DRHI(), TEXT("SourceTexture")));
 						}
+						PassParameters->PS.InSourceTexture[TextureIndex] = SourceTexture;
+					}
 
-						FPixelShaderUtils::AddRasterizeToRectsPass<FCopyQuadsMultiSourcePS>(GraphBuilder,
-							ShaderMap,
-							RDG_EVENT_NAME("CopyQuadsMultiSourcePS"),
-							PixelShader,
-							PassParameters,
-							/*ViewportSize = */OutputTexture->Desc.Extent,
-							RectBufferSRV,
-							CurrentDestinationRects.Num(),
-							/*BlendState = */nullptr,
-							/*RasterizerState = */nullptr,
-							/*DepthStencilState = */nullptr,
-							/*StencilRef = */0,
-							/*TextureSize = */FIntPoint(FMath::FloorToInt32(CommonSourceTextureSize.X), FMath::FloorToInt32(CommonSourceTextureSize.Y)),
-							RectUVBufferSRV);
-
-						// Reset all for the next pass that comes along : 
-						CurrentDestinationArrayIndex = INDEX_NONE;
-						CurrentSourceTextureResources.Empty();
-						CurrentSourceRects.Empty();
-						CurrentDestinationRects.Empty();
-						CurrentQuadInfos.Empty();
-					};
+					FPixelShaderUtils::AddRasterizeToRectsPass<FCopyQuadsMultiSourcePS>(GraphBuilder,
+						ShaderMap,
+						RDG_EVENT_NAME("CopyQuadsMultiSourcePS"),
+						PixelShader,
+						PassParameters,
+						/*ViewportSize = */OutputTexture->Desc.Extent,
+						RectBufferSRV,
+						CurrentDestinationRects.Num(),
+						/*BlendState = */nullptr,
+						/*RasterizerState = */nullptr,
+						/*DepthStencilState = */nullptr,
+						/*StencilRef = */0,
+						/*TextureSize = */FIntPoint(FMath::FloorToInt32(CommonSourceTextureSize.X), FMath::FloorToInt32(CommonSourceTextureSize.Y)),
+						RectUVBufferSRV);
+
+					// Reset all for the next pass that comes along : 
+					CurrentDestinationArrayIndex = INDEX_NONE;
+					CurrentSourceTextureResources.Empty();
+					CurrentSourceRects.Empty();
+					CurrentDestinationRects.Empty();
+					CurrentQuadInfos.Empty();
+				};
 
-				for (const FCopyQuadParams& InSingleQuadParams : CopyQuadsParams)
+			for (const FCopyQuadParams& InSingleQuadParams : CopyQuadsParams)
+			{
+				check(InSingleQuadParams.DestinationArrayIndex != INDEX_NONE);
+				// If the output texture/slice has changed since the last iteration, it's time to flush and start a new pass :
+				if (CurrentDestinationArrayIndex != InSingleQuadParams.DestinationArrayIndex)
 				{
-					check(InSingleQuadParams.DestinationArrayIndex != INDEX_NONE);
-					// If the output texture/slice has changed since the last iteration, it's time to flush and start a new pass :
-					if (CurrentDestinationArrayIndex != InSingleQuadParams.DestinationArrayIndex)
+					if (CurrentDestinationArrayIndex != INDEX_NONE)
 					{
-						if (CurrentDestinationArrayIndex != INDEX_NONE)
-						{
-							FlushCopyQuads();
-						}
-						CurrentDestinationArrayIndex = InSingleQuadParams.DestinationArrayIndex;
+						FlushCopyQuads();
 					}
+					CurrentDestinationArrayIndex = InSingleQuadParams.DestinationArrayIndex;
+				}
 
-					if (CurrentSourceTextureResources.IsEmpty())
-					{
-						CurrentSourceTextureResources.Add(InSingleQuadParams.SourceTextureResource);
-					}
-					else if (InSingleQuadParams.SourceTextureResource != CurrentSourceTextureResources.Last())
+				if (CurrentSourceTextureResources.IsEmpty())
+				{
+					CurrentSourceTextureResources.Add(InSingleQuadParams.SourceTextureResource);
+				}
+				else if (InSingleQuadParams.SourceTextureResource != CurrentSourceTextureResources.Last())
+				{
+					// If we've reached the amount of textures we can render in a single pass, we flush the pass and initiate a new one : 
+					if (CurrentSourceTextureResources.Num() == FCopyQuadsMultiSourcePS::NumMultiSources)
 					{
-						// If we've reached the amount of textures we can render in a single pass, we flush the pass and initiate a new one : 
-						if (CurrentSourceTextureResources.Num() == FCopyQuadsMultiSourcePS::NumMultiSources)
-						{
-							FlushCopyQuads();
-							check(CurrentSourceTextureResources.IsEmpty());
+						FlushCopyQuads();
+						check(CurrentSourceTextureResources.IsEmpty());
 
-							CurrentDestinationArrayIndex = InSingleQuadParams.DestinationArrayIndex;
-						}
-						CurrentSourceTextureResources.Add(InSingleQuadParams.SourceTextureResource);
+						CurrentDestinationArrayIndex = InSingleQuadParams.DestinationArrayIndex;
 					}
+					CurrentSourceTextureResources.Add(InSingleQuadParams.SourceTextureResource);
+				}
 
-					// If we are using the same texture as the last one, we can render it in the same pass, just append our quad : 
-					if (InSingleQuadParams.SourceTextureResource == CurrentSourceTextureResources.Last())
-					{
-						const int32 SourceTextureIndex = CurrentSourceTextureResources.Num() - 1;
-						CurrentQuadInfos.Add(FUintVector2(static_cast<uint32>(SourceTextureIndex), InSingleQuadParams.ChannelSwizzleMask));
-						CurrentSourceRects.Add(FUintVector4(InSingleQuadParams.SourceRect.Min.X, InSingleQuadParams.SourceRect.Min.Y, InSingleQuadParams.SourceRect.Max.X, InSingleQuadParams.SourceRect.Max.Y));
-						CurrentDestinationRects.Add(FUintVector4(InSingleQuadParams.DestinationRect.Min.X, InSingleQuadParams.DestinationRect.Min.Y, InSingleQuadParams.DestinationRect.Max.X, InSingleQuadParams.DestinationRect.Max.Y));
-					}
+				// If we are using the same texture as the last one, we can render it in the same pass, just append our quad : 
+				if (InSingleQuadParams.SourceTextureResource == CurrentSourceTextureResources.Last())
+				{
+					const int32 SourceTextureIndex = CurrentSourceTextureResources.Num() - 1;
+					CurrentQuadInfos.Add(FUintVector2(static_cast<uint32>(SourceTextureIndex), InSingleQuadParams.SourceChannelIndex));
+					CurrentSourceRects.Add(FUintVector4(InSingleQuadParams.SourceRect.Min.X, InSingleQuadParams.SourceRect.Min.Y, InSingleQuadParams.SourceRect.Max.X, InSingleQuadParams.SourceRect.Max.Y));
+					CurrentDestinationRects.Add(FUintVector4(InSingleQuadParams.DestinationRect.Min.X, InSingleQuadParams.DestinationRect.Min.Y, InSingleQuadParams.DestinationRect.Max.X, InSingleQuadParams.DestinationRect.Max.Y));
 				}
+			}
 
-				// Flush the remaining quads if any : 
-				FlushCopyQuads();
-			};
+			// Flush the remaining quads if any : 
+			FlushCopyQuads();
+		};
 
-		// We need to specify the final state of the external texture to prevent the graph builder from transitioning it to SRVMask :
-		RDGBuilderRecorder.EnqueueRDGCommand(RDGCommand, { { WriteRT->GetRenderTarget()->GetResource(), ERHIAccess::RTV } });
-	}
+	// We need to specify the final state of the external texture to prevent the graph builder from transitioning it to SRVMask :
+	RDGBuilderRecorder.EnqueueRDGCommand(RDGCommand, { { WriteRT->GetRenderTarget()->GetResource(), ERHIAccess::RTV } });
 
-	// The second step is to perform the blend : merge the edit layer with the previous layers in the stack
-	{
-		RenderParams.MergeRenderContext->CycleBlendRenderTargets(RDGBuilderRecorder);
-		WriteRT = RenderParams.MergeRenderContext->GetBlendRenderTargetWrite();
-		ULandscapeScratchRenderTarget* CurrentLayerReadRT = RenderParams.MergeRenderContext->GetBlendRenderTargetRead();
-		ULandscapeScratchRenderTarget* PreviousLayersReadRT = RenderParams.MergeRenderContext->GetBlendRenderTargetReadPrevious();
+	// We've rendered at least a quad : 
+	return true;
+}
 
-		WriteRT->TransitionTo(ERHIAccess::RTV, RDGBuilderRecorder);
-		CurrentLayerReadRT->TransitionTo(ERHIAccess::SRVMask, RDGBuilderRecorder);
-		PreviousLayersReadRT->TransitionTo(ERHIAccess::SRVMask, RDGBuilderRecorder);
+void ULandscapeEditLayerPersistent::BlendLayer(FRenderParams& RenderParams, UE::Landscape::FRDGBuilderRecorder& RDGBuilderRecorder)
+{
+	using namespace UE::Landscape::EditLayers;
+	using namespace UE::Landscape::Private;
 
-		if (RenderParams.MergeRenderContext->IsHeightmapMerge())
-		{
-			auto RDGCommand =
-				[TargetLayerName = RenderParams.TargetLayerGroupLayerNames[0]
-				, OutputResource = WriteRT->GetRenderTarget2D()->GetResource()
-				, OutputResourceName = WriteRT->GetDebugName()
-				, CurrentEditLayerResource = CurrentLayerReadRT->GetRenderTarget2D()->GetResource()
-				, PreviousEditLayersResource = PreviousLayersReadRT->GetRenderTarget2D()->GetResource()
-				, EffectiveTextureSize = RenderParams.MergeRenderContext->GetCurrentRenderBatch()->GetRenderTargetResolution(/*bInWithDuplicateBorders = */false)
-				, BlendMode = LayerStruct->BlendMode, HeightmapAlpha = LayerStruct->HeightmapAlpha]
-				(FRDGBuilder& GraphBuilder)
-				{
-					FRDGTextureRef OutputTextureRef = GraphBuilder.RegisterExternalTexture(CreateRenderTarget(OutputResource->TextureRHI, TEXT("OutputTexture")));
-					FRDGTextureRef CurrentEditLayerTextureRef = GraphBuilder.RegisterExternalTexture(CreateRenderTarget(CurrentEditLayerResource->TextureRHI, TEXT("CurrentEditLayerTexture")));
-					FRDGTextureRef PreviousEditLayersTextureRef = GraphBuilder.RegisterExternalTexture(CreateRenderTarget(PreviousEditLayersResource->TextureRHI, TEXT("PreviousEditLayersTexture")));
-
-					FLandscapeEditLayersHeightmapsMergeEditLayerPS::FParameters* PSParams = GraphBuilder.AllocParameters<FLandscapeEditLayersHeightmapsMergeEditLayerPS::FParameters>();
-					PSParams->RenderTargets[0] = FRenderTargetBinding(OutputTextureRef, ERenderTargetLoadAction::ENoAction);
-					PSParams->InEditLayerBlendMode = static_cast<uint32>(LandscapeBlendModeToEditLayerBlendMode(BlendMode));
-					PSParams->InEditLayerAlpha = HeightmapAlpha;
-					PSParams->InCurrentEditLayerHeightmap = GraphBuilder.CreateSRV(FRDGTextureSRVDesc::Create(CurrentEditLayerTextureRef));
-					PSParams->InPreviousEditLayersHeightmap = GraphBuilder.CreateSRV(FRDGTextureSRVDesc::Create(PreviousEditLayersTextureRef));
-
-					FLandscapeEditLayersHeightmapsMergeEditLayerPS::MergeEditLayerPS(
-						RDG_EVENT_NAME("MergeEditLayer(%s) -> %s", *TargetLayerName.ToString(), *OutputResourceName), GraphBuilder, PSParams, EffectiveTextureSize);
-				};
+	ULandscapeInfo* Info = RenderParams.MergeRenderContext->GetLandscapeInfo();
+	check(Info != nullptr);
+	const FLandscapeLayer* LayerStruct = GetOwningLayer();
+	check(LayerStruct != nullptr);
 
+	const int32 NumTargetLayersInGroup = RenderParams.TargetLayerGroupLayerNames.Num();
 
-			// We need to specify the final state of the external textures to prevent the graph builder from transitioning them to SRVMask (even those that end up as SRVMask at the end of this command,
-			//  because they will likely be part of another RDGCommand down the line so we need to maintain an accurate picture of every external texture ever involved in the recorded command so that 
-			//  we can set a proper access when the recorder is flushed (and the FRDGBuilder, executed) :
-			TArray<FRDGBuilderRecorder::FRDGExternalTextureAccessFinal> RDGExternalTextureAccessFinalList =
-			{
-				{ WriteRT->GetRenderTarget()->GetResource(), ERHIAccess::RTV },
-				{ CurrentLayerReadRT->GetRenderTarget()->GetResource(), ERHIAccess::SRVMask },
-				{ PreviousLayersReadRT->GetRenderTarget()->GetResource(), ERHIAccess::SRVMask },
-			};
-			RDGBuilderRecorder.EnqueueRDGCommand(RDGCommand, RDGExternalTextureAccessFinalList);
-		}
-		else
+	// Prepare the generic blend params based on the layer's data : 
+	FBlendParams BlendParams;
+	if (RenderParams.MergeRenderContext->IsHeightmapMerge())
+	{
+		BlendParams.HeightmapBlendParams.BlendMode = LandscapeBlendModeToHeightmapBlendMode(LayerStruct->BlendMode);
+		// TODO [jonathan.bard] : Remove the splines type check when HeightmapAlpha is replaced by a virtual GetHeightmapAlpha() on ULandscapeEditLayer (alpha is fixed to 1.0 in this case, but 
+		//  HeightmapAlpha might still be set to a different value)
+		BlendParams.HeightmapBlendParams.Alpha = IsA<ULandscapeEditLayerSplines>() ? 1.0f : LayerStruct->HeightmapAlpha;
+	}
+	else
+	{
+		check(NumTargetLayersInGroup == RenderParams.TargetLayerGroupLayerInfos.Num());
+		BlendParams.WeightmapBlendParams.Reserve(NumTargetLayersInGroup);
+		for (int32 i = 0; i < NumTargetLayersInGroup; ++i)
 		{
-			struct FMergePaintLayerParams
-			{
-				FMergePaintLayerParams(const FName& InPaintLayerName)
-					: PaintLayerName(InPaintLayerName)
-				{}
+			const FName& TargetLayerName = RenderParams.TargetLayerGroupLayerNames[i];
+			const ULandscapeLayerInfoObject* LayerInfoObj = RenderParams.TargetLayerGroupLayerInfos[i];
+			check(LayerInfoObj != nullptr);
 
-				FName PaintLayerName;
-				EWeightmapBlendMode EditLayerPaintLayerBlendMode = EWeightmapBlendMode::Additive;
-				float EditLayerAlpha = 1.0f;
-			};
-
-			// Prepare a list of all paint layers to merge and do it all in 1 command :
-			TArray<FMergePaintLayerParams> MergePaintLayersParams;
-			MergePaintLayersParams.Reserve(NumTargetLayersInGroup);
-
-			check(NumTargetLayersInGroup == RenderParams.TargetLayerGroupLayerInfos.Num());
-			for (int32 i = 0; i < NumTargetLayersInGroup; ++i)
+			// only blend the layers involved in this step (the others are using EWeightmapBlendMode::Passthrough): 
+			if (RenderParams.TargetLayerGroupLayerNames.Contains(TargetLayerName))
 			{
-				const FName& TargetLayerName = RenderParams.TargetLayerGroupLayerNames[i];
-				const ULandscapeLayerInfoObject* LayerInfoObj = RenderParams.TargetLayerGroupLayerInfos[i];
-				check(LayerInfoObj != nullptr);
-
-				FMergePaintLayerParams& MergePaintLayerParams = MergePaintLayersParams.Emplace_GetRef(TargetLayerName);
+				FWeightmapBlendParams& TargetLayerBlendParams = BlendParams.WeightmapBlendParams.Emplace(TargetLayerName, EWeightmapBlendMode::Additive);
 
-				bool bIsSubtractiveBlend = false;
 				if (const bool* bSubstractiveInLayer = LayerStruct->WeightmapLayerAllocationBlend.Find(LayerInfoObj); (bSubstractiveInLayer != nullptr) && *bSubstractiveInLayer)
 				{
-					MergePaintLayerParams.EditLayerPaintLayerBlendMode = EWeightmapBlendMode::Subtractive;
+					TargetLayerBlendParams.BlendMode = EWeightmapBlendMode::Subtractive;
 				}
 
 				if (TargetLayerName != UMaterialExpressionLandscapeVisibilityMask::ParameterName)
 				{
-					MergePaintLayerParams.EditLayerAlpha = LayerStruct->WeightmapAlpha;
+					TargetLayerBlendParams.Alpha = LayerStruct->WeightmapAlpha;
 				}
 			}
-
-			auto RDGCommand =
-				[MergePaintLayersParams
-				, OutputResource = WriteRT->GetRenderTarget2DArray()->GetResource()
-				, OutputResourceName = WriteRT->GetDebugName()
-				, CurrentEditLayerResource = CurrentLayerReadRT->GetRenderTarget2DArray()->GetResource()
-				, PreviousEditLayersResource = PreviousLayersReadRT->GetRenderTarget2DArray()->GetResource()
-				, EffectiveTextureSize = RenderParams.MergeRenderContext->GetCurrentRenderBatch()->GetRenderTargetResolution(/*bInWithDuplicateBorders = */false)]
-				(FRDGBuilder& GraphBuilder)
-				{
-					FRDGTextureRef OutputTextureRef = GraphBuilder.RegisterExternalTexture(CreateRenderTarget(OutputResource->TextureRHI, TEXT("OutputTexture")));
-					FRDGTextureRef CurrentEditLayerTextureRef = GraphBuilder.RegisterExternalTexture(CreateRenderTarget(CurrentEditLayerResource->TextureRHI, TEXT("CurrentEditLayerTexture")));
-					FRDGTextureSRVRef CurrentEditLayerTextureSRVRef = GraphBuilder.CreateSRV(FRDGTextureSRVDesc::Create(CurrentEditLayerTextureRef));
-					FRDGTextureRef PreviousEditLayersTextureRef = GraphBuilder.RegisterExternalTexture(CreateRenderTarget(PreviousEditLayersResource->TextureRHI, TEXT("PreviousEditLayersTexture")));
-					FRDGTextureSRVRef PreviousEditLayersTextureSRVRef = GraphBuilder.CreateSRV(FRDGTextureSRVDesc::Create(PreviousEditLayersTextureRef));
-
-					const int32 NumTargetLayers = MergePaintLayersParams.Num();
-					for (int32 TargetLayerIndex = 0; TargetLayerIndex < NumTargetLayers; ++TargetLayerIndex)
-					{
-						FLandscapeEditLayersWeightmapsMergeEditLayerPS::FParameters* PSParams = GraphBuilder.AllocParameters<FLandscapeEditLayersWeightmapsMergeEditLayerPS::FParameters>();
-						PSParams->RenderTargets[0] = FRenderTargetBinding(OutputTextureRef, ERenderTargetLoadAction::ENoAction, /*InMipIndex = */0, /*InArraySlice = */TargetLayerIndex);
-						PSParams->InPaintLayerIndex = TargetLayerIndex;
-						PSParams->InEditLayerPaintLayerBlendMode = static_cast<uint32>(MergePaintLayersParams[TargetLayerIndex].EditLayerPaintLayerBlendMode);
-						PSParams->InEditLayerAlpha = MergePaintLayersParams[TargetLayerIndex].EditLayerAlpha;
-						PSParams->InCurrentEditLayerWeightmaps = CurrentEditLayerTextureSRVRef;
-						PSParams->InPreviousEditLayersWeightmaps = PreviousEditLayersTextureSRVRef;
-
-						FLandscapeEditLayersWeightmapsMergeEditLayerPS::MergeEditLayerPS(
-							RDG_EVENT_NAME("MergeEditLayer(%s) -> %s", *MergePaintLayersParams[TargetLayerIndex].PaintLayerName.ToString(), *OutputResourceName), GraphBuilder, PSParams, EffectiveTextureSize);
-					}
-				};
-
-			// We need to specify the final state of the external textures to prevent the graph builder from transitioning them to SRVMask (even those that end up as SRVMask at the end of this command,
-			//  because they will likely be part of another RDGCommand down the line so we need to maintain an accurate picture of every external texture ever involved in the recorded command so that 
-			//  we can set a proper access when the recorder is flushed (and the FRDGBuilder, executed) :
-			TArray<FRDGBuilderRecorder::FRDGExternalTextureAccessFinal> RDGExternalTextureAccessFinalList =
-			{
-				{ WriteRT->GetRenderTarget()->GetResource(), ERHIAccess::RTV },
-				{ CurrentLayerReadRT->GetRenderTarget()->GetResource(), ERHIAccess::SRVMask },
-				{ PreviousLayersReadRT->GetRenderTarget()->GetResource(), ERHIAccess::SRVMask },
-			};
-			RDGBuilderRecorder.EnqueueRDGCommand(RDGCommand, RDGExternalTextureAccessFinalList);
 		}
 	}
+
+	// Then perform the generic blend : 
+	RenderParams.MergeRenderContext->GenericBlendLayer(BlendParams, RDGBuilderRecorder);
 }
 
 FString ULandscapeEditLayerPersistent::GetEditLayerRendererDebugName() const
@@ -6078,6 +5966,25 @@ namespace UE::Landscape::EditLayers::Private
 		FIntPoint MaxComponentKey = FIntPoint(MIN_int32, MIN_int32);
 	};
 
+	/** Struct that holds all the per-render item information on a given renderer, needed when preparing the batched merge context */
+	struct FEditLayerRendererRenderItemRenderInfo
+	{
+		FEditLayerRendererRenderItemRenderInfo() = delete;
+		FEditLayerRendererRenderItemRenderInfo(const FEditLayerRenderItem& InRenderItem, const TBitArray<>& InOutputLayerBitIndices, int32 InNumComponentsToRender)
+			: RenderItem(InRenderItem)
+			, OutputLayerBitIndices(InOutputLayerBitIndices)
+			, RenderedComponentBitIndices(false, InNumComponentsToRender)
+		{}
+
+		FEditLayerRenderItem RenderItem;
+
+		// Indices of the target layers affected by this render item on this renderer : 
+		TBitArray<> OutputLayerBitIndices; 
+
+		// Indices (in AllComponentsToRenderInfos) of the components which which are involved (as inputs or outputs) with this render item 
+		TBitArray<> RenderedComponentBitIndices;
+	};
+
 	/** Struct that holds all the per-renderer information needed when preparing the batched merge context */
 	struct FEditLayerRendererRenderInfo
 	{
@@ -6101,6 +6008,8 @@ namespace UE::Landscape::EditLayers::Private
 		FEditLayerRendererState RendererState;
 		int32 RendererIndex = INDEX_NONE;
 
+		TArray<FEditLayerRendererRenderItemRenderInfo> RenderItemRenderInfos;
+
 		// Indices (in AllComponentsToRenderInfos) of the components which we need to render for this renderer
 		TBitArray<> ComponentToRenderInfoBitIndices;
 
@@ -6110,6 +6019,16 @@ namespace UE::Landscape::EditLayers::Private
 #endif // ENABLE_VISUAL_LOG
 	};
 
+	/** Struct that holds all the per-renderer information for a given batch, needed when preparing the batched merge context */
+	struct FPerBatchEditLayerRendererRenderInfo
+	{
+		// Indices (in AllComponentsToRenderInfos) of the components which we need to render for this renderer in this batch
+		TBitArray<> ComponentsToRenderBitIndices;
+
+		// Components which we need to render for this renderer in this batch
+		TArray<ULandscapeComponent*> ComponentsToRender;
+	};
+
 	/** Given the list of components that actually need to be rendered, divide the work into batches, such that:
 	 *   - When a component is being rendered by a given batch, all components needed for rendering this given component are present in the batch
 	 *   - All components end up being rendered in at least one of the batches
@@ -6431,7 +6350,7 @@ namespace UE::Landscape::EditLayers::Private
 #endif // ENABLE_VISUAL_LOG
 } // UE::Landscape::EditLayers::Private
 
-void ULandscapeHeightmapNormalsEditLayerRenderer::RenderLayer(FRenderParams& RenderParams, UE::Landscape::FRDGBuilderRecorder& RDGBuilderRecorder)
+bool ULandscapeHeightmapNormalsEditLayerRenderer::RenderLayer(FRenderParams& RenderParams, UE::Landscape::FRDGBuilderRecorder& RDGBuilderRecorder)
 {
 	using namespace UE::Landscape;
 	using namespace UE::Landscape::EditLayers;
@@ -6494,9 +6413,11 @@ void ULandscapeHeightmapNormalsEditLayerRenderer::RenderLayer(FRenderParams& Ren
 		{ ValidityRT->GetRenderTarget()->GetResource(), ERHIAccess::SRVMask },
 	};
 	RDGBuilderRecorder.EnqueueRDGCommand(RDGCommand, RDGExternalTextureAccessFinalList);
+
+	return true;
 }
 
-void ULandscapeWeightmapWeightBlendedLayersRenderer::RenderLayer(FRenderParams& RenderParams, UE::Landscape::FRDGBuilderRecorder& RDGBuilderRecorder)
+bool ULandscapeWeightmapWeightBlendedLayersRenderer::RenderLayer(FRenderParams& RenderParams, UE::Landscape::FRDGBuilderRecorder& RDGBuilderRecorder)
 {
 	using namespace UE::Landscape;
 	using namespace UE::Landscape::EditLayers;
@@ -6512,26 +6433,26 @@ void ULandscapeWeightmapWeightBlendedLayersRenderer::RenderLayer(FRenderParams&
 	WriteRT->TransitionTo(ERHIAccess::RTV, RDGBuilderRecorder);
 	ReadRT->TransitionTo(ERHIAccess::SRVMask, RDGBuilderRecorder);
 
-	TArray<FWeightmapPaintLayerInfo> WeightmapPaintLayerInfos;
-	WeightmapPaintLayerInfos.Reserve(RenderParams.TargetLayerGroupLayerInfos.Num());
-	Algo::Transform(RenderParams.TargetLayerGroupLayerInfos, WeightmapPaintLayerInfos, [](ULandscapeLayerInfoObject* InLayerInfo)
+	TArray<FWeightmapTargetLayerInfo> WeightmapTargetLayerInfos;
+	WeightmapTargetLayerInfos.Reserve(RenderParams.TargetLayerGroupLayerInfos.Num());
+	Algo::Transform(RenderParams.TargetLayerGroupLayerInfos, WeightmapTargetLayerInfos, [](ULandscapeLayerInfoObject* InLayerInfo)
 		{
 			check(InLayerInfo != nullptr); // There should only be valid layer infos at this point
-			FWeightmapPaintLayerInfo WeightmapPaintLayerInfo;
+			FWeightmapTargetLayerInfo WeightmapTargetLayerInfo;
 			if (!InLayerInfo->bNoWeightBlend)
 			{
-				WeightmapPaintLayerInfo.Flags = EWeightmapPaintLayerFlags::IsWeightBlended;
+				WeightmapTargetLayerInfo.Flags = EWeightmapTargetLayerFlags::IsWeightBlended;
 			}
 			if (InLayerInfo == ALandscapeProxy::VisibilityLayer)
 			{
-				WeightmapPaintLayerInfo.Flags = EWeightmapPaintLayerFlags::IsVisibilityLayer;
+				WeightmapTargetLayerInfo.Flags = EWeightmapTargetLayerFlags::IsVisibilityLayer;
 			}
-			return WeightmapPaintLayerInfo;
+			return WeightmapTargetLayerInfo;
 		});
 
 	auto RDGCommand =
 		[TargetLayerNames = RenderParams.TargetLayerGroupLayerNames
-		, WeightmapPaintLayerInfos
+		, WeightmapTargetLayerInfos
 		, OutputResource = WriteRT->GetRenderTarget2DArray()->GetResource()
 		, OutputResourceName = WriteRT->GetDebugName()
 		, CurrentEditLayerResource = ReadRT->GetRenderTarget2DArray()->GetResource()
@@ -6541,17 +6462,17 @@ void ULandscapeWeightmapWeightBlendedLayersRenderer::RenderLayer(FRenderParams&
 			FRDGTextureRef CurrentEditLayerTextureRef = GraphBuilder.RegisterExternalTexture(CreateRenderTarget(CurrentEditLayerResource->TextureRHI, TEXT("CurrentEditLayerTexture")));
 			FRDGTextureSRVRef CurrentEditLayerTextureSRVRef = GraphBuilder.CreateSRV(FRDGTextureSRVDesc::Create(CurrentEditLayerTextureRef));
 
-			FRDGBufferRef PaintLayerInfosBuffer = CreateStructuredBuffer(GraphBuilder, TEXT("LandscapePaintLayerInfosBuffer"), WeightmapPaintLayerInfos);
-			FRDGBufferSRVRef PaintLayerInfosBufferSRV = GraphBuilder.CreateSRV(FRDGBufferSRVDesc(PaintLayerInfosBuffer));
+			FRDGBufferRef TargetLayerInfosBuffer = CreateStructuredBuffer(GraphBuilder, TEXT("LandscapeTargetLayerInfosBuffer"), WeightmapTargetLayerInfos);
+			FRDGBufferSRVRef TargetLayerInfosBufferSRV = GraphBuilder.CreateSRV(FRDGBufferSRVDesc(TargetLayerInfosBuffer));
 
 			const int32 NumTargetLayers = TargetLayerNames.Num();
 			for (int32 TargetLayerIndex = 0; TargetLayerIndex < NumTargetLayers; ++TargetLayerIndex)
 			{
 				FLandscapeEditLayersWeightmapsPerformLegacyWeightBlendingPS::FParameters* PSParams = GraphBuilder.AllocParameters<FLandscapeEditLayersWeightmapsPerformLegacyWeightBlendingPS::FParameters>();
 				PSParams->RenderTargets[0] = FRenderTargetBinding(OutputTextureRef, ERenderTargetLoadAction::ENoAction, /*InMipIndex = */0, /*InArraySlice = */TargetLayerIndex);
-				PSParams->InPaintLayerIndex = TargetLayerIndex;
-				PSParams->InNumPaintLayers = NumTargetLayers;
-				PSParams->InPaintLayerInfos = PaintLayerInfosBufferSRV;
+				PSParams->InTargetLayerIndex = TargetLayerIndex;
+				PSParams->InNumTargetLayers = NumTargetLayers;
+				PSParams->InTargetLayerInfos = TargetLayerInfosBufferSRV;
 				PSParams->InCurrentEditLayerWeightmaps = CurrentEditLayerTextureSRVRef;
 
 				FLandscapeEditLayersWeightmapsPerformLegacyWeightBlendingPS::PerformLegacyWeightBlendingPS(
@@ -6569,6 +6490,8 @@ void ULandscapeWeightmapWeightBlendedLayersRenderer::RenderLayer(FRenderParams&
 		{ ReadRT->GetRenderTarget()->GetResource(), ERHIAccess::SRVMask },
 	};
 	RDGBuilderRecorder.EnqueueRDGCommand(RDGCommand, RDGExternalTextureAccessFinalList);
+
+	return true;
 }
 
 UE::Landscape::EditLayers::FMergeRenderContext ALandscape::PrepareEditLayersMergeRenderContext(const UE::Landscape::EditLayers::FMergeContext& InMergeContext, const UE::Landscape::EditLayers::FMergeRenderParams& InMergeRenderParams)
@@ -6586,6 +6509,7 @@ UE::Landscape::EditLayers::FMergeRenderContext ALandscape::PrepareEditLayersMerg
 	const FString VisualLogShowRenderItemsEditLayerRendererFilter = CVarLandscapeBatchedMergeVisualLogShowRenderItemsEditLayerRendererFilter.GetValueOnGameThread();
 	const int32 VisualLogShowComponentDependencies = CVarLandscapeBatchedMergeVisualLogShowComponentDependencies.GetValueOnGameThread();
 	const FString VisualLogShowComponentDependenciesFilter = CVarLandscapeBatchedMergeVisualLogShowComponentDependenciesFilter.GetValueOnGameThread();
+	const bool bEnableRenderLayerGrouping = CVarLandscapeBatchedMergeEnableRenderLayerGroup.GetValueOnGameThread();
 
 	ULandscapeInfo* Info = GetLandscapeInfo();
 	check(Info != nullptr);
@@ -6796,7 +6720,6 @@ UE::Landscape::EditLayers::FMergeRenderContext ALandscape::PrepareEditLayersMerg
 	const TArray<ULandscapeComponent*> AllComponents = Component2DIndexer.GetAllValues();
 	const TBitArray<> AllValidComponentBitIndices = Component2DIndexer.GetValidValueBitIndices();
 	const int32 NumAllComponents = AllComponents.Num();
-	const FEditLayerRendererState DummyRendererState(&MergeRenderContext);
 
 	// The list of all components that will end up being rendered across all renderers (one bit per component) :
 	TBitArray<> FinalComponentsToRenderInfoBitIndices(false, NumAllComponents);
@@ -7036,6 +6959,9 @@ UE::Landscape::EditLayers::FMergeRenderContext ALandscape::PrepareEditLayersMerg
 						}
 					}
 
+					// Inform the renderer about this how this render item affects it :
+					FEditLayerRendererRenderItemRenderInfo& RenderItemRenderInfo = EditLayerRendererRenderInfo.RenderItemRenderInfos.Emplace_GetRef(RenderItem, OutputLayerBitIndices, NumAllComponents);
+
 					// Iterate on all the required components :
 					for (TConstSetBitIterator It(*ComponentsToIterateBitIndices); It; ++It)
 					{
@@ -7058,6 +6984,9 @@ UE::Landscape::EditLayers::FMergeRenderContext ALandscape::PrepareEditLayersMerg
 						// There should always be a dependency between this renderer and the previous in the stack for the component itself : 
 						InputComponentBitIndices[ComponentToRenderIndex] = true;
 
+						// Tell the render item which component it needs : 
+						RenderItemRenderInfo.RenderedComponentBitIndices.CombineWithBitwiseOR(InputComponentBitIndices, EBitwiseOperatorFlags::MinSize);
+
 						// If these components are not yet in the final list of components to render, add them: 
 						ComponentsToRenderInfoBitIndicesAfterRenderer.CombineWithBitwiseOR(InputComponentBitIndices, EBitwiseOperatorFlags::MinSize);
 
@@ -7142,22 +7071,23 @@ UE::Landscape::EditLayers::FMergeRenderContext ALandscape::PrepareEditLayersMerg
 
 			MergeRenderContext.MaxNeededResolution = MergeRenderContext.MaxNeededResolution.ComponentMax(MergeRenderBatch.Resolution);
 		
+
 			// For each renderer, find the list of components actually needed and store that in a separate array, to avoid recomputing it for each target group : 
 			int32 LocalNumRenderers = OrderedEditLayerRendererRenderInfos.Num();
-			TArray<TArray<ULandscapeComponent*>> ComponentsForRenderers;
-			ComponentsForRenderers.SetNum(LocalNumRenderers);
+			TArray<FPerBatchEditLayerRendererRenderInfo> BatchEditLayerRendererRenderInfos;
+			BatchEditLayerRendererRenderInfos.SetNum(LocalNumRenderers);
 			for (int32 EditLayerRendererIndex = 0; EditLayerRendererIndex < LocalNumRenderers; ++EditLayerRendererIndex)
 			{
 				const FEditLayerRendererRenderInfo& EditLayerRendererRenderInfo = OrderedEditLayerRendererRenderInfos[EditLayerRendererIndex];
 				// Find out all components that are in common between the renderer's list and the batch's list : only these will need to be rendered in that render step : 
-				TBitArray<> ComponentToRenderForRendererBitIndices = TBitArray<>::BitwiseAND(EditLayerRendererRenderInfo.ComponentToRenderInfoBitIndices, BatchInfo.ComponentToRenderInfoBitIndices, EBitwiseOperatorFlags::MinSize);
-				TArray<ULandscapeComponent*>& ComponentsToRender = ComponentsForRenderers[EditLayerRendererIndex];
-				ComponentsToRender.Reserve(ComponentToRenderForRendererBitIndices.CountSetBits());
+				FPerBatchEditLayerRendererRenderInfo& PerBatchEditLayerRendererRenderInfo = BatchEditLayerRendererRenderInfos[EditLayerRendererIndex];
+				PerBatchEditLayerRendererRenderInfo.ComponentsToRenderBitIndices = TBitArray<>::BitwiseAND(EditLayerRendererRenderInfo.ComponentToRenderInfoBitIndices, BatchInfo.ComponentToRenderInfoBitIndices, EBitwiseOperatorFlags::MinSize);
+				PerBatchEditLayerRendererRenderInfo.ComponentsToRender.Reserve(PerBatchEditLayerRendererRenderInfo.ComponentsToRenderBitIndices.CountSetBits());
 				// Transform the bit indices back into a proper component list : 
-				for (TConstSetBitIterator ItComponent(ComponentToRenderForRendererBitIndices); ItComponent; ++ItComponent)
+				for (TConstSetBitIterator ItComponent(PerBatchEditLayerRendererRenderInfo.ComponentsToRenderBitIndices); ItComponent; ++ItComponent)
 				{
 					FComponentToRenderInfo& ComponentToRenderInfo = AllComponentsToRenderInfos[ItComponent.GetIndex()];
-					ComponentsToRender.Add(ComponentToRenderInfo.Component);
+					PerBatchEditLayerRendererRenderInfo.ComponentsToRender.Add(ComponentToRenderInfo.Component);
 
 					// Inform the render batch and context of the target layer names associated with each component :
 
@@ -7174,55 +7104,194 @@ UE::Landscape::EditLayers::FMergeRenderContext ALandscape::PrepareEditLayersMerg
 				}
 
 				// Declare these components as being in use for the batch : 
-				MergeRenderBatch.ComponentsToRender.Append(ComponentsToRender);
+				MergeRenderBatch.ComponentsToRender.Append(PerBatchEditLayerRendererRenderInfo.ComponentsToRender);
 			}
 
 			// Now, we have all the info to build our list of successive render steps : process group by group :
 			for (const TBitArray<>& TargetLayerGroup : FinalTargetLayerGroups)
 			{
-				ERenderFlags LastRenderMode = ERenderFlags::None;
+				TRACE_CPUPROFILER_EVENT_SCOPE(BuildRenderSteps);
+				bool bInRecordedSequence = false;
+				int32 RenderLayerGroupStartIndex = INDEX_NONE;
+				int32 RenderLayerGroupEndIndex = INDEX_NONE;
 
 				// Execute the required operations for the whole stack of renderers for this group :
 				for (int32 RendererIndex = 0; RendererIndex < LocalNumRenderers; ++RendererIndex)
 				{
 					const FEditLayerRendererRenderInfo& EditLayerRendererRenderInfo = OrderedEditLayerRendererRenderInfos[RendererIndex];
-					TArrayView<ULandscapeComponent*> ComponentsToRender = MakeArrayView(ComponentsForRenderers[RendererIndex]);
+					TArrayView<ULandscapeComponent*> ComponentsToRender = MakeArrayView(BatchEditLayerRendererRenderInfos[RendererIndex].ComponentsToRender);
+					const ERenderFlags CurrentRenderFlags = EditLayerRendererRenderInfo.RendererState.GetRenderer()->GetRenderFlags(&MergeRenderContext);
+					const ERenderFlags CurrentRenderMode = CurrentRenderFlags & ERenderFlags::RenderMode_Mask;
+					checkf(FMath::CountBits(static_cast<uint8>(CurrentRenderMode)) == 1, TEXT("Render mode should be either Immediate or Recorded"));
+					const bool bIsLastRenderer = (RendererIndex == (LocalNumRenderers - 1));
+					const bool bHasSeparateBlend = EnumHasAnyFlags(CurrentRenderFlags, ERenderFlags::BlendMode_SeparateBlend);
+					const bool bSupportsGrouping = EnumHasAnyFlags(CurrentRenderFlags, ERenderFlags::RenderLayerGroup_SupportsGrouping);
+					checkf(!bSupportsGrouping || bHasSeparateBlend, TEXT("RenderLayerGroup_SupportsGrouping expects BlendMode_SeparateBlend"));
+
+					// TODO [jonathan.bard] : take into account target layer bit indices : only take into account render items that output the same layers as this target layer group
+					auto IsCompatibleWithRenderLayerGroup = [&]() -> bool
+						{
+							// Global switch to disable grouping altogether for debug purposes : 
+							// Note : in order not to force the renderers to implement the grouping AND non-grouping behavior, we simply force render groups to contain a single element :
+							if (!bEnableRenderLayerGrouping)
+							{
+								return false;
+							}
+
+							check((RenderLayerGroupStartIndex != INDEX_NONE) && (RendererIndex > RenderLayerGroupStartIndex));
+
+							const FPerBatchEditLayerRendererRenderInfo& PerBatchEditLayerRendererRenderInfo = BatchEditLayerRendererRenderInfos[RendererIndex];
+							for (int32 OtherRendererIndex = RenderLayerGroupStartIndex; OtherRendererIndex < RendererIndex; ++OtherRendererIndex)
+							{
+								const FEditLayerRendererRenderInfo& OtherEditLayerRendererRenderInfo = OrderedEditLayerRendererRenderInfos[OtherRendererIndex];
+								const FPerBatchEditLayerRendererRenderInfo& OtherPerBatchEditLayerRendererRenderInfo = BatchEditLayerRendererRenderInfos[OtherRendererIndex];
+
+								// The renderer itself can tell whether it is compatible with others in that group : 
+								if (!EditLayerRendererRenderInfo.RendererState.GetRenderer()->CanGroupRenderLayerWith(OtherEditLayerRendererRenderInfo.RendererState.GetRenderer()))
+								{
+									return false;
+								}
+								// The renderer is compatible with others in the layer group if none of its render items overlaps with any of the their respective render items (a blend is necessary in that case) : 
+								// First, perform a first cheap test with the components overlapped by each renderer in this batch. If none match, no need to look further, the renderers are compatible: 
+								else if (TBitArray<>::BitwiseAND(PerBatchEditLayerRendererRenderInfo.ComponentsToRenderBitIndices, OtherPerBatchEditLayerRendererRenderInfo.ComponentsToRenderBitIndices, EBitwiseOperatorFlags::MinSize).Contains(true))
+								{
+									// We have a potential overlap, let's test render item per render item now : 
+									for (const FEditLayerRendererRenderItemRenderInfo& RenderItemRenderInfo : EditLayerRendererRenderInfo.RenderItemRenderInfos)
+									{
+										const TBitArray<> RenderItemRenderedComponentsThisBatch = TBitArray<>::BitwiseAND(RenderItemRenderInfo.RenderedComponentBitIndices, PerBatchEditLayerRendererRenderInfo.ComponentsToRenderBitIndices, EBitwiseOperatorFlags::MinSize);
+										for (const FEditLayerRendererRenderItemRenderInfo& OtherRenderItemRenderInfo : OtherEditLayerRendererRenderInfo.RenderItemRenderInfos)
+										{
+											const TBitArray<> OtherRenderItemRenderedComponentsThisBatch = TBitArray<>::BitwiseAND(OtherRenderItemRenderInfo.RenderedComponentBitIndices, PerBatchEditLayerRendererRenderInfo.ComponentsToRenderBitIndices, EBitwiseOperatorFlags::MinSize);
+
+											// Another cheap test, this time with the render item's rendered components : 
+											if (TBitArray<>::BitwiseAND(RenderItemRenderedComponentsThisBatch, OtherRenderItemRenderedComponentsThisBatch, EBitwiseOperatorFlags::MinSize).Contains(true))
+											{
+												// If both render items are OOBBoxes, we can even have a finer-grained overlap test : 
+												// TODO [jonathan.bard] : perform a OOBB to OOBB test here to reduce the potential overlaps :
+												const FOOBox2D* RenderItemInputOOBox = RenderItemRenderInfo.RenderItem.GetInputWorldArea().TryGetOOBox();
+												const FOOBox2D* RenderItemOutputOOBox = RenderItemRenderInfo.RenderItem.GetOutputWorldArea().TryGetOOBox();
+												const FOOBox2D* OtherRenderItemInputOOBox = OtherRenderItemRenderInfo.RenderItem.GetInputWorldArea().TryGetOOBox();
+												const FOOBox2D* OtherRenderItemOutputOOBox = OtherRenderItemRenderInfo.RenderItem.GetOutputWorldArea().TryGetOOBox();
+												if ((RenderItemInputOOBox != nullptr) && (RenderItemOutputOOBox != nullptr)
+													&& (OtherRenderItemInputOOBox != nullptr) && (OtherRenderItemOutputOOBox != nullptr))
+												{
+													FBox RenderItemAABB = RenderItemInputOOBox->BuildAABB() + RenderItemOutputOOBox->BuildAABB();
+													FBox OtherRenderItemAABB = OtherRenderItemInputOOBox->BuildAABB() + OtherRenderItemOutputOOBox->BuildAABB();
+													// If the 2 boxes overlap, the render item cannot be part of the group :
+													if (RenderItemAABB.IntersectXY(OtherRenderItemAABB))
+													{
+														return false;
+													}
+												}
+												else
+												{
+													// We cannot perform a more precise test, so consider it's an overlap, since we know that at least, the components rendered by these 2 render items overlap : 
+													return false;
+												}
+											}
+										}
+									}
+								}
+							}
+							return true;
+						};
+
+					auto AppendBeginRenderLayerGroupStep = [&]()
+						{
+							check((RenderLayerGroupStartIndex == INDEX_NONE) && (RenderLayerGroupEndIndex == INDEX_NONE));
+							MergeRenderBatch.RenderSteps.Emplace(FMergeRenderStep::EType::BeginRenderLayerGroup, CurrentRenderFlags, EditLayerRendererRenderInfo.RendererState, TargetLayerGroup, ComponentsToRender);
+							RenderLayerGroupStartIndex = RendererIndex;
+						};
+
+					auto AppendEndRenderLayerGroupAndBlendLayerSteps = [&]()
+						{
+							check((RenderLayerGroupStartIndex != INDEX_NONE) && (RenderLayerGroupEndIndex != INDEX_NONE));
+							const FEditLayerRendererRenderInfo& LastEditLayerRendererRenderInfo = OrderedEditLayerRendererRenderInfos[RenderLayerGroupEndIndex];
+							const ERenderFlags LastRenderFlags = LastEditLayerRendererRenderInfo.RendererState.GetRenderer()->GetRenderFlags(&MergeRenderContext);
+							TArrayView<ULandscapeComponent*> LastComponentsToRender = MakeArrayView(BatchEditLayerRendererRenderInfos[RenderLayerGroupEndIndex].ComponentsToRender);
+							MergeRenderBatch.RenderSteps.Emplace(FMergeRenderStep::EType::EndRenderLayerGroup, LastRenderFlags, LastEditLayerRendererRenderInfo.RendererState, TargetLayerGroup, LastComponentsToRender);
+							MergeRenderBatch.RenderSteps.Emplace(FMergeRenderStep::EType::BlendLayer, LastRenderFlags, LastEditLayerRendererRenderInfo.RendererState, TargetLayerGroup, LastComponentsToRender);
+							RenderLayerGroupStartIndex = INDEX_NONE;
+							RenderLayerGroupEndIndex = INDEX_NONE;
+						};
 
 					if (!ComponentsToRender.IsEmpty())
 					{
-						const ERenderFlags CurrentRenderFlags = EditLayerRendererRenderInfo.RendererState.GetRenderer()->GetRenderFlags(&MergeRenderContext);
-						const ERenderFlags CurrentRenderMode = CurrentRenderFlags & ERenderFlags::RenderMode_Mask;
-						checkf(FMath::CountBits(static_cast<uint8>(CurrentRenderMode)) == 1, TEXT("Render mode should be either Immediate or Recorded"));
-
-						// Initiate or terminate the "render layer on render thread" sequence if necessary :
+						// Initiate the "render command recorder" sequence if necessary :
 						if (CurrentRenderMode == ERenderFlags::RenderMode_Recorded)
 						{
-							if (LastRenderMode != ERenderFlags::RenderMode_Recorded)
+							if (!bInRecordedSequence)
 							{
-								MergeRenderBatch.RenderSteps.Emplace(FMergeRenderStep::EType::BeginRenderCommandRecorder, DummyRendererState, TargetLayerGroup, MergeRenderBatch.ComponentsToRender.Array());
+								MergeRenderBatch.RenderSteps.Emplace(FMergeRenderStep::EType::BeginRenderCommandRecorder);
+								bInRecordedSequence = true;
 							}
 						}
-						else if (LastRenderMode == ERenderFlags::RenderMode_Recorded)
+						// Terminate the "render command recorder" sequence if necessary :
+						else if (bInRecordedSequence)
+						{
+							MergeRenderBatch.RenderSteps.Emplace(FMergeRenderStep::EType::EndRenderCommandRecorder);
+							bInRecordedSequence = false;
+						}
+
+						if (bSupportsGrouping)
+						{
+							// Start a new render layer group if this renderer supports it and none is currently being built
+							if (RenderLayerGroupStartIndex == INDEX_NONE)
+							{
+								AppendBeginRenderLayerGroupStep();
+							}
+							// If a render layer group is currently being built but the renderer cannot be added to it, we need to stop the group, perform the blend and start a new group :
+							else if (!IsCompatibleWithRenderLayerGroup())
+							{
+								AppendEndRenderLayerGroupAndBlendLayerSteps();
+								AppendBeginRenderLayerGroupStep();
+							}
+						}
+						// Stop the current render layer group if this renderer doesn't support it and one is currently being built
+						else if (RenderLayerGroupStartIndex != INDEX_NONE)
+						{
+							AppendEndRenderLayerGroupAndBlendLayerSteps();
+						}
+
+						// Render the content of this layer : 
+						MergeRenderBatch.RenderSteps.Emplace(FMergeRenderStep::EType::RenderLayer, CurrentRenderFlags, EditLayerRendererRenderInfo.RendererState, TargetLayerGroup, ComponentsToRender);
+						if (bSupportsGrouping)
+						{
+							RenderLayerGroupEndIndex = RendererIndex;
+						}
+						else
 						{
-							MergeRenderBatch.RenderSteps.Emplace(FMergeRenderStep::EType::EndRenderCommandRecorder, DummyRendererState, TargetLayerGroup, MergeRenderBatch.ComponentsToRender.Array());
+							check((RenderLayerGroupStartIndex == INDEX_NONE) && (RenderLayerGroupEndIndex == INDEX_NONE));
 						}
-						LastRenderMode = CurrentRenderMode;
 
-						// Render the content of this render layer : 
-						FMergeRenderStep::EType RenderStepType = (CurrentRenderMode == ERenderFlags::RenderMode_Immediate)
-							? FMergeRenderStep::EType::RenderLayer_Immediate : FMergeRenderStep::EType::RenderLayer_Recorded;
-						MergeRenderBatch.RenderSteps.Emplace(RenderStepType, EditLayerRendererRenderInfo.RendererState, TargetLayerGroup, ComponentsToRender);
+						// Add the blend step of this layer if it's separate and not part of an on-going group :
+						if (bHasSeparateBlend && (RenderLayerGroupStartIndex == INDEX_NONE))
+						{
+							MergeRenderBatch.RenderSteps.Emplace(FMergeRenderStep::EType::BlendLayer, CurrentRenderFlags, EditLayerRendererRenderInfo.RendererState, TargetLayerGroup, ComponentsToRender);
+						}
 					}
-				}
 
-				// Insert a final render sequence end step if we ended on a render on the render thread :
-				if (LastRenderMode == ERenderFlags::RenderMode_Recorded)
-				{
-					MergeRenderBatch.RenderSteps.Emplace(FMergeRenderStep::EType::EndRenderCommandRecorder, DummyRendererState, TargetLayerGroup, MergeRenderBatch.ComponentsToRender.Array());
+					if (bIsLastRenderer)
+					{
+						// Terminate the current render layer group if necessary :
+						if (RenderLayerGroupStartIndex != INDEX_NONE)
+						{
+							AppendEndRenderLayerGroupAndBlendLayerSteps();
+						}
+
+						// Terminate the current "render command recorder" sequence if necessary :
+						if (bInRecordedSequence)
+						{
+							MergeRenderBatch.RenderSteps.Emplace(FMergeRenderStep::EType::EndRenderCommandRecorder);
+							bInRecordedSequence = false;
+						}
+					}
 				}
+				check(!bInRecordedSequence);
+				check((RenderLayerGroupStartIndex == INDEX_NONE) && (RenderLayerGroupEndIndex == INDEX_NONE));
 
 				// Finally, signal the group as done for this batch
-				MergeRenderBatch.RenderSteps.Emplace(FMergeRenderStep::EType::SignalBatchMergeGroupDone, DummyRendererState, TargetLayerGroup, MergeRenderBatch.ComponentsToRender.Array());
+				MergeRenderBatch.RenderSteps.Emplace(FMergeRenderStep::EType::SignalBatchMergeGroupDone, TargetLayerGroup, MergeRenderBatch.ComponentsToRender.Array());
 
 				MergeRenderContext.MaxNeededNumSlices = FMath::Max(MergeRenderContext.MaxNeededNumSlices, TargetLayerGroup.CountSetBits());
 			}
@@ -7540,14 +7609,15 @@ int32 ALandscape::PerformLayersHeightmapsBatchedMerge(const FUpdateLayersContent
 		(const FMergeRenderContext::FOnRenderBatchTargetGroupDoneParams& InParams, UE::Landscape::FRDGBuilderRecorder& RDGBuilderRecorder)
 	{
 		ALandscape* Landscape = InParams.MergeRenderContext->GetLandscape();
+		const FMergeRenderBatch* RenderBatch = InParams.MergeRenderContext->GetCurrentRenderBatch();
 
 		//  Note: thanks to HeightmapNormalsRenderer, we have the guarantee that the (up to) 8 neighbors of each of the components requested 
 		//  for are present in the batch, which means we have all the data to generate the normals already 
 		TSet<ULandscapeComponent*> ComponentsToResolveThisBatch;
 		{
 			TRACE_CPUPROFILER_EVENT_SCOPE(PrepareResolve);
-			ComponentsToResolveThisBatch.Reserve(InParams.RenderBatch->ComponentsToRender.Num());
-			for (ULandscapeComponent* RenderedComponent : InParams.RenderBatch->ComponentsToRender)
+			ComponentsToResolveThisBatch.Reserve(RenderBatch->ComponentsToRender.Num());
+			for (ULandscapeComponent* RenderedComponent : RenderBatch->ComponentsToRender)
 			{
 				checkf(InParams.SortedComponentMergeRenderInfos.FindByPredicate([RenderedComponent](const FComponentMergeRenderInfo& InComponentMergeInfo) { return InComponentMergeInfo.Component == RenderedComponent; }) != nullptr,
 					TEXT("All components in the batch must be present in SortedComponentMergeRenderInfos"));
@@ -7567,7 +7637,7 @@ int32 ALandscape::PerformLayersHeightmapsBatchedMerge(const FUpdateLayersContent
 					}
 
 					// We need all neighbors to be present in this batch in order to be able to finalize that component :
-					if ((ValidNeighborComponents.Intersect(InParams.RenderBatch->ComponentsToRender).Num() == ValidNeighborComponents.Num())
+					if ((ValidNeighborComponents.Intersect(RenderBatch->ComponentsToRender).Num() == ValidNeighborComponents.Num())
 						&& InUpdateLayersContentContext.LandscapeComponentsHeightmapsToResolve.Contains(RenderedComponent))
 					{
 						ResolvedComponents.Add(RenderedComponent);
@@ -7585,11 +7655,11 @@ int32 ALandscape::PerformLayersHeightmapsBatchedMerge(const FUpdateLayersContent
 			// Recompose mip0 of the final heightmaps, subsection by subsection, to duplicate borders : 
 			InParams.MergeRenderContext->CycleBlendRenderTargets(RDGBuilderRecorder); 
 			ULandscapeScratchRenderTarget* ReadRT = InParams.MergeRenderContext->GetBlendRenderTargetRead();
-			
+
 			ReadRT->TransitionTo(ERHIAccess::SRVMask, RDGBuilderRecorder); // TODO [jonathan.bard] : This should be CopyDst but ExecuteCopyLayersTexture doesn't allow for it ATM
 
 			OnEditLayersMergedDelegate->Broadcast(FOnLandscapeEditLayersMergedParams(ReadRT->GetRenderTarget(),
-				/*InRenderAreaResolution = */InParams.RenderBatch->GetRenderTargetResolution(/*bInWithDuplicateBorders = */false), /*bInIsHeightmapMerge = */true));
+				/*InRenderAreaResolution = */RenderBatch->GetRenderTargetResolution(/*bInWithDuplicateBorders = */false), /*bInIsHeightmapMerge = */true));
 
 			ULandscapeScratchRenderTarget* WriteRT = InParams.MergeRenderContext->GetBlendRenderTargetWrite();
 
@@ -7605,7 +7675,7 @@ int32 ALandscape::PerformLayersHeightmapsBatchedMerge(const FUpdateLayersContent
 			const int32 TotalNumSubsections = Landscape->NumSubsections * Landscape->NumSubsections;
 			const int32 ComponentSubsectionVerts = Landscape->SubsectionSizeQuads + 1;
 
-			for (ULandscapeComponent* Component : InParams.RenderBatch->ComponentsToRender)
+			for (ULandscapeComponent* Component : RenderBatch->ComponentsToRender)
 			{
 				FComponentCopyInfo& ComponentCopyInfo = ComponentCopyInfos.Emplace_GetRef();
 				ComponentCopyInfo.Component = Component;
@@ -7613,7 +7683,7 @@ int32 ALandscape::PerformLayersHeightmapsBatchedMerge(const FUpdateLayersContent
 				// Effective area of the texture affecting this component (because of texture sharing) :
 				ComponentCopyInfo.TextureOffset = FIntPoint(static_cast<int32>(Component->HeightmapScaleBias.Z * ComponentCopyInfo.Texture->Source.GetSizeX()), static_cast<int32>(Component->HeightmapScaleBias.W * ComponentCopyInfo.Texture->Source.GetSizeY()));
 		
-				InParams.RenderBatch->ComputeSubsectionRects(Component, /*OutSubsectionRects = */ComponentCopyInfo.SourceSubsectionRects, /*OutSubsectionRectsWithDuplicateBorders = */ComponentCopyInfo.DestinationSubsectionRects);
+				RenderBatch->ComputeSubsectionRects(Component, /*OutSubsectionRects = */ComponentCopyInfo.SourceSubsectionRects, /*OutSubsectionRectsWithDuplicateBorders = */ComponentCopyInfo.DestinationSubsectionRects);
 				check(ComponentCopyInfo.SourceSubsectionRects.Num() == TotalNumSubsections);
 				check(ComponentCopyInfo.DestinationSubsectionRects.Num() == TotalNumSubsections);
 			}
@@ -7645,7 +7715,7 @@ int32 ALandscape::PerformLayersHeightmapsBatchedMerge(const FUpdateLayersContent
 
 			// "Expand" the scratch render target : 
 			// Copy sub-section by sub-section in order to duplicate borders :
-			InParams.MergeRenderContext->RenderExpandedRenderTarget(*InParams.RenderBatch, RDGBuilderRecorder);
+			InParams.MergeRenderContext->RenderExpandedRenderTarget(RDGBuilderRecorder);
 		}
 
 		// Generate the mips from the expanded RT and copy to the final texture mips
@@ -7654,7 +7724,7 @@ int32 ALandscape::PerformLayersHeightmapsBatchedMerge(const FUpdateLayersContent
 			const int32 NumMips = (int32)FMath::CeilLogTwo(Landscape->SubsectionSizeQuads) + 1;
 			RHI_BREADCRUMB_EVENT_GAMETHREAD("Generate %i remaining mips", NumMips - 1);
 
-			FIntPoint CurrentMipResolution = InParams.RenderBatch->GetRenderTargetResolution(/*bInWithDuplicateBorders = */true); // Mips are generated after the borders have been duplicated
+			FIntPoint CurrentMipResolution = RenderBatch->GetRenderTargetResolution(/*bInWithDuplicateBorders = */true); // Mips are generated after the borders have been duplicated
 			FIntPoint CurrentMipSubsectionSize = Landscape->SubsectionSizeQuads + 1;
 			for (int32 MipIndex = 1; MipIndex < NumMips; ++MipIndex)
 			{
@@ -7710,7 +7780,7 @@ int32 ALandscape::PerformLayersHeightmapsBatchedMerge(const FUpdateLayersContent
 					{
 						UTexture2D* ComponentHeightmap = Component->GetHeightmap(false);
 
-						FIntRect SourceSectionRect = InParams.RenderBatch->ComputeSectionRect(Component, /*bInWithDuplicateBorders = */true);
+						FIntRect SourceSectionRect = RenderBatch->ComputeSectionRect(Component, /*bInWithDuplicateBorders = */true);
 						SourceSectionRect.Min.X >>= MipIndex;
 						SourceSectionRect.Min.Y >>= MipIndex;
 						SourceSectionRect.Max.X >>= MipIndex;
@@ -9000,13 +9070,11 @@ struct FEditLayersWeightmapMergeParams
 // Render-thread version of the data / functions we need for the local merge of edit layers : 
 namespace EditLayersWeightmapLocalMerge_RenderThread
 {
-	using namespace UE::Landscape::EditLayers;
-
 	struct FEditLayerInfo
 	{
 		int32 SourceWeightmapTextureIndex = INDEX_NONE; // The index in VisibleEditLayerWeightmapTextures of the texture to read from for this layer
 		int32 SourceWeightmapTextureChannel = INDEX_NONE; // The channel of the texture to read from for this layer
-		EWeightmapBlendMode BlendMode = EWeightmapBlendMode::Num; // See EEditLayerWeightmapBlendMode
+		ELandscapeEditLayerWeightmapBlendMode BlendMode = ELandscapeEditLayerWeightmapBlendMode::Num; // See ELandscapeEditLayerWeightmapBlendMode
 		float Alpha = 1.0f; // Alpha value to be used in the blend
 	};
 
@@ -9157,7 +9225,7 @@ namespace EditLayersWeightmapLocalMerge_RenderThread
 		TArray<FTextureResolveInfo> TextureToResolveInfos;
 
 		// List of infos for each individual paint layer involved in the merge operation (including the visibility layer)
-		TArray<FWeightmapPaintLayerInfo> PaintLayerInfos;
+		TArray<FLandscapeEditLayerWeightmapPaintLayerInfo> PaintLayerInfos;
 
 		// Not truly render-thread data because it references UTextures but it's just because FLandscapeEditLayerReadback were historically game-thread initiated so for as long as we'll use those for readback, we need to store this here : 
 		TArray<FLandscapeLayersCopyReadbackTextureParams> DeferredCopyReadbackTextures;
@@ -9540,6 +9608,7 @@ namespace EditLayersWeightmapLocalMerge_RenderThread
 void ALandscape::PrepareLayersWeightmapsLocalMergeRenderThreadData(const FUpdateLayersContentContext& InUpdateLayersContentContext, const FEditLayersWeightmapMergeParams& InMergeParams, EditLayersWeightmapLocalMerge_RenderThread::FMergeInfo& OutRenderThreadData)
 {
 	using namespace EditLayersWeightmapLocalMerge_RenderThread;
+	using namespace UE::Landscape::Private;
 
 	TRACE_CPUPROFILER_EVENT_SCOPE(LandscapeLayers_PrepareLayersWeightmapsLocalMergeRenderThreadData);
 
@@ -9585,14 +9654,14 @@ void ALandscape::PrepareLayersWeightmapsLocalMergeRenderThreadData(const FUpdate
 							int32* PaintLayerInfoIndex = PaintLayerToPaintLayerInfoIndex.Find(ComponentEditLayerAllocation.LayerInfo);
 							if (PaintLayerInfoIndex == nullptr)
 							{
-								FWeightmapPaintLayerInfo NewPaintLayerInfo;
+								FLandscapeEditLayerWeightmapPaintLayerInfo NewPaintLayerInfo;
 								if (!ComponentEditLayerAllocation.LayerInfo->bNoWeightBlend)
 								{
-									NewPaintLayerInfo.Flags = EWeightmapPaintLayerFlags::IsWeightBlended;
+									NewPaintLayerInfo.Flags = ELandscapeEditLayerWeightmapPaintLayerFlags::IsWeightBlended;
 								}
 								if (ComponentEditLayerAllocation.LayerInfo == ALandscapeProxy::VisibilityLayer)
 								{
-									NewPaintLayerInfo.Flags = EWeightmapPaintLayerFlags::IsVisibilityLayer;
+									NewPaintLayerInfo.Flags = ELandscapeEditLayerWeightmapPaintLayerFlags::IsVisibilityLayer;
 								}
 
 								int32 NewPaintLayerInfoIndex = OutRenderThreadData.PaintLayerInfos.Add(NewPaintLayerInfo);
@@ -9633,7 +9702,7 @@ void ALandscape::PrepareLayersWeightmapsLocalMergeRenderThreadData(const FUpdate
 							NewEditLayerInfo.SourceWeightmapTextureChannel = ComponentEditLayerAllocation.WeightmapTextureChannel;
 							NewEditLayerInfo.Alpha = (ComponentEditLayerAllocation.LayerInfo == ALandscapeProxy::VisibilityLayer) ? 1.0f : Layer.WeightmapAlpha; // visibility can't affect or be affected by other paint layer weights
 							const bool* BlendSubstractive = Layer.WeightmapLayerAllocationBlend.Find(ComponentEditLayerAllocation.LayerInfo);
-							NewEditLayerInfo.BlendMode = ((BlendSubstractive != nullptr) && *BlendSubstractive) ? EWeightmapBlendMode::Subtractive : EWeightmapBlendMode::Additive;
+							NewEditLayerInfo.BlendMode = ((BlendSubstractive != nullptr) && *BlendSubstractive) ? ELandscapeEditLayerWeightmapBlendMode::Subtractive : ELandscapeEditLayerWeightmapBlendMode::Additive;
 
 							ComponentPaintLayerRenderInfo.VisibleEditLayerInfos.Add(NewEditLayerInfo);
 						}
@@ -10461,6 +10530,7 @@ int32 ALandscape::PerformLayersWeightmapsBatchedMerge(FUpdateLayersContentContex
 		//  already so we need to repack those into the appropriate weightmap channels
 
 		ALandscape* Landscape = InParams.MergeRenderContext->GetLandscape();
+		const FMergeRenderBatch* RenderBatch = InParams.MergeRenderContext->GetCurrentRenderBatch();
 
 		// Cycle render targets one last time so we can read from the last RT we've written to :
 		InParams.MergeRenderContext->CycleBlendRenderTargets(RDGBuilderRecorder);
@@ -10469,7 +10539,7 @@ int32 ALandscape::PerformLayersWeightmapsBatchedMerge(FUpdateLayersContentContex
 		ReadRT->TransitionTo(ERHIAccess::SRVMask, RDGBuilderRecorder);
 
 		OnEditLayersMergedDelegate->Broadcast(FOnLandscapeEditLayersMergedParams(ReadRT->GetRenderTarget(),
-			/*InRenderAreaResolution = */InParams.RenderBatch->GetRenderTargetResolution(/*bInWithDuplicateBorders = */false), /*bInIsHeightmapMerge = */false));
+			/*InRenderAreaResolution = */RenderBatch->GetRenderTargetResolution(/*bInWithDuplicateBorders = */false), /*bInIsHeightmapMerge = */false));
 
 		const int32 TotalNumSubsections = Landscape->NumSubsections * Landscape->NumSubsections;
 		FIntPoint MinWeightmapResolution(MAX_int32, MAX_int32);
@@ -10511,12 +10581,12 @@ int32 ALandscape::PerformLayersWeightmapsBatchedMerge(FUpdateLayersContentContex
 				for (ULandscapeComponent* RenderedComponent : ComponentsForWeightmap)
 				{
 					// If the component is present in this batch, we can resolve its channels :
-					if (InParams.RenderBatch->ComponentsToRender.Contains(RenderedComponent) 
+					if (RenderBatch->ComponentsToRender.Contains(RenderedComponent) 
 						// The component was possibly rendered in a previous batch, in which case we don't have to resolved it again : 
 						&& !ResolvedLandscapeComponents.Contains(RenderedComponent))
 					{
 						TArray<FIntRect, TInlineAllocator<4>> SourceSubsectionRects, DummySubsectionRects;
-						InParams.RenderBatch->ComputeSubsectionRects(RenderedComponent, /*OutSubsectionRects = */SourceSubsectionRects, /*OutSubsectionRectsWithDuplicateBorders = */DummySubsectionRects);
+						RenderBatch->ComputeSubsectionRects(RenderedComponent, /*OutSubsectionRects = */SourceSubsectionRects, /*OutSubsectionRectsWithDuplicateBorders = */DummySubsectionRects);
 						check(SourceSubsectionRects.Num() == TotalNumSubsections);
 
 						FWeightmapResolveInfo* WeightmapResolveInfo = WeightmapResolveInfosForBatch.Find(Weightmap);
diff --git a/Engine/Source/Runtime/Landscape/Public/LandscapeBlueprintBrushBase.h b/Engine/Source/Runtime/Landscape/Public/LandscapeBlueprintBrushBase.h
--- a/Engine/Source/Runtime/Landscape/Public/LandscapeBlueprintBrushBase.h
+++ b/Engine/Source/Runtime/Landscape/Public/LandscapeBlueprintBrushBase.h
@@ -108,7 +108,7 @@ public:
 	LANDSCAPE_EDIT_LAYERS_BATCHED_MERGE_EXPERIMENTAL
 	LANDSCAPE_API virtual TArray<UE::Landscape::EditLayers::FEditLayerRenderItem> GetRenderItems(const UE::Landscape::EditLayers::FMergeContext* InMergeContext) const override;
 	LANDSCAPE_EDIT_LAYERS_BATCHED_MERGE_EXPERIMENTAL
-	LANDSCAPE_API virtual void RenderLayer(ILandscapeEditLayerRenderer::FRenderParams& RenderParams, UE::Landscape::FRDGBuilderRecorder& RDGBuilderRecorder) override;
+	LANDSCAPE_API virtual bool RenderLayer(ILandscapeEditLayerRenderer::FRenderParams& RenderParams, UE::Landscape::FRDGBuilderRecorder& RDGBuilderRecorder) override;
 	LANDSCAPE_EDIT_LAYERS_BATCHED_MERGE_EXPERIMENTAL
 	LANDSCAPE_API virtual FString GetEditLayerRendererDebugName() const override;
 	//~ End ILandscapeEditLayerRenderer implementation
diff --git a/Engine/Source/Runtime/Landscape/Public/LandscapeEditLayerRenderer.h b/Engine/Source/Runtime/Landscape/Public/LandscapeEditLayerRenderer.h
--- a/Engine/Source/Runtime/Landscape/Public/LandscapeEditLayerRenderer.h
+++ b/Engine/Source/Runtime/Landscape/Public/LandscapeEditLayerRenderer.h
@@ -6,8 +6,9 @@
 #include "LandscapeEditTypes.h"
 #include "LandscapeUtils.h"
 #include "Engine/TextureRenderTarget2D.h"
-#include "RHIAccess.h"
 #include "EngineDefines.h"
+#include "RHIAccess.h"
+#include "StructUtils/InstancedStruct.h"
 
 #include "LandscapeEditLayerRenderer.generated.h"
 
@@ -20,11 +21,18 @@ class ULandscapeEditLayerRenderer;
 class ILandscapeEditLayerRenderer;
 class FRDGBuilder;
 class FRHICommandListImmediate;
+struct FLandscapeEditLayerRenderBlackboardItem;
+
 namespace UE::Landscape
 {
 	class FRDGBuilderRecorder;
 } // namespace UE::Landscape
 
+// TODO [jonathan.bard] : Move some classes into their own file, it's getting messy in there
+
+
+// ----------------------------------------------------------------------------------
+
 namespace UE::Landscape::EditLayers
 {
 
@@ -35,31 +43,108 @@ FString ConvertTargetLayerNamesToString(const TArrayView<const FName>& InTargetL
 
 // ----------------------------------------------------------------------------------
 
-// Must match EEditLayerWeightmapBlendMode in LandscapeLayersWeightmapsPS.usf / LandscapeEditLayersWeightmapsPS.usf
+// Must match EHeightmapBlendMode in LandscapeEditLayersHeightmaps.usf
+enum class EHeightmapBlendMode : uint32
+{
+	Additive = 0,
+	LegacyAlphaBlend, // In legacy alpha blend, the layer stores the height with premultiplied alpha (legacy landscape splines)
+	AlphaBlend,
+
+	Num,
+};
+
+// Must match EHeightmapAlphaFlags in LandscapeCommon.ush
+enum class EHeightmapAlphaFlags : uint8
+{
+	None = 0, // aka Additive
+	Additive = 0, // The height is considered to be an offset (positive or negative)
+	Min = (1 << 0), // Only lower the existing landscape values
+	Max = (1 << 1), // Only raise the existing landscape values
+	AlphaBlend = (Min | Max), // Raise or lower the existing landscape values
+};
+ENUM_CLASS_FLAGS(EHeightmapAlphaFlags);
+
+// Must match EWeightmapBlendMode in LandscapeEditLayersWeightmaps.usf
 enum class EWeightmapBlendMode : uint32
 {
+	None = 0, // aka Additive
 	Additive = 0,
 	Subtractive,
+	Passthrough,
+	AlphaBlend, 
 
 	Num,
 };
 
-// Must match EWeightmapPaintLayerFlags in LandscapeLayersWeightmapsPS.usf / LandscapeEditLayersWeightmapsPS.usf
-enum class EWeightmapPaintLayerFlags : uint32
+// Must match EWeightmapAlphaFlags in LandscapeCommon.ush
+enum class EWeightmapAlphaFlags : uint8
+{
+	None = 0, // aka Additive
+	Additive = 0, // The weight is considered to be an offset (positive)
+	Min = (1 << 0), // Only retain the min between the weight and the existing landscape weight value
+	Max = (1 << 1), // Only retain the max between the weight and the existing landscape weight value
+	AlphaBlend = (Min | Max), // Full alpha blending of the weight against the existing landscape weight value
+};
+ENUM_CLASS_FLAGS(EWeightmapAlphaFlags);
+
+// Must match EWeightmapTargetLayerFlags in LandscapeEditLayersWeightmaps.usf
+enum class EWeightmapTargetLayerFlags : uint32
 {
-	IsVisibilityLayer = (1 << 0), // This paint layer is the visibility layer
-	IsWeightBlended = (1 << 1), // Blend the paint layer's value with all the other paint layers weights
+	IsVisibilityLayer = (1 << 0), // This target layer is the visibility layer
+	IsWeightBlended = (1 << 1), // Blend the target layer's value with all the other target layers weights
 
 	None = 0
 };
-ENUM_CLASS_FLAGS(EWeightmapPaintLayerFlags);
+ENUM_CLASS_FLAGS(EWeightmapTargetLayerFlags);
 
-// ----------------------------------------------------------------------------------
+// Must match FWeightmapTargetLayerInfo in LandscapeEditLayersWeightmaps.usf
+struct FWeightmapTargetLayerInfo
+{
+	EWeightmapTargetLayerFlags Flags = EWeightmapTargetLayerFlags::None; // Additional info about this target layer
+};
 
-// Must match FWeightmapPaintLayerInfo in LandscapeLayersWeightmapsPS.usf / LandscapeEditLayersWeightmaps.usf
-struct FWeightmapPaintLayerInfo
+// Defines how heightmaps should be blended (see GenericBlendLayer)
+struct FHeightmapBlendParams
 {
-	EWeightmapPaintLayerFlags Flags = EWeightmapPaintLayerFlags::None; // Additional info about this paint layer
+	FHeightmapBlendParams() = default;
+	FHeightmapBlendParams(EHeightmapBlendMode InBlendMode)
+		: BlendMode(InBlendMode)
+	{}
+
+	EHeightmapBlendMode BlendMode = EHeightmapBlendMode::Additive;
+	float Alpha = 1.0f;
+};
+
+// Defines how weightmaps should be blended (see GenericBlendLayer)
+struct FWeightmapBlendParams
+{
+	FWeightmapBlendParams() = default;
+	FWeightmapBlendParams(EWeightmapBlendMode InBlendMode)
+		: BlendMode(InBlendMode)
+	{}
+
+	static const FWeightmapBlendParams& GetDefaultPassthroughBlendParams()
+	{
+		static const FWeightmapBlendParams BlendParams(EWeightmapBlendMode::Passthrough);
+		return BlendParams;
+	}
+
+	static const FWeightmapBlendParams& GetDefaultAdditiveBlendParams()
+	{
+		static const FWeightmapBlendParams BlendParams(EWeightmapBlendMode::Additive);
+		return BlendParams;
+	}
+
+	EWeightmapBlendMode BlendMode = EWeightmapBlendMode::Passthrough;
+	float Alpha = 1.0f;
+};
+
+// Defines heightmaps+weightmaps blending params (see GenericBlendLayer)
+//  There should be as many blend params as there are target layers to be blended in the blend operation, the others are simply passthrough
+struct FBlendParams
+{
+	FHeightmapBlendParams HeightmapBlendParams;
+	TMap<FName, FWeightmapBlendParams> WeightmapBlendParams;
 };
 
 
@@ -300,11 +385,16 @@ protected:
 class FEditLayerTargetTypeState
 {
 	friend class ILandscapeEditLayerRenderer;
+	friend class FEditLayerRendererState;
 
 private:
-	FEditLayerTargetTypeState() = delete;
+	// Private constructor : either use the constructors taking a merge context in parameter or use GetDummyTargetTypeState()
+	FEditLayerTargetTypeState() = default;
 
 public:
+	LANDSCAPE_EDIT_LAYERS_BATCHED_MERGE_EXPERIMENTAL
+	static const FEditLayerTargetTypeState& GetDummyTargetTypeState();
+
 	/**
 	 * Constructs an empty FEditLayerTargetTypeState (all target types turned off, no active weightmaps)
 	 * 
@@ -471,9 +561,13 @@ private:
 class FEditLayerRendererState
 {
 private:
-	FEditLayerRendererState() = delete;
+	// Private constructor : either use the constructors taking a merge context in parameter or use GetDummyRendererState()
+	FEditLayerRendererState() = default;
 
 public:
+	LANDSCAPE_EDIT_LAYERS_BATCHED_MERGE_EXPERIMENTAL
+	static const FEditLayerRendererState& GetDummyRendererState();
+
 	LANDSCAPE_EDIT_LAYERS_BATCHED_MERGE_EXPERIMENTAL
 	LANDSCAPE_API explicit FEditLayerRendererState(const FMergeContext* InMergeContext);
 	LANDSCAPE_EDIT_LAYERS_BATCHED_MERGE_EXPERIMENTAL
@@ -607,13 +701,13 @@ private:
 	FString DebugName;
 	
 	/** Struct that contains the supported target types and the which weightmap are supported by this renderer. Immutable. */
-	FEditLayerTargetTypeState SupportedTargetTypeState;
+	FEditLayerTargetTypeState SupportedTargetTypeState = FEditLayerTargetTypeState::GetDummyTargetTypeState();
 	
 	/** Struct that contains the enabled target types and the which weightmap are currently enabled by this renderer. Can be set by the user using EnableTargetType. */
-	FEditLayerTargetTypeState EnabledTargetTypeState;
+	FEditLayerTargetTypeState EnabledTargetTypeState = FEditLayerTargetTypeState::GetDummyTargetTypeState();
 
 	/** Intersection of SupportedTargetTypeState and EnabledTargetTypeState (defines what is both supported and enabled on this renderer) */
-	FEditLayerTargetTypeState ActiveTargetTypeState;
+	FEditLayerTargetTypeState ActiveTargetTypeState = FEditLayerTargetTypeState::GetDummyTargetTypeState();
 
 	/* List of groups of target layers that this renderer requires to be rendered together. All target layers listed in SupportedTargetTypeState must belong to 1 (and 1 only) target layer group.
 	 *  Each target layer group is a bit array for which each bit corresponds to an entry in FMergeContext's AllTargetLayerNames
@@ -647,6 +741,29 @@ struct FMergeRenderParams
 };
 
 
+// ----------------------------------------------------------------------------------
+
+/** Flags that allow a given ILandscapeEditLayerRenderer to customize the way it renders/blends */
+enum class ERenderFlags : uint8
+{
+	None = 0,
+
+	// Render mode flags
+	RenderMode_Recorded = (1 << 0), // This renderer can record its render commands into a single FRDGBuilder on the render thread (prefer this if possible), via a FRDGBuilderRecorder in "recording" mode. Exclusive with RenderMode_Immediate
+	RenderMode_Immediate = (1 << 1), // This renderer enqueues its render commands immediately either via the FRDGBuilderRecorder in "immediate" mode or just enqueuing render commands the usual way. Exclusive with RenderMode_Recorded
+	RenderMode_Mask = RenderMode_Recorded | RenderMode_Immediate,
+
+	// Blend mode flags
+	BlendMode_SeparateBlend = (1 << 2), // This renderer has a separate render function for blending. When this flag is not set, only RenderLayer is called and is assumed to both render the layer and blend it. When it's set, RenderLayer will be followed by BlendLayer.
+
+	// Render layer group flags
+	RenderLayerGroup_SupportsGrouping = (1 << 3), // This renderer supports being rendered along with others in a series of RenderLayer steps before performing a single BlendLayer. Assumes BlendMode_SeparateBlend
+};
+ENUM_CLASS_FLAGS(ERenderFlags);
+
+
+// ----------------------------------------------------------------------------------
+
 /**
  * Defines an individual render step of the batch merge
  */
@@ -658,26 +775,50 @@ struct FMergeRenderStep
 		Invalid,
 
 		BeginRenderCommandRecorder, // Initiates a render command recorder that will batch one or several RenderLayer_Recorded operations, runs on the game thread
-		RenderLayer_Recorded, // Performs the rendering of a target layer group on an edit layer on a given world region (i.e. in a batch), runs on the game thread but this is where render operation (a lambda that will run on the render thread) can be enqueued on the command recorder initiated by BeginRenderCommandRecorder
 		EndRenderCommandRecorder, // Ends the render command recorder initiated by the last BeginRenderCommandRecorder, runs on the game thread but this is where the render command executing the recorded operations from (several) RenderLayer_Recorded step(s) will be pushed (on the render thread, this is where the FRDBuilder is created and executed)
 
-		RenderLayer_Immediate, // Perform the rendering of a target layer group on an edit layer on a given world region (i.e. in a batch), runs on the game thread
+		BeginRenderLayerGroup, // Initiates a series of RenderLayer steps for renderers that can be rendered one after another without an intermediate BlendLayer. Triggers a call to BeginRenderLayerGroup() on the first renderer in the group 
+		EndRenderLayerGroup, // Ends a series of RenderLayer steps for renderers that can be rendered one after another without an intermediate BlendLayer. Triggers a call to EndRenderLayerGroup() on the last renderer in the group 
+
+		RenderLayer, // Performs the rendering of a target layer group on an edit layer on a given world region (i.e. in a batch), runs on the game thread. Use the command recorder to perform the render operations. In RenderMode_Recorded, these will be delayed until EndRenderCommandRecorder
+		BlendLayer, // Performs the rendering of a target layer group on an edit layer on a given world region (i.e. in a batch), runs on the game thread. Use the command recorder to perform the render operations. In RenderMode_Recorded, these will be delayed until EndRenderCommandRecorder
+
 		SignalBatchMergeGroupDone, // Final step when rendering a target layer group on a given world region (i.e. in a batch) : runs on the game thread and allows to retrieve the result of the merge and do something with it (e.g. resolve the final textures)
 	};
 
-	FMergeRenderStep(EType InType, const FEditLayerRendererState& InRendererState, const TBitArray<>& InTargetLayerGroupBitIndices, const TArrayView<ULandscapeComponent*>& InComponentsToRender)
+	FMergeRenderStep(EType InType)
 		: Type(InType)
+	{}
+
+
+	FMergeRenderStep(EType InType, const TBitArray<>& InTargetLayerGroupBitIndices, const TArrayView<ULandscapeComponent*>& InComponentsToRender)
+		: Type(InType)
+		, TargetLayerGroupBitIndices(InTargetLayerGroupBitIndices)
+		, ComponentsToRender(InComponentsToRender)
+	{}
+
+	FMergeRenderStep(EType InType, ERenderFlags InRenderFlags, const FEditLayerRendererState& InRendererState, const TBitArray<>& InTargetLayerGroupBitIndices, const TArrayView<ULandscapeComponent*>& InComponentsToRender)
+		: Type(InType)
+		, RenderFlags(InRenderFlags)
 		, RendererState(InRendererState)
 		, TargetLayerGroupBitIndices(InTargetLayerGroupBitIndices)
 		, ComponentsToRender(InComponentsToRender)
 	{}
 
+	UE_DEPRECATED(5.6, "Use the other constructors")
+	FMergeRenderStep(EType InType, const FEditLayerRendererState& InRendererState, const TBitArray<>& InTargetLayerGroupBitIndices, const TArrayView<ULandscapeComponent*>& InComponentsToRender)
+		: FMergeRenderStep(InType, ERenderFlags::None, InRendererState, InTargetLayerGroupBitIndices, InComponentsToRender)
+	{}
+
 	/** Type of operation for this step */
 	EType Type = EType::Invalid;
 	
-	/** Renderer state to be used this step (only when Type == EType::RenderLayer). This includes the renderer as well as its precise step (e.g. which weightmap are supported? which are enabled?) */
-	FEditLayerRendererState RendererState;
-	
+	/** The render flags corresponding to this state */
+	ERenderFlags RenderFlags = ERenderFlags::None;
+
+	/** Renderer state to be used this step. This includes the renderer as well as its precise step (e.g. which weightmap are supported? which are enabled?) */
+	FEditLayerRendererState RendererState = FEditLayerRendererState::GetDummyRendererState();
+
 	/** List of target layers being involved in this step. Each bit in that bit array corresponds to an entry in FMergeContext's AllTargetLayerNames */
 	TBitArray<> TargetLayerGroupBitIndices;
 
@@ -823,12 +964,10 @@ public:
 	struct FOnRenderBatchTargetGroupDoneParams
 	{
 		FOnRenderBatchTargetGroupDoneParams(FMergeRenderContext* InMergeRenderContext, 
-			const FMergeRenderBatch& InRenderBatch, 
 			const TArrayView<FName>& InTargetLayerGroupLayerNames,
 			const TArrayView<ULandscapeLayerInfoObject*>& InTargetLayerGroupLayerInfos,
 			const TArrayView<FComponentMergeRenderInfo>& InSortedComponentMergeRenderInfos)
 			: MergeRenderContext(InMergeRenderContext)
-			, RenderBatch(&InRenderBatch)
 			, TargetLayerGroupLayerNames(InTargetLayerGroupLayerNames)
 			, TargetLayerGroupLayerInfos(InTargetLayerGroupLayerInfos)
 			, SortedComponentMergeRenderInfos(InSortedComponentMergeRenderInfos)
@@ -837,9 +976,6 @@ public:
 		/** Render context : this is still active in this step and can be used for doing additional renders in the blend render targets, etc. */
 		FMergeRenderContext* MergeRenderContext = nullptr;
 
-		/** Batch to that was just rendered for this target layer group */
-		const FMergeRenderBatch* RenderBatch = nullptr;
-
 		/** List of target layers being involved in this step */
 		TArray<FName> TargetLayerGroupLayerNames;
 
@@ -853,6 +989,8 @@ public:
 		TArray<FName> RenderGroupTargetLayerNames;
 		UE_DEPRECATED(5.6, "Renamed : use TargetLayerGroupLayerInfos")
 		TArray<ULandscapeLayerInfoObject*> RenderGroupTargetLayerInfos;
+		UE_DEPRECATED(5.6, "Removed : use MergeRenderContext->GetCurrentRenderBatch()")
+		const FMergeRenderBatch* RenderBatch = nullptr;
 	};
 	void Render(TFunction<void(const FOnRenderBatchTargetGroupDoneParams& /*InParams*/, UE::Landscape::FRDGBuilderRecorder& /*RDGBuilderRecorder*/)> OnRenderBatchTargetGroupDone);
 
@@ -879,23 +1017,83 @@ public:
 
 	/** 
 	 * Render the stencil render targets for each target layer in this merge for this batch 
-	 * @param InRenderBatch current render batch
-	* @param RDGBuilderRecorder recorder to append operations to a single FRDGBuilder, if the the recorder in "recording" mode. In "immediate" mode, a render command will be enqueued immediately
+	 * 
+	 * @param RDGBuilderRecorder recorder to append operations to a single FRDGBuilder, if the the recorder in "recording" mode. In "immediate" mode, a render command will be enqueued immediately
 	 */
-	void RenderValidityRenderTargets(const FMergeRenderBatch& InRenderBatch, UE::Landscape::FRDGBuilderRecorder& RDGBuilderRecorder);
+	void RenderValidityRenderTargets(UE::Landscape::FRDGBuilderRecorder& RDGBuilderRecorder);
 
 	/** 
 	 * Duplicates the vertex data from the (sub-)sections of the batch, assuming GetBlendRenderTargetRead() is the RT that is read from and GetBlendRenderTargetWrite() the one that is written to 
-	 * @param InRenderBatch current render batch
-	* @param RDGBuilderRecorder recorder to append operations to a single FRDGBuilder, if the the recorder in "recording" mode. In "immediate" mode, a render command will be enqueued immediately
+	 * 
+	 * @param RDGBuilderRecorder recorder to append operations to a single FRDGBuilder, if the the recorder in "recording" mode. In "immediate" mode, a render command will be enqueued immediately
+	 */
+	void RenderExpandedRenderTarget(UE::Landscape::FRDGBuilderRecorder& RDGBuilderRecorder);
+
+	/** 
+	 * Performs a generic blend assuming GetBlendRenderTargetWrite() is the RT that contains the layer to blend and GetBlendRenderTargetRead() the one that contains the result of the merge up until this layer.
+	 *  It will cycle the RTs such as after this call, GetBlendRenderTargetWrite() will contain the merge result
+	 * 
+	 * @param InBlendParams defines how heightmaps/weightmaps should be blended. In the case of weightmaps, there must be as may entries in WeightmapBlendParams as there are layers in the target layers group
+	 * @param RDGBuilderRecorder recorder to append operations to a single FRDGBuilder, if the the recorder in "recording" mode. In "immediate" mode, a render command will be enqueued immediately
 	 */
-	void RenderExpandedRenderTarget(const FMergeRenderBatch& InRenderBatch, UE::Landscape::FRDGBuilderRecorder& RDGBuilderRecorder);
+	LANDSCAPE_EDIT_LAYERS_BATCHED_MERGE_EXPERIMENTAL
+	LANDSCAPE_API void GenericBlendLayer(const FBlendParams& InBlendParams, UE::Landscape::FRDGBuilderRecorder& RDGBuilderRecorder);
 
 	const TBitArray<>& GetFinalTargetLayerBitIndices() const 
 	{ 
 		return FinalTargetLayerBitIndices; 
 	}
 
+	/** 
+	 * @return true if there is at least one blackboard item of type T in the context's list 
+	 */
+	template <typename T, typename = std::enable_if_t<std::is_base_of_v<FLandscapeEditLayerRenderBlackboardItem, std::decay_t<T>>>>
+	LANDSCAPE_EDIT_LAYERS_BATCHED_MERGE_EXPERIMENTAL
+	bool HasBlackboardItem() const;
+
+	/** 
+	 * Create a new blackboard item of type T and adds it to the context's list
+	 * 
+	 * @return the newly-created blackboard item
+	 */
+	template <typename T, typename... TArgs, typename = std::enable_if_t<std::is_base_of_v<FLandscapeEditLayerRenderBlackboardItem, std::decay_t<T>>>>
+	LANDSCAPE_EDIT_LAYERS_BATCHED_MERGE_EXPERIMENTAL
+	T& AddBlackboardItem(TArgs&&... InArgs);
+
+	/**
+	 * @return the first blackboard item of type T from the context's list or nullptr if there isn't one
+	 */
+	template <typename T, typename = std::enable_if_t<std::is_base_of_v<FLandscapeEditLayerRenderBlackboardItem, std::decay_t<T>>>>
+	LANDSCAPE_EDIT_LAYERS_BATCHED_MERGE_EXPERIMENTAL
+	T* TryGetBlackboardItem();
+
+	/**
+	 * @return the first blackboard item of type T from the context's list. Asserts if there isn't one
+	 */
+	template <typename T, typename = std::enable_if_t<std::is_base_of_v<FLandscapeEditLayerRenderBlackboardItem, std::decay_t<T>>>>
+	LANDSCAPE_EDIT_LAYERS_BATCHED_MERGE_EXPERIMENTAL
+	T& GetBlackboardItem();
+
+	/**
+	 * @return the first blackboard item of type T from the context's list or create a new one if there isn't one
+	 */
+	template <typename T, typename... TArgs, typename = std::enable_if_t<std::is_base_of_v<FLandscapeEditLayerRenderBlackboardItem, std::decay_t<T>>>>
+	LANDSCAPE_EDIT_LAYERS_BATCHED_MERGE_EXPERIMENTAL
+	T& GetOrCreateBlackboardItem(TArgs&&... InArgs);
+
+	/**
+	 * @return all blackboard items of type T from the context's list
+	 */
+	template <typename T, typename = std::enable_if_t<std::is_base_of_v<FLandscapeEditLayerRenderBlackboardItem, std::decay_t<T>>>>
+	LANDSCAPE_EDIT_LAYERS_BATCHED_MERGE_EXPERIMENTAL
+	TArray<T*> GetBlackboardItems();
+
+	/**
+	 * @return all blackboard items currently stored in the context
+	 */
+	LANDSCAPE_EDIT_LAYERS_BATCHED_MERGE_EXPERIMENTAL
+	const TArray<TInstancedStruct<FLandscapeEditLayerRenderBlackboardItem>>& GetAllBlackboardItems() const { return BlackboardItems; }
+
 private:
 	/** Allocate all needed render targets for this merge */
 	void AllocateResources();
@@ -947,6 +1145,9 @@ private:
 
 	/** Reverse lookup of ComponentToTargetLayerBitIndices : one entry per element in ComponentToTargetLayerBitIndices, each entry containing all of the components involved in this merge for this target layer */
 	TArray<TSet<ULandscapeComponent*>> TargetLayersToComponents;
+
+	/** Generic data struct to store some data to pass around throughout the context's lifetime. Basically, this is a mini-RTTI system that e.g. allows to store/pass around specific data between renderers in a generic way */
+	TArray<TInstancedStruct<FLandscapeEditLayerRenderBlackboardItem>> BlackboardItems;
 };
 
 
@@ -964,6 +1165,8 @@ struct FOOBox2D
 
 	FTransform Transform;
 	FVector2D Extents = FVector2D(ForceInit);
+
+	FBox BuildAABB() const;
 };
 
 
@@ -995,6 +1198,7 @@ public:
 	FIntRect GetSpecificComponentKeys() const;
 	/** In the EType::OOBox case, returns the OOBox */
 	const FOOBox2D& GetOOBox() const { check(Type == EType::OOBox); return OOBox2D; }
+	const FOOBox2D* TryGetOOBox() const { return (Type == EType::OOBox) ? &OOBox2D : nullptr; }
 
 	FBox ComputeWorldAreaAABB(const FTransform& InLandscapeTransform, const FBox& InLandscapeLocalBounds, const FTransform& InComponentTransform, const FBox& InComponentLocalBounds) const;
 	FOOBox2D ComputeWorldAreaOOBB(const FTransform& InLandscapeTransform, const FBox& InLandscapeLocalBounds, const FTransform& InComponentTransform, const FBox& InComponentLocalBounds) const;
@@ -1045,6 +1249,7 @@ public:
 	const FIntPoint& GetSpecificComponentKey() const { check(Type == EType::SpecificComponent); return SpecificComponentKey; }
 	/** In the EType::OOBox case, returns the OOBox */
 	const FOOBox2D& GetOOBox() const { check(Type == EType::OOBox); return OOBox2D; }
+	const FOOBox2D* TryGetOOBox() const { return (Type == EType::OOBox) ? &OOBox2D : nullptr; }
 
 	FBox ComputeWorldAreaAABB(const FTransform& InComponentTransform, const FBox& InComponentLocalBounds) const;
 	FOOBox2D ComputeWorldAreaOOBB(const FTransform& InComponentTransform, const FBox& InComponentLocalBounds) const;
@@ -1139,18 +1344,22 @@ public:
 };
 #endif // WITH_EDITOR
 
-enum class ERenderFlags : uint8
-{
-	None = 0,
+} //namespace UE::Landscape::EditLayers
 
-	// Render mode flags
-	RenderMode_Recorded = (1 << 0), // This renderer can record its render commands into a single FRDGBuilder on the render thread (prefer this if possible), via a FRDGBuilderRecorder in "recording" mode
-	RenderMode_Immediate = (1 << 1), // This renderer enqueues its render commands immediately either via the FRDGBuilderRecorder in "immediate" mode or just enqueuing render commands the usual way
-	RenderMode_Mask = RenderMode_Recorded | RenderMode_Immediate
-};
-ENUM_CLASS_FLAGS(ERenderFlags);
 
-} //namespace UE::Landscape::EditLayers
+// ----------------------------------------------------------------------------------
+
+/**
+ * Base class for various blackboard items that can be attached to FMergeRenderContext in order to transfer temporary information
+ *  between renderers. It's a USTRUCT only to benefit from the RTTI services provided by TInstancedStruct
+ */
+USTRUCT(meta = (Abstract))
+struct FLandscapeEditLayerRenderBlackboardItem
+{
+	GENERATED_BODY()
+	
+	virtual ~FLandscapeEditLayerRenderBlackboardItem() = default;
+};
 
 // ----------------------------------------------------------------------------------
 
@@ -1183,7 +1392,9 @@ public:
 			const TArrayView<ULandscapeLayerInfoObject*>& InTargetLayerGroupLayerInfos,
 			const UE::Landscape::EditLayers::FEditLayerRendererState& InRendererState, 
 			const TArrayView<UE::Landscape::EditLayers::FComponentMergeRenderInfo>& InSortedComponentMergeRenderInfos, 
-			const FTransform& InRenderAreaWorldTransform, const FIntRect& InRenderAreaSectionRect)
+			const FTransform& InRenderAreaWorldTransform, 
+			const FIntRect& InRenderAreaSectionRect,
+			int32 InNumSuccessfulRenderLayerStepsUntilBlendLayerStep)
 			: MergeRenderContext(InMergeRenderContext)
 			, TargetLayerGroupLayerNames(InTargetLayerGroupLayerNames)
 			, TargetLayerGroupLayerInfos(InTargetLayerGroupLayerInfos)
@@ -1191,6 +1402,7 @@ public:
 			, SortedComponentMergeRenderInfos(InSortedComponentMergeRenderInfos)
 			, RenderAreaWorldTransform(InRenderAreaWorldTransform)
 			, RenderAreaSectionRect(InRenderAreaSectionRect)
+			, NumSuccessfulRenderLayerStepsUntilBlendLayerStep(InNumSuccessfulRenderLayerStepsUntilBlendLayerStep)
 		{}
 
 		/** Merge context */
@@ -1217,6 +1429,9 @@ public:
 		/** SectionRect (i.e. landscape vertex coordinates, in landscape space) that corresponds to this render area */
 		FIntRect RenderAreaSectionRect;
 
+		/** When separate blend is enabled, tracks how many RenderLayer calls have succeeded yet (valid until the BlendLayer step occurs) */
+		int32 NumSuccessfulRenderLayerStepsUntilBlendLayerStep = 0;
+
 		UE_DEPRECATED(5.6, "Renamed : use TargetLayerGroupLayerNames")
 		TArray<FName> RenderGroupTargetLayerNames;
 		UE_DEPRECATED(5.6, "Renamed : use TargetLayerGroupLayerInfos")
@@ -1268,15 +1483,16 @@ public:
 
 	/**
 	 * @param InMergeContext context containing all sorts of information related to this merge operation
+	 * 
 	 * @return details about how this renderer's render method is implemented 
 	*/
 	LANDSCAPE_API virtual UE::Landscape::EditLayers::ERenderFlags GetRenderFlags(const UE::Landscape::EditLayers::FMergeContext* InMergeContext) const
 		PURE_VIRTUAL(ILandscapeEditLayerRenderer::GetRenderFlags, return UE::Landscape::EditLayers::ERenderFlags::None; );
 
 	/**
-	 * RenderLayer is where the renderer has a chance to render its content and eventually blend it with the merged result of all preceding layers
-	 * It operates on a limited set of components (depending on the size of the render batches) and on a set of target layers (e.g. multiple weightmaps).
-	 * It guarantees access to merged result from preceding layers of each target layer. 
+	 * RenderLayer is where the renderer has a chance to render its content and eventually blend it with the merged result of all preceding layers (if ERenderFlags::BlendMode_SeparateBlend is not returned)
+	 *  It operates on a limited set of components (depending on the size of the render batches) and on a set of target layers (e.g. multiple weightmaps).
+	 *  It guarantees access to the merged result from preceding layers of each target layer. 
 	 * 
 	 * @param RenderParams contains parameters necessary for this render (affected components, merge context, etc.)
 	 * 
@@ -1287,13 +1503,142 @@ public:
 	 *     consecutive "render commands" via "RDG/render command" accumulated on a FRDGBuilderRecorder. This allows to coalesce several render commands onto the same FRDGBuilder (which is critical 
 	 *     for performance) while still allowing to interleave game thread-based renders (flushing the command recorder, enqueuing the game thread-based render, and starting recording on the same (or a 
 	 *     new) command recorder again...) 
-	 *     TLDR: when using the ERenderFlags::Recorded method, the user should use the command recorder to enqueue lambdas instead if enqueuing render commands directly. 
+	 *     TLDR: when using the ERenderFlags::RenderMode_Recorded method, the user should use the command recorder to enqueue lambdas instead if enqueuing render commands directly. 
 	 *     Corollary: Any render command issued by RenderLayer will end up being pushed before the render operations recorded on the RDG command recorder (unless it is flushed), so there's no reason 
 	 *     to actually do it, unless you want those commands to run before the recorded commands
+	 * 
+	 * @return true if anything was rendered
+	*/
+	LANDSCAPE_EDIT_LAYERS_BATCHED_MERGE_EXPERIMENTAL
+	LANDSCAPE_API virtual bool RenderLayer(ILandscapeEditLayerRenderer::FRenderParams& RenderParams, UE::Landscape::FRDGBuilderRecorder& RDGBuilderRecorder) 
+		PURE_VIRTUAL(ILandscapeEditLayerRenderer::RenderLayer, return false; );
+
+	/**
+	 * BlendLayer is where the renderer has a chance to blend its content with the merged result of all preceding layers
+	 * It operates on a limited set of components (depending on the size of the render batches) and on a set of target layers (e.g. multiple weightmaps).
+	 * It guarantees access to the merged result from preceding layers of each target layer. 
+	 * 
+	 * @param RenderParams contains parameters necessary for this render (affected components, merge context, etc.)
+	 * 
+	 * @param RDGBuilderRecorder (see RenderLayer)
 	*/
 	LANDSCAPE_EDIT_LAYERS_BATCHED_MERGE_EXPERIMENTAL
-	LANDSCAPE_API virtual void RenderLayer(ILandscapeEditLayerRenderer::FRenderParams& RenderParams, UE::Landscape::FRDGBuilderRecorder& RDGBuilderRecorder) 
-		PURE_VIRTUAL(ILandscapeEditLayerRenderer::RenderLayer, );
+	LANDSCAPE_API virtual void BlendLayer(ILandscapeEditLayerRenderer::FRenderParams& RenderParams, UE::Landscape::FRDGBuilderRecorder& RDGBuilderRecorder) {}
+
+	/**
+	 * Return whether this renderer's RenderLayer operation can be grouped with one of the previous renderers in the render layer group. 
+	 *  Only called if ERenderFlags::RenderLayerGroup_SupportsGrouping is returned by GetRenderFlags() on both renderers
+	 * 
+	 * @param InOtherRenderer the other renderer we're trying to group this renderer's RenderLayer operation with
+	 * 
+	 * @return true if both RenderLayer can be done in the same render layer group
+	*/
+	LANDSCAPE_EDIT_LAYERS_BATCHED_MERGE_EXPERIMENTAL
+	LANDSCAPE_API virtual bool CanGroupRenderLayerWith(TScriptInterface<ILandscapeEditLayerRenderer> InOtherRenderer) const { return false; }
+
+	/**
+	 * Called before the first call to RenderLayer, on the first renderer of a render layer group : 
+	 *  Only called if ERenderFlags::RenderLayerGroup_SupportsGrouping is returned by GetRenderFlags() 
+	 *
+	 * @param RenderParams contains parameters necessary for this render (affected components, merge context, etc.)
+	 */
+	LANDSCAPE_EDIT_LAYERS_BATCHED_MERGE_EXPERIMENTAL
+	LANDSCAPE_API virtual void BeginRenderLayerGroup(ILandscapeEditLayerRenderer::FRenderParams& RenderParams, UE::Landscape::FRDGBuilderRecorder& RDGBuilderRecorder) {}
+
+	/**
+	 * Called after the last call to RenderLayer, on the last renderer of a render layer group. 
+	 *  Only called if ERenderFlags::RenderLayerGroup_SupportsGrouping is returned by GetRenderFlags() 
+	 *
+	 * @param RenderParams contains parameters necessary for this render (affected components, merge context, etc.)
+	 */
+	LANDSCAPE_EDIT_LAYERS_BATCHED_MERGE_EXPERIMENTAL
+	LANDSCAPE_API virtual void EndRenderLayerGroup(ILandscapeEditLayerRenderer::FRenderParams& RenderParams, UE::Landscape::FRDGBuilderRecorder& RDGBuilderRecorder) {}
 
 #endif // WITH_EDITOR
 };
+
+// TODO [jonathan.bard] : move to .inl file :
+namespace UE::Landscape::EditLayers
+{
+
+#if WITH_EDITOR
+
+template <typename T, typename>
+bool UE::Landscape::EditLayers::FMergeRenderContext::HasBlackboardItem() const
+{
+	for (const TInstancedStruct<FLandscapeEditLayerRenderBlackboardItem>& Item : BlackboardItems)
+	{
+		if (const T* TypedItem = Item.template GetPtr<T>())
+		{
+			return true;
+		}
+	}
+	return false;
+}
+
+template <typename T, typename... TArgs, typename>
+T& UE::Landscape::EditLayers::FMergeRenderContext::AddBlackboardItem(TArgs&&... InArgs)
+{
+	T* Item = BlackboardItems.Add_GetRef(TInstancedStruct<T>::Make(Forward<TArgs>(InArgs)...)).template GetMutablePtr<T>();
+	return *Item;
+}
+
+template <typename T, typename>
+T* UE::Landscape::EditLayers::FMergeRenderContext::TryGetBlackboardItem()
+{
+	for (TInstancedStruct<FLandscapeEditLayerRenderBlackboardItem>& Item : BlackboardItems)
+	{
+		if (T* TypedItem = Item.template GetMutablePtr<T>())
+		{
+			return TypedItem;
+		}
+	}
+	return nullptr;
+}
+
+template <typename T, typename>
+T& UE::Landscape::EditLayers::FMergeRenderContext::GetBlackboardItem()
+{
+	for (TInstancedStruct<FLandscapeEditLayerRenderBlackboardItem>& Item : BlackboardItems)
+	{
+		if (T* TypedItem = Item.template GetMutablePtr<T>())
+		{
+			return *TypedItem;
+		}
+	}
+	check(false);
+
+	// this is absolutely wrong, but we shouldn't ever reach that line and we need to return something anyway : 
+	static FLandscapeEditLayerRenderBlackboardItem DummyItem;
+	return *reinterpret_cast<T*>(&DummyItem);
+}
+
+template <typename T, typename... TArgs, typename>
+T& UE::Landscape::EditLayers::FMergeRenderContext::GetOrCreateBlackboardItem(TArgs&&... InArgs)
+{
+	T* Item = TryGetBlackboardItem<T>();
+	if (Item == nullptr)
+	{
+		Item = BlackboardItems.Add_GetRef(TInstancedStruct<T>::Make(Forward<TArgs>(InArgs)...)).template GetMutablePtr<T>();
+	}
+	check(Item != nullptr);
+	return *Item;
+}
+
+template <typename T, typename>
+TArray<T*> UE::Landscape::EditLayers::FMergeRenderContext::GetBlackboardItems()
+{
+	TArray<T*> Items;
+	for (TInstancedStruct<FLandscapeEditLayerRenderBlackboardItem>& Item : BlackboardItems)
+	{
+		if (T* TypedItem = Item.template GetMutablePtr<T>())
+		{
+			Items.Add(TypedItem);
+		}
+	}
+	return Items;
+}
+
+#endif // WITH_EDITOR
+
+} //namespace UE::Landscape::EditLayers
\ No newline at end of file
