Split landscape edit layer batched merge code into several smaller files

[RN]

diff --git a/Engine/Plugins/Experimental/LandscapePatch/Source/LandscapePatch/Private/LandscapeCircleHeightPatch.cpp b/Engine/Plugins/Experimental/LandscapePatch/Source/LandscapePatch/Private/LandscapeCircleHeightPatch.cpp
--- a/Engine/Plugins/Experimental/LandscapePatch/Source/LandscapePatch/Private/LandscapeCircleHeightPatch.cpp
+++ b/Engine/Plugins/Experimental/LandscapePatch/Source/LandscapePatch/Private/LandscapeCircleHeightPatch.cpp
@@ -6,6 +6,7 @@
 #include "Engine/TextureRenderTarget2DArray.h"
 #include "Landscape.h"
 #include "LandscapeCircleHeightPatchPS.h"
+#include "LandscapeEditLayerMergeRenderContext.h"
 #include "LandscapeEditResourcesSubsystem.h" // ULandscapeScratchRenderTarget
 #include "LandscapeInfo.h"
 #include "LandscapePatchUtil.h" // GetHeightmapToWorld
diff --git a/Engine/Plugins/Experimental/LandscapePatch/Source/LandscapePatch/Private/LandscapePatchEditLayer.cpp b/Engine/Plugins/Experimental/LandscapePatch/Source/LandscapePatch/Private/LandscapePatchEditLayer.cpp
--- a/Engine/Plugins/Experimental/LandscapePatch/Source/LandscapePatch/Private/LandscapePatchEditLayer.cpp
+++ b/Engine/Plugins/Experimental/LandscapePatch/Source/LandscapePatch/Private/LandscapePatchEditLayer.cpp
@@ -5,6 +5,7 @@
 #include "Algo/AnyOf.h"
 #include "Algo/Sort.h"
 #include "Landscape.h"
+#include "LandscapeEditLayerMergeRenderContext.h"
 #include "LandscapePatchLogging.h"
 #include "LandscapePatchUtil.h"
 #include "CoreGlobals.h" // UE::GetIsEditorLoadingPackage()
diff --git a/Engine/Source/Runtime/Landscape/Classes/Landscape.h b/Engine/Source/Runtime/Landscape/Classes/Landscape.h
--- a/Engine/Source/Runtime/Landscape/Classes/Landscape.h
+++ b/Engine/Source/Runtime/Landscape/Classes/Landscape.h
@@ -8,7 +8,11 @@
 #include "LandscapeProxy.h"
 #include "LandscapeBlueprintBrushBase.h"
 #include "LandscapeEditTypes.h"
+
+#if UE_ENABLE_INCLUDE_ORDER_DEPRECATED_IN_5_6
 #include "LandscapeEditLayerRenderer.h"
+#endif // UE_ENABLE_INCLUDE_ORDER_DEPRECATED_IN_5_6
+
 #include "Delegates/DelegateCombinations.h"
 #include "Templates/SubclassOf.h"
 #include "Algo/Transform.h"
@@ -16,6 +20,10 @@
 
 #include "Landscape.generated.h"
 
+
+// ----------------------------------------------------------------------------------
+// Forward declarations
+
 class FTextureRenderTargetResource;
 class ULandscapeComponent;
 class ILandscapeEdModeInterface;
@@ -42,6 +50,11 @@ namespace EditLayersWeightmapLocalMerge_RenderThread
 	struct FMergeInfo;
 }
 
+namespace UE::Landscape::EditLayers
+{
+	struct FMergeRenderParams;
+}
+
 #if WITH_EDITOR
 extern LANDSCAPE_API TAutoConsoleVariable<int32> CVarLandscapeSplineFalloffModulation;
 #endif
diff --git a/Engine/Source/Runtime/Landscape/Private/LandscapeEditLayerRenderer.cpp b/Engine/Source/Runtime/Landscape/Private/LandscapeEditLayerRenderer.cpp
--- a/Engine/Source/Runtime/Landscape/Private/LandscapeEditLayerRenderer.cpp
+++ b/Engine/Source/Runtime/Landscape/Private/LandscapeEditLayerRenderer.cpp
@@ -1,1867 +1,29 @@
 // Copyright Epic Games, Inc. All Rights Reserved.
 
 #include "LandscapeEditLayerRenderer.h"
-#include "Algo/AllOf.h"
-#include "Algo/Transform.h"
-#include "Algo/Count.h"
-#include "Engine/Engine.h"
-#include "Engine/TextureRenderTarget2D.h"
-#include "Engine/TextureRenderTarget2DArray.h"
-#include "EngineModule.h"
-#include "Landscape.h"
-#include "LandscapeComponent.h"
-#include "LandscapeEditResourcesSubsystem.h"
+#include "LandscapeEditLayerMergeRenderContext.h"
 #include "LandscapeLayerInfoObject.h"
-#include "LandscapePrivate.h"
-#include "LandscapeUtils.h"
-#include "Materials/MaterialExpressionLandscapeVisibilityMask.h"
-#include "PixelShaderUtils.h"
-#include "RenderGraph.h"
-#include "RenderingThread.h"
-#include "RHIAccess.h"
-#include "SceneView.h"
-#include "TextureResource.h"
-#include "VisualLogger/VisualLogger.h"
-
-#define LOCTEXT_NAMESPACE "LandscapeEditLayerRenderer"
-
-extern TAutoConsoleVariable<float> CVarLandscapeBatchedMergeVisualLogOffsetIncrement;
-extern TAutoConsoleVariable<int32> CVarLandscapeEditLayersClearBeforeEachWriteToScratch;
-extern TAutoConsoleVariable<int32> CVarLandscapeBatchedMergeVisualLogShowMergeType;
-extern TAutoConsoleVariable<bool> CVarLandscapeBatchedMergeVisualLogShowMergeProcess;
-extern TAutoConsoleVariable<float> CVarLandscapeBatchedMergeVisualLogAlpha;
-
-namespace UE::Landscape::EditLayers
-{
-
-// ----------------------------------------------------------------------------------
-// /Engine/Private/Landscape/LandscapeEditLayersUtils.usf shaders : 
-
-class FMarkValidityPS : public FGlobalShader
-{
-	DECLARE_GLOBAL_SHADER(FMarkValidityPS);
-	SHADER_USE_PARAMETER_STRUCT(FMarkValidityPS, FGlobalShader);
-
-	BEGIN_SHADER_PARAMETER_STRUCT(FParameters, )
-	END_SHADER_PARAMETER_STRUCT()
-
-	using FPermutationDomain = TShaderPermutationDomain<>;
-
-	static bool ShouldCompilePermutation(const FGlobalShaderPermutationParameters& Parameters)
-	{
-		return UE::Landscape::DoesPlatformSupportEditLayers(Parameters.Platform);
-	}
-
-	static void ModifyCompilationEnvironment(const FGlobalShaderPermutationParameters& Parameters, FShaderCompilerEnvironment& OutEnvironment)
-	{
-		OutEnvironment.SetDefine(TEXT("MARK_VALIDITY"), 1);
-	}
-};
-
-IMPLEMENT_GLOBAL_SHADER(FMarkValidityPS, "/Engine/Private/Landscape/LandscapeEditLayersUtils.usf", "MarkValidityPS", SF_Pixel);
-
-BEGIN_SHADER_PARAMETER_STRUCT(FMarkValidityPSParameters, )
-	SHADER_PARAMETER_STRUCT_INCLUDE(FPixelShaderUtils::FRasterizeToRectsVS::FParameters, VS)
-	SHADER_PARAMETER_STRUCT_INCLUDE(FMarkValidityPS::FParameters, PS)
-	RENDER_TARGET_BINDING_SLOTS()
-END_SHADER_PARAMETER_STRUCT()
-
-
-// ----------------------------------------------------------------------------------
-
-class FCopyQuadsPS : public FGlobalShader
-{
-	DECLARE_GLOBAL_SHADER(FCopyQuadsPS);
-	SHADER_USE_PARAMETER_STRUCT(FCopyQuadsPS, FGlobalShader);
-
-	BEGIN_SHADER_PARAMETER_STRUCT(FParameters, )
-		SHADER_PARAMETER_STRUCT_REF(FViewUniformShaderParameters, View)
-		SHADER_PARAMETER_RDG_TEXTURE(Texture2D, InSourceTexture)
-	END_SHADER_PARAMETER_STRUCT()
-
-	using FPermutationDomain = TShaderPermutationDomain<>;
-
-	static bool ShouldCompilePermutation(const FGlobalShaderPermutationParameters& Parameters)
-	{
-		return UE::Landscape::DoesPlatformSupportEditLayers(Parameters.Platform);
-	}
-
-	static void ModifyCompilationEnvironment(const FGlobalShaderPermutationParameters& Parameters, FShaderCompilerEnvironment& OutEnvironment)
-	{
-		OutEnvironment.SetDefine(TEXT("COPY_QUADS"), 1);
-	}
-};
-
-IMPLEMENT_GLOBAL_SHADER(FCopyQuadsPS, "/Engine/Private/Landscape/LandscapeEditLayersUtils.usf", "CopyQuadsPS", SF_Pixel);
-
-BEGIN_SHADER_PARAMETER_STRUCT(FCopyQuadsPSParameters, )
-	SHADER_PARAMETER_STRUCT_INCLUDE(FPixelShaderUtils::FRasterizeToRectsVS::FParameters, VS)
-	SHADER_PARAMETER_STRUCT_INCLUDE(FCopyQuadsPS::FParameters, PS)
-	RENDER_TARGET_BINDING_SLOTS()
-END_SHADER_PARAMETER_STRUCT()
-
-
-// ----------------------------------------------------------------------------------
-// /Engine/Private/Landscape/LandscapeEditLayersHeightmaps.usf shaders :
-
-class FLandscapeEditLayersHeightmapsMergeEditLayerPS : public FGlobalShader
-{
-	DECLARE_GLOBAL_SHADER(FLandscapeEditLayersHeightmapsMergeEditLayerPS);
-	SHADER_USE_PARAMETER_STRUCT(FLandscapeEditLayersHeightmapsMergeEditLayerPS, FGlobalShader);
-
-public:
-	BEGIN_SHADER_PARAMETER_STRUCT(FParameters, )
-		SHADER_PARAMETER(uint32, InEditLayerBlendMode)
-		SHADER_PARAMETER(float, InEditLayerAlpha)
-		SHADER_PARAMETER_RDG_TEXTURE_SRV(Texture2D<float4>, InCurrentEditLayerHeightmap)
-		SHADER_PARAMETER_RDG_TEXTURE_SRV(Texture2D<float4>, InPreviousEditLayersHeightmap)
-		RENDER_TARGET_BINDING_SLOTS()
-	END_SHADER_PARAMETER_STRUCT()
-
-	static bool ShouldCompilePermutation(const FGlobalShaderPermutationParameters& InParameters)
-	{
-		return UE::Landscape::DoesPlatformSupportEditLayers(InParameters.Platform);
-	}
-
-	static void ModifyCompilationEnvironment(const FGlobalShaderPermutationParameters& InParameters, FShaderCompilerEnvironment& OutEnvironment)
-	{
-		OutEnvironment.SetDefine(TEXT("MERGE_EDIT_LAYER"), 1);
-	}
-
-	static void MergeEditLayerPS(FRDGEventName&& InRDGEventName, FRDGBuilder& GraphBuilder, FParameters* InParameters, const FIntPoint& InTextureSize)
-	{
-		FGlobalShaderMap* ShaderMap = GetGlobalShaderMap(GMaxRHIFeatureLevel);
-		TShaderMapRef<FLandscapeEditLayersHeightmapsMergeEditLayerPS> PixelShader(ShaderMap);
-
-		FPixelShaderUtils::AddFullscreenPass(
-			GraphBuilder,
-			ShaderMap,
-			MoveTemp(InRDGEventName),
-			PixelShader,
-			InParameters,
-			FIntRect(0, 0, InTextureSize.X, InTextureSize.Y),
-			TStaticBlendStateWriteMask<CW_RG>::GetRHI());
-	}
-};
-
-IMPLEMENT_GLOBAL_SHADER(FLandscapeEditLayersHeightmapsMergeEditLayerPS, "/Engine/Private/Landscape/LandscapeEditLayersHeightmaps.usf", "MergeEditLayerPS", SF_Pixel);
-
-
-// ----------------------------------------------------------------------------------
-// /Engine/Private/Landscape/LandscapeEditLayersWeightmaps.usf shaders :
-
-class FLandscapeEditLayersWeightmapsMergeEditLayerPS : public FGlobalShader
-{
-	DECLARE_GLOBAL_SHADER(FLandscapeEditLayersWeightmapsMergeEditLayerPS);
-	SHADER_USE_PARAMETER_STRUCT(FLandscapeEditLayersWeightmapsMergeEditLayerPS, FGlobalShader);
-
-public:
-	BEGIN_SHADER_PARAMETER_STRUCT(FParameters, )
-		SHADER_PARAMETER(uint32, InTargetLayerIndex)
-		SHADER_PARAMETER(uint32, InEditLayerTargetLayerBlendMode)
-		SHADER_PARAMETER(float, InEditLayerAlpha)
-		SHADER_PARAMETER_RDG_TEXTURE_SRV(Texture2DArray<float4>, InCurrentEditLayerWeightmaps)
-		SHADER_PARAMETER_RDG_TEXTURE_SRV(Texture2DArray<float4>, InPreviousEditLayersWeightmaps)
-		RENDER_TARGET_BINDING_SLOTS()
-	END_SHADER_PARAMETER_STRUCT()
-
-	static bool ShouldCompilePermutation(const FGlobalShaderPermutationParameters& InParameters)
-	{
-		return UE::Landscape::DoesPlatformSupportEditLayers(InParameters.Platform);
-	}
-
-	static void ModifyCompilationEnvironment(const FGlobalShaderPermutationParameters& InParameters, FShaderCompilerEnvironment& OutEnvironment)
-	{
-		OutEnvironment.SetDefine(TEXT("MERGE_EDIT_LAYER"), 1);
-	}
-
-	static void MergeEditLayerPS(FRDGEventName&& InRDGEventName, FRDGBuilder& GraphBuilder, FParameters* InParameters, const FIntPoint& InTextureSize)
-	{
-		FGlobalShaderMap* ShaderMap = GetGlobalShaderMap(GMaxRHIFeatureLevel);
-		TShaderMapRef<FLandscapeEditLayersWeightmapsMergeEditLayerPS> PixelShader(ShaderMap);
-
-		FPixelShaderUtils::AddFullscreenPass(
-			GraphBuilder,
-			ShaderMap,
-			MoveTemp(InRDGEventName),
-			PixelShader,
-			InParameters,
-			FIntRect(0, 0, InTextureSize.X, InTextureSize.Y),
-			TStaticBlendStateWriteMask<CW_RG>::GetRHI());
-	}
-};
-
-IMPLEMENT_GLOBAL_SHADER(FLandscapeEditLayersWeightmapsMergeEditLayerPS, "/Engine/Private/Landscape/LandscapeEditLayersWeightmaps.usf", "MergeEditLayerPS", SF_Pixel);
-
-
-// ----------------------------------------------------------------------------------
-
-FString ConvertTargetLayerNamesToString(const TArrayView<const FName>& InTargetLayerNames)
-{
-	TArray<FString> TargetLayerStrings;
-	Algo::Transform(InTargetLayerNames, TargetLayerStrings, [](FName InTargetLayerName) { return InTargetLayerName.ToString(); });
-	return *FString::Join(TargetLayerStrings, TEXT(","));
-}
-
-
-// ----------------------------------------------------------------------------------
 
 #if WITH_EDITOR
 
-FMergeContext::FMergeContext(ALandscape* InLandscape, bool bInIsHeightmapMerge, bool bInSkipProceduralRenderers)
-	: bIsHeightmapMerge(bInIsHeightmapMerge)
-	, bSkipProceduralRenderers(bInSkipProceduralRenderers)
-	, Landscape(InLandscape)
-	, LandscapeInfo(InLandscape->GetLandscapeInfo())
-{
-	check(LandscapeInfo != nullptr);
-
-	// Start by gathering all possible target layer names on this landscape: 
-	//  This list of all unique target layer names will help accelerate the gathering of output layers on each component (using bit arrays) as well as the target layers intersection tests:
-	if (bInIsHeightmapMerge)
-	{
-		// Only one target layer in the case of heightmap: 
-		const FName TargetLayerName = FName("Height");
-		AllTargetLayerNames = { TargetLayerName };
-		// And it's valid :
-		ValidTargetLayerBitIndices = TBitArray<>(true, 1);
-		VisibilityTargetLayerMask = TBitArray<>(false, 1);
-		AllWeightmapLayerInfos = { nullptr };
-	}
-	else
-	{
-		// Gather all target layer names and mark those that are valid layers :
-		for (const FLandscapeInfoLayerSettings& LayerSettings : LandscapeInfo->Layers)
-		{
-			check(!LayerSettings.LayerName.IsNone());
-			ULandscapeLayerInfoObject*& LayerInfo = AllWeightmapLayerInfos.Add_GetRef(LayerSettings.LayerInfoObj);
-			AllTargetLayerNames.Add(LayerSettings.LayerName);
-			ValidTargetLayerBitIndices.Add(LayerInfo != nullptr);
-		}
-
-		// Visibility is always a valid layer : 
-		VisibilityTargetLayerIndex = AllTargetLayerNames.Find(UMaterialExpressionLandscapeVisibilityMask::ParameterName);
-		if (VisibilityTargetLayerIndex == INDEX_NONE)
-		{
-			VisibilityTargetLayerIndex = AllTargetLayerNames.Num();
-			AllTargetLayerNames.Add(UMaterialExpressionLandscapeVisibilityMask::ParameterName);
-			check(!AllWeightmapLayerInfos.Contains(ALandscape::VisibilityLayer));
-			AllWeightmapLayerInfos.Add(ALandscape::VisibilityLayer);
-			check(ALandscape::VisibilityLayer != nullptr);
-			ValidTargetLayerBitIndices.Add(true);
-		}
-		check(VisibilityTargetLayerIndex != INDEX_NONE);
-
-		VisibilityTargetLayerMask = TBitArray<>(false, AllTargetLayerNames.Num());
-		VisibilityTargetLayerMask[VisibilityTargetLayerIndex] = true;
-	}
-
-	NegatedVisibilityTargetLayerMask = VisibilityTargetLayerMask;
-	NegatedVisibilityTargetLayerMask.BitwiseNOT();
-}
-
-TArray<FName> FMergeContext::GetValidTargetLayerNames() const
-{
-	return bIsHeightmapMerge ? AllTargetLayerNames : ConvertTargetLayerBitIndicesToNames(ValidTargetLayerBitIndices);
-}
-
-int32 FMergeContext::IsValidTargetLayerName(const FName& InName) const
-{
-	int32 Index = GetTargetLayerIndexForName(InName);
-	return (Index != INDEX_NONE) ? ValidTargetLayerBitIndices[Index] : false;
-}
-
-int32 FMergeContext::IsValidTargetLayerNameChecked(const FName& InName) const
-{
-	int32 Index = AllTargetLayerNames.Find(InName);
-	check(Index != INDEX_NONE);
-	return ValidTargetLayerBitIndices[Index];
-}
-
-int32 FMergeContext::IsValidTargetLayerIndex(int32 InIndex) const
-{
-	check(AllTargetLayerNames.IsValidIndex(InIndex));
-	return ValidTargetLayerBitIndices[InIndex];
-}
-
-int32 FMergeContext::GetTargetLayerIndexForName(const FName& InName) const
-{
-	return AllTargetLayerNames.Find(InName);
-}
-
-int32 FMergeContext::GetTargetLayerIndexForNameChecked(const FName& InName) const
-{
-	int32 Index = AllTargetLayerNames.Find(InName);
-	check(Index != INDEX_NONE);
-	return Index;
-}
-
-FName FMergeContext::GetTargetLayerNameForIndex(int32 InIndex) const
-{
-	return AllTargetLayerNames.IsValidIndex(InIndex) ? AllTargetLayerNames[InIndex] : NAME_None;
-}
-
-FName FMergeContext::GetTargetLayerNameForIndexChecked(int32 InIndex) const
-{
-	check(AllTargetLayerNames.IsValidIndex(InIndex));
-	return AllTargetLayerNames[InIndex];
-}
-
-int32 FMergeContext::GetTargetLayerIndexForLayerInfo(ULandscapeLayerInfoObject* InLayerInfo) const
-{
-	return AllWeightmapLayerInfos.Find(InLayerInfo);
-}
-
-int32 FMergeContext::GetTargetLayerIndexForLayerInfoChecked(ULandscapeLayerInfoObject* InLayerInfo) const
-{
-	int32 Index = AllWeightmapLayerInfos.Find(InLayerInfo);;
-	check(Index != INDEX_NONE);
-	return Index;
-}
-
-ULandscapeLayerInfoObject* FMergeContext::GetTargetLayerInfoForName(const FName& InName) const
-{
-	int32 Index = GetTargetLayerIndexForName(InName);
-	return (Index != INDEX_NONE) ? AllWeightmapLayerInfos[Index] : nullptr;
-}
-
-ULandscapeLayerInfoObject* FMergeContext::GetTargetLayerInfoForNameChecked(const FName& InName) const
-{
-	int32 Index = GetTargetLayerIndexForNameChecked(InName);
-	return AllWeightmapLayerInfos[Index];
-}
-
-ULandscapeLayerInfoObject* FMergeContext::GetTargetLayerInfoForIndex(int32 InIndex) const
-{
-	check(AllTargetLayerNames.IsValidIndex(InIndex));
-	return AllWeightmapLayerInfos[InIndex];
-}
-
-TBitArray<> FMergeContext::ConvertTargetLayerNamesToBitIndices(TConstArrayView<FName> InTargetLayerNames) const
-{
-	TBitArray<> Result(false, AllTargetLayerNames.Num());
-	for (FName Name : InTargetLayerNames)
-	{
-		if (int32 Index = GetTargetLayerIndexForName(Name); Index != INDEX_NONE)
-		{
-			Result[Index] = true;
-		}
-	}
-	return Result;
-}
-
-TBitArray<> FMergeContext::ConvertTargetLayerNamesToBitIndicesChecked(TConstArrayView<FName> InTargetLayerNames) const
-{
-	TBitArray<> Result(false, AllTargetLayerNames.Num());
-	for (FName Name : InTargetLayerNames)
-	{
-		int32 Index = GetTargetLayerIndexForNameChecked(Name);
-		Result[Index] = true;
-	}
-	return Result;
-}
-
-TArray<FName> FMergeContext::ConvertTargetLayerBitIndicesToNames(const TBitArray<>& InTargetLayerBitIndices) const
-{
-	const int32 NumNames = AllTargetLayerNames.Num();
-	check(InTargetLayerBitIndices.Num() == NumNames);
-	TArray<FName> Names;
-	Names.Reserve(NumNames);
-	for (TConstSetBitIterator It(InTargetLayerBitIndices); It; ++It)
-	{
-		Names.Add(AllTargetLayerNames[It.GetIndex()]);
-	}
-	return Names;
-}
-
-TArray<ULandscapeLayerInfoObject*> FMergeContext::ConvertTargetLayerBitIndicesToLayerInfos(const TBitArray<>& InTargetLayerBitIndices) const
-{
-	const int32 NumTargetLayerInfos = AllTargetLayerNames.Num();
-	check(InTargetLayerBitIndices.Num() == NumTargetLayerInfos);
-	TArray<ULandscapeLayerInfoObject*> LayerInfos;
-	LayerInfos.Reserve(NumTargetLayerInfos);
-	for (TConstSetBitIterator It(InTargetLayerBitIndices); It; ++It)
-	{
-		LayerInfos.Add(AllWeightmapLayerInfos[It.GetIndex()]);
-	}
-	return LayerInfos;
-}
-
-void FMergeContext::ForEachTargetLayer(const TBitArray<>& InTargetLayerBitIndices, TFunctionRef<bool(int32 /*InTargetLayerIndex*/, const FName& /*InTargetLayerName*/, ULandscapeLayerInfoObject* /*InWeightmapLayerInfo*/)> Fn) const
-{
-	check(InTargetLayerBitIndices.Num() == AllTargetLayerNames.Num());
-	for (TConstSetBitIterator It(InTargetLayerBitIndices); It; ++It)
-	{
-		const int32 TargetLayerIndex = It.GetIndex();
-		if (!AllTargetLayerNames.IsValidIndex(TargetLayerIndex))
-		{
-			return;
-		}
-
-		if (!Fn(TargetLayerIndex, AllTargetLayerNames[TargetLayerIndex], AllWeightmapLayerInfos[TargetLayerIndex]))
-		{
-			return;
-		}
-	}
-}
-
-void FMergeContext::ForEachValidTargetLayer(TFunctionRef<bool(int32 /*InTargetLayerIndex*/, const FName& /*InTargetLayerName*/, ULandscapeLayerInfoObject* /*InWeightmapLayerInfo*/)> Fn) const
-{
-	return ForEachTargetLayer(ValidTargetLayerBitIndices, Fn);
-}
-
-TBitArray<> FMergeContext::BuildTargetLayerBitIndices(bool bInBitValue) const
-{
-	return TBitArray<>(bInBitValue, AllTargetLayerNames.Num());
-}
-
-
-// ----------------------------------------------------------------------------------
-
-const FEditLayerTargetTypeState& FEditLayerTargetTypeState::GetDummyTargetTypeState()
-{
-	static FEditLayerTargetTypeState DummyTargetTypeState;
-	return DummyTargetTypeState;
-}
-
-FEditLayerTargetTypeState::FEditLayerTargetTypeState(const FMergeContext* InMergeContext)
-	: FEditLayerTargetTypeState(InMergeContext, ELandscapeToolTargetTypeFlags::None)
-{}
-
-
-FEditLayerTargetTypeState::FEditLayerTargetTypeState(const FMergeContext* InMergeContext, ELandscapeToolTargetTypeFlags InTargetTypeMask)
-	: FEditLayerTargetTypeState(InMergeContext, InTargetTypeMask, TBitArray<>())
-{}
-
-FEditLayerTargetTypeState::FEditLayerTargetTypeState(const FMergeContext* InMergeContext, ELandscapeToolTargetTypeFlags InTargetTypeMask, const TArrayView<const FName>& InSupportedWeightmaps, bool bInChecked)
-	: FEditLayerTargetTypeState(InMergeContext, InTargetTypeMask, 
-		bInChecked ? InMergeContext->ConvertTargetLayerNamesToBitIndicesChecked(InSupportedWeightmaps) : InMergeContext->ConvertTargetLayerNamesToBitIndices(InSupportedWeightmaps))
-{}
-
-FEditLayerTargetTypeState::FEditLayerTargetTypeState(const FMergeContext* InMergeContext, ELandscapeToolTargetTypeFlags InTargetTypeMask, const TBitArray<>& InSupportedWeightmapLayerIndices)
-	: MergeContext(InMergeContext)
-{
-	if (InSupportedWeightmapLayerIndices.IsEmpty())
-	{
-		// Even if all weightmaps are turned off in this constructor, make sure to build a bit array that is dealing with as many layers as the merge context : 
-		WeightmapTargetLayerBitIndices = MergeContext->BuildTargetLayerBitIndices(false);
-	}
-	else
-	{
-		checkf(InSupportedWeightmapLayerIndices.Num() == MergeContext->GetAllTargetLayerNames().Num(), 
-			TEXT("Make sure that the target type state is dealing with the same amount of target layers as the merge context"));
-		WeightmapTargetLayerBitIndices = InSupportedWeightmapLayerIndices;
-	}
-	SetTargetTypeMask(InTargetTypeMask);
-}
-
-bool FEditLayerTargetTypeState::IsActive(ELandscapeToolTargetType InTargetType, const FName& InWeightmapLayerName) const
-{
-	if (InWeightmapLayerName != NAME_None)
-	{
-		if (int32 TargetLayerIndex = MergeContext->GetTargetLayerIndexForName(InWeightmapLayerName); TargetLayerIndex != INDEX_NONE)
-		{
-			return IsActive(InTargetType, TargetLayerIndex);
-		}
-
-		return false;
-	}
-
-	return IsActive(InTargetType, INDEX_NONE);
-}
-
-bool FEditLayerTargetTypeState::IsActiveChecked(ELandscapeToolTargetType InTargetType, const FName& InWeightmapLayerName) const
-{
-	if (InWeightmapLayerName != NAME_None)
-	{
-		return IsActive(InTargetType, MergeContext->GetTargetLayerIndexForNameChecked(InWeightmapLayerName));
-	}
-
-	return IsActive(InTargetType, INDEX_NONE);
-}
-
-bool FEditLayerTargetTypeState::IsActive(ELandscapeToolTargetType InTargetType, int32 InWeightmapLayerIndex) const
-{
-	check((InWeightmapLayerIndex == INDEX_NONE) || MergeContext->IsValidTargetLayerIndex(InWeightmapLayerIndex));
-
-	if (EnumHasAnyFlags(TargetTypeMask, GetLandscapeToolTargetTypeAsFlags(InTargetType)))
-	{
-		if (InTargetType != ELandscapeToolTargetType::Heightmap)
-		{
-			return GetActiveWeightmapBitIndices()[InWeightmapLayerIndex];
-		}
-
-		return true;
-	}
-
-	return false;
-}
-
-TArray<FName> FEditLayerTargetTypeState::GetActiveWeightmaps() const
-{
-	return MergeContext->ConvertTargetLayerBitIndicesToNames(GetActiveWeightmapBitIndices());
-}
-
-TBitArray<> FEditLayerTargetTypeState::GetActiveWeightmapBitIndices() const
-{
-	TBitArray<> Result = WeightmapTargetLayerBitIndices;
-	if (!EnumHasAnyFlags(TargetTypeMask, ELandscapeToolTargetTypeFlags::Visibility))
-	{
-		Result.CombineWithBitwiseAND(MergeContext->GetNegatedVisibilityTargetLayerMask(), EBitwiseOperatorFlags::MinSize);
-	}
-	if (!EnumHasAnyFlags(TargetTypeMask, ELandscapeToolTargetTypeFlags::Weightmap))
-	{
-		Result.CombineWithBitwiseAND(MergeContext->GetVisibilityTargetLayerMask(), EBitwiseOperatorFlags::MinSize);
-	}
-	return Result;
-}
-
-void FEditLayerTargetTypeState::SetTargetTypeMask(ELandscapeToolTargetTypeFlags InTargetTypeMask)
-{
-	if (InTargetTypeMask != TargetTypeMask)
-	{
-		TargetTypeMask = InTargetTypeMask;
-
-		// Special case for the visibility weightmap, where we want to make sure the weightmap layer name is specified if visibility is supported (and vice versa) : 
-		const int32 VisibilityTargetLayerIndex = MergeContext->GetVisibilityTargetLayerIndex();
-		if (VisibilityTargetLayerIndex != INDEX_NONE)
-		{
-			if (EnumHasAnyFlags(TargetTypeMask, ELandscapeToolTargetTypeFlags::Visibility))
-			{
-				AddWeightmap(VisibilityTargetLayerIndex);
-			}
-			else if (WeightmapTargetLayerBitIndices[VisibilityTargetLayerIndex])
-			{
-				RemoveWeightmap(VisibilityTargetLayerIndex);
-			}
-		}
-	}
-}
-
-void FEditLayerTargetTypeState::AddTargetType(ELandscapeToolTargetType InTargetType)
-{
-	SetTargetTypeMask(TargetTypeMask | GetLandscapeToolTargetTypeAsFlags(InTargetType));
-}
-
-void FEditLayerTargetTypeState::AddTargetTypeMask(ELandscapeToolTargetTypeFlags InTargetTypeMask)
-{
-	SetTargetTypeMask(TargetTypeMask | InTargetTypeMask);
-}
-
-void FEditLayerTargetTypeState::RemoveTargetType(ELandscapeToolTargetType InTargetType)
-{
-	SetTargetTypeMask(TargetTypeMask & ~GetLandscapeToolTargetTypeAsFlags(InTargetType));
-}
-
-void FEditLayerTargetTypeState::RemoveTargetTypeMask(ELandscapeToolTargetTypeFlags InTargetTypeMask)
-{
-	SetTargetTypeMask(TargetTypeMask & ~InTargetTypeMask);
-}
-
-FEditLayerTargetTypeState FEditLayerTargetTypeState::Intersect(const FEditLayerTargetTypeState& InOther) const
-{
-	checkf(WeightmapTargetLayerBitIndices.Num() == InOther.WeightmapTargetLayerBitIndices.Num(), 
-		TEXT("It is assumed that the 2 target type states to intersect are from the same context and are therefore dealing with the same amount of target layers"));
-	return FEditLayerTargetTypeState(MergeContext, InOther.GetTargetTypeMask() & TargetTypeMask, TBitArray<>::BitwiseAND(InOther.WeightmapTargetLayerBitIndices, WeightmapTargetLayerBitIndices, EBitwiseOperatorFlags::MinSize));
-}
-
-void FEditLayerTargetTypeState::AddWeightmap(const FName& InWeightmapLayerName)
-{
-	if (int32 TargetLayerIndex = MergeContext->GetTargetLayerIndexForName(InWeightmapLayerName); TargetLayerIndex != INDEX_NONE)
-	{
-		AddWeightmap(TargetLayerIndex);
-	}
-}
-
-void FEditLayerTargetTypeState::AddWeightmapChecked(const FName& InWeightmapLayerName)
-{
-	AddWeightmap(MergeContext->GetTargetLayerIndexForNameChecked(InWeightmapLayerName));
-}
-
-void FEditLayerTargetTypeState::AddWeightmap(int32 InWeightmapLayerIndex)
-{
-	check(MergeContext->IsValidTargetLayerIndex(InWeightmapLayerIndex));
-
-	checkf((InWeightmapLayerIndex != MergeContext->GetVisibilityTargetLayerIndex()) || EnumHasAnyFlags(TargetTypeMask, ELandscapeToolTargetTypeFlags::Visibility), 
-		TEXT("The visibility layer may only be used for target type states that support visibility"));
-
-	checkf((InWeightmapLayerIndex == MergeContext->GetVisibilityTargetLayerIndex()) || EnumHasAnyFlags(TargetTypeMask, ELandscapeToolTargetTypeFlags::Weightmap), 
-		TEXT("A weightmap layer (here, %s) may only be used for target type states that support weightmaps"), *MergeContext->GetTargetLayerNameForIndexChecked(InWeightmapLayerIndex).ToString());
-
-	WeightmapTargetLayerBitIndices[InWeightmapLayerIndex] = true;
-}
-
-void FEditLayerTargetTypeState::RemoveWeightmap(const FName& InWeightmapLayerName)
-{
-	if (int32 TargetLayerIndex = MergeContext->GetTargetLayerIndexForName(InWeightmapLayerName); TargetLayerIndex != INDEX_NONE)
-	{
-		RemoveWeightmap(TargetLayerIndex);
-	}
-}
-
-void FEditLayerTargetTypeState::RemoveWeightmapChecked(const FName& InWeightmapLayerName)
-{
-	RemoveWeightmap(MergeContext->GetTargetLayerIndexForNameChecked(InWeightmapLayerName));
-}
-
-void FEditLayerTargetTypeState::RemoveWeightmap(int32 InWeightmapLayerIndex)
-{
-	check(MergeContext->IsValidTargetLayerIndex(InWeightmapLayerIndex));
-
-	checkf(!EnumHasAnyFlags(TargetTypeMask, ELandscapeToolTargetTypeFlags::Visibility) || (InWeightmapLayerIndex != MergeContext->GetVisibilityTargetLayerIndex()),
-		TEXT("Cannot remove weightmap %s from a target type state that supports visibility"), *MergeContext->GetTargetLayerNameForIndexChecked(InWeightmapLayerIndex).ToString());
-
-	WeightmapTargetLayerBitIndices[InWeightmapLayerIndex] = false;
-}
-
-bool FEditLayerTargetTypeState::operator==(const FEditLayerTargetTypeState& InOther) const
-{
-	return (TargetTypeMask == InOther.TargetTypeMask) 
-		&& (WeightmapTargetLayerBitIndices == InOther.WeightmapTargetLayerBitIndices);
-}
-
-FString FEditLayerTargetTypeState::ToString() const
-{
-	FString Result = FString::Printf(TEXT("Target types: %s"), *UE::Landscape::GetLandscapeToolTargetTypeFlagsAsString(TargetTypeMask));
-	if (EnumHasAnyFlags(TargetTypeMask, ELandscapeToolTargetTypeFlags::Weightmap | ELandscapeToolTargetTypeFlags::Visibility))
-	{
-		Result += FString::Printf(TEXT("\nWeightmaps: %s"), *ConvertTargetLayerNamesToString(MergeContext->ConvertTargetLayerBitIndicesToNames(WeightmapTargetLayerBitIndices)));
-	}
-	return Result;
-}
-
-
-// ----------------------------------------------------------------------------------
-
-const FEditLayerRendererState& FEditLayerRendererState::GetDummyRendererState()
-{
-	static FEditLayerRendererState DummyRendererState;
-	return DummyRendererState;
-}
-
-FEditLayerRendererState::FEditLayerRendererState(const FMergeContext* InMergeContext)
-	: MergeContext(InMergeContext)
-	, SupportedTargetTypeState(InMergeContext)
-	, EnabledTargetTypeState(InMergeContext)
-	, ActiveTargetTypeState(InMergeContext)
-{
-	UpdateActiveTargetTypeState();
-}
-
-FEditLayerRendererState::FEditLayerRendererState(const FMergeContext* InMergeContext, TScriptInterface<ILandscapeEditLayerRenderer> InRenderer)
-	: MergeContext(InMergeContext)
-	, Renderer(InRenderer)
-	, DebugName(InRenderer->GetEditLayerRendererDebugName())
-	, SupportedTargetTypeState(InMergeContext)
-	, EnabledTargetTypeState(InMergeContext)
-	, ActiveTargetTypeState(InMergeContext)
-{
-	Renderer->GetRendererStateInfo(InMergeContext, SupportedTargetTypeState, EnabledTargetTypeState, TargetLayerGroups);
-
-	// Make sure that each supported weightmap belongs to one target layer group and one only. For those that are in no target layer group, put them in their own group, that simply means this renderer 
-	//  can render them without requesting the presence of other target layers (e.g. no weight-blending)
-	InMergeContext->ForEachTargetLayer(SupportedTargetTypeState.GetActiveWeightmapBitIndices(), 
-		[InMergeContext, InRenderer, this](int32 InTargetLayerIndex, const FName& InTargetLayerName, ULandscapeLayerInfoObject* InWeightmapLayerInfo)
-		{
-			const int32 TargetLayerGroupCount = Algo::CountIf(TargetLayerGroups, [InTargetLayerIndex](const TBitArray<>& InTargetLayerGroup) { return InTargetLayerGroup[InTargetLayerIndex]; });
-			checkf(TargetLayerGroupCount < 2, TEXT("Target layer %s belongs to more than 1 target layer group in edit layer renderer %s. This is forbidden: in the end, it must belong to 1 and 1 only."),
-				*InTargetLayerName.ToString(), *InRenderer->GetEditLayerRendererDebugName());
-			if (TargetLayerGroupCount == 0)
-			{
-				TBitArray<> TargetLayerGroup = InMergeContext->BuildTargetLayerBitIndices(false);
-				TargetLayerGroup[InTargetLayerIndex] = true;
-				TargetLayerGroups.Add(TargetLayerGroup);
-			}
-			return true;
-		});
-
-	UpdateActiveTargetTypeState();
-}
-
-void FEditLayerRendererState::EnableTargetType(ELandscapeToolTargetType InTargetType)
-{
-	checkf(!EnumHasAnyFlags(SupportedTargetTypeState.GetTargetTypeMask(), GetLandscapeToolTargetTypeAsFlags(InTargetType)), 
-		TEXT("Target type %s cannot be enabled on this renderer state because it is not supported. Make sure that target types are supported before enabling them"), *UEnum::GetValueAsString(InTargetType));
-	EnabledTargetTypeState.AddTargetType(InTargetType);
-	UpdateActiveTargetTypeState();
-}
-
-void FEditLayerRendererState::EnableTargetTypeMask(ELandscapeToolTargetTypeFlags InTargetTypeMask)
-{
-	for (ELandscapeToolTargetTypeFlags TargetTypeFlag : MakeFlagsRange(InTargetTypeMask))
-	{
-		EnableTargetType(UE::Landscape::GetLandscapeToolTargetTypeSingleFlagAsType(TargetTypeFlag));
-	}
-}
-
-void FEditLayerRendererState::DisableTargetType(ELandscapeToolTargetType InTargetType)
-{
-	EnabledTargetTypeState.RemoveTargetType(InTargetType);
-	UpdateActiveTargetTypeState();
-}
-
-void FEditLayerRendererState::DisableTargetTypeMask(ELandscapeToolTargetTypeFlags InTargetTypeMask)
-{
-	EnabledTargetTypeState.RemoveTargetTypeMask(InTargetTypeMask);
-	UpdateActiveTargetTypeState();
-}
-
-void FEditLayerRendererState::EnableWeightmap(ELandscapeToolTargetType InTargetType, const FName& InWeightmapLayerName)
-{
-	EnableTargetType(InTargetType);
-	if (int32 TargetLayerIndex = MergeContext->GetTargetLayerIndexForName(InWeightmapLayerName); TargetLayerIndex != INDEX_NONE)
-	{
-		EnabledTargetTypeState.AddWeightmap(TargetLayerIndex);
-		UpdateActiveTargetTypeState();
-	}
-}
-
-void FEditLayerRendererState::EnableWeightmapChecked(ELandscapeToolTargetType InTargetType, const FName& InWeightmapLayerName)
-{
-	EnableWeightmap(InTargetType, (InWeightmapLayerName != NAME_None) ? MergeContext->GetTargetLayerIndexForNameChecked(InWeightmapLayerName) : INDEX_NONE);
-}
-
-void FEditLayerRendererState::EnableWeightmap(ELandscapeToolTargetType InTargetType, int32 InWeightmapLayerIndex)
-{
-	EnableTargetType(InTargetType);
-	EnabledTargetTypeState.AddWeightmap(InWeightmapLayerIndex);
-	UpdateActiveTargetTypeState();
-}
-
-void FEditLayerRendererState::DisableWeightmap(const FName& InWeightmapLayerName)
-{
-	if (int32 TargetLayerIndex = MergeContext->GetTargetLayerIndexForName(InWeightmapLayerName); TargetLayerIndex != INDEX_NONE)
-	{
-		DisableWeightmap(TargetLayerIndex);
-	}
-}
-
-void FEditLayerRendererState::DisableWeightmapChecked(const FName& InWeightmapLayerName)
-{
-	DisableWeightmap(MergeContext->GetTargetLayerIndexForNameChecked(InWeightmapLayerName));
-}
-
-void FEditLayerRendererState::DisableWeightmap(int32 InWeightmapLayerIndex)
-{
-	EnabledTargetTypeState.RemoveWeightmap(InWeightmapLayerIndex);
-	UpdateActiveTargetTypeState();
-}
-
-ELandscapeToolTargetTypeFlags FEditLayerRendererState::GetActiveTargetTypeMask() const
-{
-	return ActiveTargetTypeState.GetTargetTypeMask();
-}
-
-bool FEditLayerRendererState::IsTargetActive(ELandscapeToolTargetType InTargetType, const FName& InWeightmapLayerName) const
-{
-	return ActiveTargetTypeState.IsActive(InTargetType, InWeightmapLayerName);
-}
-
-bool FEditLayerRendererState::IsTargetActiveChecked(ELandscapeToolTargetType InTargetType, const FName& InWeightmapLayerName) const
-{
-	return ActiveTargetTypeState.IsActiveChecked(InTargetType, InWeightmapLayerName);
-}
-
-bool FEditLayerRendererState::IsTargetActive(ELandscapeToolTargetType InTargetType, int32 InWeightmapLayerIndex) const
-{
-	return ActiveTargetTypeState.IsActive(InTargetType, InWeightmapLayerIndex);
-}
-
-TArray<FName> FEditLayerRendererState::GetActiveTargetWeightmaps() const
-{
-	return ActiveTargetTypeState.GetActiveWeightmaps();
-}
-
-TBitArray<> FEditLayerRendererState::GetActiveTargetWeightmapBitIndices() const
-{
-	return ActiveTargetTypeState.GetActiveWeightmapBitIndices();
-}
-
-void FEditLayerRendererState::UpdateActiveTargetTypeState()
-{
-	ActiveTargetTypeState = SupportedTargetTypeState.Intersect(EnabledTargetTypeState);
-}
-
-
-// ----------------------------------------------------------------------------------
-
-bool FMergeRenderBatch::operator<(const FMergeRenderBatch& InOther) const
-{
-	// Sort by coordinates for making debugging more "logical" : 
-	if (MinComponentKey.Y < InOther.MinComponentKey.Y)
-	{
-		return true;
-	}
-	else if (MinComponentKey.Y == InOther.MinComponentKey.Y)
-	{
-		return (MinComponentKey.X < InOther.MinComponentKey.X);
-	}
-	return false;
-}
-
-int32 FMergeRenderBatch::ComputeSubsectionRects(ULandscapeComponent* InComponent, TArray<FIntRect, TInlineAllocator<4>>& OutSubsectionRects, TArray<FIntRect, TInlineAllocator<4>>& OutSubsectionRectsWithDuplicateBorders) const
-{
-	check(ComponentsToRender.Contains(InComponent));
-	const int32 NumSubsections = Landscape->NumSubsections;
-	const int32 ComponentSizeQuads = Landscape->ComponentSizeQuads;
-	const int32 SubsectionSizeQuads = Landscape->SubsectionSizeQuads;
-	const int32 SubsectionVerts = SubsectionSizeQuads + 1;
-	const int32 TotalNumSubsections = NumSubsections * NumSubsections;
-	OutSubsectionRects.Reserve(TotalNumSubsections);
-	OutSubsectionRectsWithDuplicateBorders.Reserve(TotalNumSubsections);
-
-	const FIntPoint ComponentSectionBase = InComponent->GetSectionBase();
-	checkf((ComponentSectionBase.X >= SectionRect.Min.X) && (ComponentSectionBase.Y >= SectionRect.Min.Y)
-		&& ((ComponentSectionBase.X + ComponentSizeQuads + 1) <= SectionRect.Max.X) && ((ComponentSectionBase.Y + ComponentSizeQuads + 1) <= SectionRect.Max.Y), 
-		TEXT("The requested component is not included in the render batch"));
-
-	const FIntPoint ComponentLocalKey = (ComponentSectionBase - SectionRect.Min) / ComponentSizeQuads;
-	for (int32 SubY = 0; SubY < NumSubsections; ++SubY)
-	{
-		for (int32 SubX = 0; SubX < NumSubsections; ++SubX)
-		{
-			{
-				FIntPoint SubSectionMin = ComponentSectionBase - SectionRect.Min + FIntPoint(SubX * SubsectionSizeQuads, SubY * SubsectionSizeQuads);
-				FIntPoint SubSectionMax = SubSectionMin + FIntPoint(SubsectionVerts, SubsectionVerts);
-				OutSubsectionRects.Add(FIntRect(SubSectionMin, SubSectionMax));
-			}
-			{
-				FIntPoint SubSectionMin = (ComponentLocalKey * NumSubsections + FIntPoint(SubX, SubY)) * SubsectionVerts;
-				FIntPoint SubSectionMax = SubSectionMin + SubsectionVerts;
-				OutSubsectionRectsWithDuplicateBorders.Add(FIntRect(SubSectionMin, SubSectionMax));
-			}
-		}
-	}
-
-	return TotalNumSubsections;
-}
-
-FIntRect FMergeRenderBatch::ComputeSectionRect(ULandscapeComponent* InComponent, bool bInWithDuplicateBorders) const
-{
-	check(ComponentsToRender.Contains(InComponent));
-
-	const FIntPoint ComponentSectionBase = InComponent->GetSectionBase();
-	checkf((ComponentSectionBase.X >= SectionRect.Min.X) && (ComponentSectionBase.Y >= SectionRect.Min.Y)
-		&& ((ComponentSectionBase.X + InComponent->ComponentSizeQuads + 1) <= SectionRect.Max.X) && ((ComponentSectionBase.Y + InComponent->ComponentSizeQuads + 1) <= SectionRect.Max.Y),
-		TEXT("The requested component is not included in the render batch"));
-
-	const FIntPoint ComponentLocalKey = (ComponentSectionBase - SectionRect.Min) / InComponent->ComponentSizeQuads;
-	const int32 ComponentSubsectionVerts = InComponent->SubsectionSizeQuads + 1;
-	
-	const int32 ComponentSize = InComponent->NumSubsections * (bInWithDuplicateBorders ? ComponentSubsectionVerts : InComponent->SubsectionSizeQuads);
-	FIntPoint SectionMin = ComponentLocalKey * ComponentSize;
-	FIntPoint SectionMax = SectionMin + ComponentSize;
-
-	return FIntRect(SectionMin, SectionMax);
-}
-
-void FMergeRenderBatch::ComputeAllSubsectionRects(TArray<FIntRect>& OutSubsectionRects, TArray<FIntRect>& OutSubsectionRectsWithDuplicateBorders) const
-{
-	const int32 NumSubsections = Landscape->NumSubsections;
-	const int32 ComponentSizeQuads = Landscape->ComponentSizeQuads;
-	const int32 SubsectionSizeQuads = Landscape->SubsectionSizeQuads;
-	const int32 SubsectionVerts = SubsectionSizeQuads + 1;
-	const int32 TotalNumSubsectionRects = ComponentsToRender.Num() * NumSubsections * NumSubsections;
-	OutSubsectionRects.Reserve(TotalNumSubsectionRects);
-	OutSubsectionRectsWithDuplicateBorders.Reserve(TotalNumSubsectionRects);
-
-	for (ULandscapeComponent* Component : ComponentsToRender)
-	{
-		const FIntPoint ComponentSectionBase = Component->GetSectionBase();
-		checkf((ComponentSectionBase.X >= SectionRect.Min.X) && (ComponentSectionBase.Y >= SectionRect.Min.Y)
-			&& ((ComponentSectionBase.X + ComponentSizeQuads + 1) <= SectionRect.Max.X) && ((ComponentSectionBase.Y + ComponentSizeQuads + 1) <= SectionRect.Max.Y),
-			TEXT("The requested component is not included in the render batch"));
-
-		const FIntPoint ComponentLocalKey = (ComponentSectionBase - SectionRect.Min) / ComponentSizeQuads;
-		TArray<FIntRect, TInlineAllocator<4>> SubSectionRects;
-		for (int32 SubY = 0; SubY < NumSubsections; ++SubY)
-		{
-			for (int32 SubX = 0; SubX < NumSubsections; ++SubX)
-			{
-				{
-					FIntPoint SubSectionMin = ComponentSectionBase - SectionRect.Min + FIntPoint(SubX * SubsectionSizeQuads, SubY * SubsectionSizeQuads);
-					FIntPoint SubSectionMax = SubSectionMin + FIntPoint(SubsectionVerts, SubsectionVerts);
-					OutSubsectionRects.Add(FIntRect(SubSectionMin, SubSectionMax));
-				}
-				{
-					FIntPoint SubSectionMin = (ComponentLocalKey * NumSubsections + FIntPoint(SubX, SubY)) * SubsectionVerts;
-					FIntPoint SubSectionMax = SubSectionMin + SubsectionVerts;
-					OutSubsectionRectsWithDuplicateBorders.Add(FIntRect(SubSectionMin, SubSectionMax));
-				}
-			}
-		}
-	}
-}
-
-FIntPoint FMergeRenderBatch::GetRenderTargetResolution(bool bInWithDuplicateBorders) const
-{
-	return bInWithDuplicateBorders ? Resolution : SectionRect.Size();
-}
-
-
-// ----------------------------------------------------------------------------------
-
-FMergeRenderContext::FMergeRenderContext(const FMergeContext& InMergeContext)
-	: FMergeContext(InMergeContext)
-{
-	for (ULandscapeScratchRenderTarget*& BlendRenderTarget : BlendRenderTargets)
-	{
-		BlendRenderTarget = nullptr;
-	}
-}
-
-FMergeRenderContext::~FMergeRenderContext()
-{
-	FreeResources();
-
-	checkf(Algo::AllOf(BlendRenderTargets, [](ULandscapeScratchRenderTarget* InRT) { return (InRT == nullptr); }), TEXT("Every scratch render target should have been freed at this point."));
-}
-
-void FMergeRenderContext::AllocateResources()
-{
-	TRACE_CPUPROFILER_EVENT_SCOPE(FMergeRenderContext::AllocateResources);
-
-	using namespace UE::Landscape;
-
-	// Prepare the transient render resources we'll need throughout the merge: 
-	const int32 NumSlices = IsHeightmapMerge() ? 0 : MaxNeededNumSlices;
-	FLinearColor RenderTargetClearColor(ForceInitToZero);
-	ETextureRenderTargetFormat RenderTargetFormat = ETextureRenderTargetFormat::RTF_R8;
-	if (IsHeightmapMerge())
-	{
-		// Convert the height value 0.0f to how it's stored in the texture : 
-		const uint16 HeightValue = LandscapeDataAccess::GetTexHeight(0.0f);
-		RenderTargetClearColor = FLinearColor((float)((HeightValue - (HeightValue & 255)) >> 8) / 255.0f, (float)(HeightValue & 255) / 255.0f, 0.0f, 0.0f);
-
-		RenderTargetFormat = ETextureRenderTargetFormat::RTF_RGBA8;
-	}
-	// When rendering weightmaps, we should have at least 1 slice (if == 1, we can use a UTextureRenderTarget2D, otherwise, we'll need to use a UTextureRenderTarget2DArray) : 
-	else
-	{
-		checkf(MaxNeededNumSlices > 0, TEXT("Weightmaps should have at least 1 slice"));
-		// We use extra channels for weightmaps for storing alpha / alpha flags :
-		RenderTargetFormat = ETextureRenderTargetFormat::RTF_RGBA8;
-	}
-
-	ULandscapeEditResourcesSubsystem* LandscapeEditResourcesSubsystem = GEngine->GetEngineSubsystem<ULandscapeEditResourcesSubsystem>();
-	check(LandscapeEditResourcesSubsystem != nullptr);
-	checkf(Algo::AllOf(BlendRenderTargets, [](ULandscapeScratchRenderTarget* InRT) { return (InRT == nullptr); }), TEXT("We shouldn't allocate without having freed first."));
-	check(CurrentBlendRenderTargetWriteIndex == -1);
-
-	// We need N render targets large enough to fit all batches : 
-	{
-		// Write : 
-		FScratchRenderTargetParams ScratchRenderTargetParams(TEXT("ScratchRT0"), /*bInExactDimensions = */false, /*bInUseUAV = */false, /*bInTargetArraySlicesIndependently = */(NumSlices > 0),
-			MaxNeededResolution, NumSlices, RenderTargetFormat, RenderTargetClearColor, ERHIAccess::RTV);
-		BlendRenderTargets[0] = LandscapeEditResourcesSubsystem->RequestScratchRenderTarget(ScratchRenderTargetParams);
-		// Read and ReadPrevious : 
-		ScratchRenderTargetParams.DebugName = TEXT("ScratchRT1");
-		ScratchRenderTargetParams.InitialState = ERHIAccess::SRVMask;
-		BlendRenderTargets[1] = LandscapeEditResourcesSubsystem->RequestScratchRenderTarget(ScratchRenderTargetParams);
-		ScratchRenderTargetParams.DebugName = TEXT("ScratchRT2");
-		BlendRenderTargets[2] = LandscapeEditResourcesSubsystem->RequestScratchRenderTarget(ScratchRenderTargetParams);
-	}
-}
-
-void FMergeRenderContext::FreeResources()
-{
-	TRACE_CPUPROFILER_EVENT_SCOPE(FMergeRenderContext::FreeResources);
-
-	ULandscapeEditResourcesSubsystem* LandscapeEditResourcesSubsystem = GEngine->GetEngineSubsystem<ULandscapeEditResourcesSubsystem>();
-	check(LandscapeEditResourcesSubsystem != nullptr);
-
-	// We can now return those scratch render targets to the pool : 
-	for (int32 Index = 0; Index < BlendRenderTargets.Num(); ++Index)
-	{
-		if (BlendRenderTargets[Index] != nullptr)
-		{
-			LandscapeEditResourcesSubsystem->ReleaseScratchRenderTarget(BlendRenderTargets[Index]);
-			BlendRenderTargets[Index] = nullptr;
-		}
-	}
-	
-	CurrentBlendRenderTargetWriteIndex = -1;
-}
-
-void FMergeRenderContext::AllocateBatchResources(const FMergeRenderBatch& InRenderBatch)
-{
-	TRACE_CPUPROFILER_EVENT_SCOPE(FMergeRenderContext::AllocateBatchResources);
-
-	using namespace UE::Landscape;
-
-	// Prepare the transient render resources we'll need for this batch:
-	ULandscapeEditResourcesSubsystem* LandscapeEditResourcesSubsystem = GEngine->GetEngineSubsystem<ULandscapeEditResourcesSubsystem>();
-	check(LandscapeEditResourcesSubsystem != nullptr);
-	check(PerTargetLayerValidityRenderTargets.IsEmpty());
-
-	// We need a RT version of the stencil buffer, one per target layer, to let users sample it as a standard texture :
-	int32 VisibilityScratchRTIndex = 0;
-	ForEachTargetLayer(InRenderBatch.TargetLayerBitIndices,
-		[this, LandscapeEditResourcesSubsystem, &VisibilityScratchRTIndex](int32 InTargetLayerIndex, const FName& InTargetLayerName, ULandscapeLayerInfoObject* InWeightmapLayerInfo)
-		{
-			FScratchRenderTargetParams ScratchRenderTargetParams(FString::Printf(TEXT("VisibilityScratchRT(%i)"), VisibilityScratchRTIndex), /*bInExactDimensions = */false, /*bInUseUAV = */false, /*bInTargetArraySlicesIndependently = */false,
-				MaxNeededResolution, 0, ETextureRenderTargetFormat::RTF_R8, FLinearColor::Black, ERHIAccess::RTV);
-			ULandscapeScratchRenderTarget* RenderTarget = LandscapeEditResourcesSubsystem->RequestScratchRenderTarget(ScratchRenderTargetParams);
-			PerTargetLayerValidityRenderTargets.FindOrAdd(InTargetLayerName) = RenderTarget;
-			++VisibilityScratchRTIndex;
-			return true;
-		});
-}
-
-void FMergeRenderContext::FreeBatchResources(const FMergeRenderBatch& InRenderBatch)
-{
-	TRACE_CPUPROFILER_EVENT_SCOPE(FMergeRenderContext::FreeBatchResources);
-
-	ULandscapeEditResourcesSubsystem* LandscapeEditResourcesSubsystem = GEngine->GetEngineSubsystem<ULandscapeEditResourcesSubsystem>();
-	check(LandscapeEditResourcesSubsystem != nullptr);
-
-	// We can now return those scratch render targets to the pool : 
-	for (auto ItLayerNameRenderTargetPair : PerTargetLayerValidityRenderTargets)
-	{
-		LandscapeEditResourcesSubsystem->ReleaseScratchRenderTarget(ItLayerNameRenderTargetPair.Value);
-	}
-	PerTargetLayerValidityRenderTargets.Empty();
-}
-
-void FMergeRenderContext::CycleBlendRenderTargets(UE::Landscape::FRDGBuilderRecorder& RDGBuilderRecorder)
-{
-	const bool bFirstWrite = (CurrentBlendRenderTargetWriteIndex < 0);
-	CurrentBlendRenderTargetWriteIndex = (CurrentBlendRenderTargetWriteIndex + 1) % BlendRenderTargets.Num();
-
-	if (!bFirstWrite)
-	{
-		// Optionally clear the write render target for debug purposes : 
-		if (CVarLandscapeEditLayersClearBeforeEachWriteToScratch.GetValueOnGameThread())
-		{
-			GetBlendRenderTargetWrite()->Clear(RDGBuilderRecorder);
-		}
-	}
-}
-
-ULandscapeScratchRenderTarget* EditLayers::FMergeRenderContext::GetBlendRenderTargetWrite() const
-{
-	checkf(CurrentBlendRenderTargetWriteIndex >= 0, TEXT("CycleBlendRenderTargets must be called at least once prior to accessing the blend render targets"));
-	return BlendRenderTargets[CurrentBlendRenderTargetWriteIndex % BlendRenderTargets.Num()];
-}
-
-ULandscapeScratchRenderTarget* FMergeRenderContext::GetBlendRenderTargetRead() const
-{
-	checkf(CurrentBlendRenderTargetWriteIndex >= 0, TEXT("CycleBlendRenderTargets must be called at least once prior to accessing the blend render targets")); 
-	return BlendRenderTargets[(CurrentBlendRenderTargetWriteIndex + BlendRenderTargets.Num() - 1) % BlendRenderTargets.Num()];
-}
-
-ULandscapeScratchRenderTarget* FMergeRenderContext::GetBlendRenderTargetReadPrevious() const
-{
-	checkf(CurrentBlendRenderTargetWriteIndex >= 0, TEXT("CycleBlendRenderTargets must be called at least once prior to accessing the blend render targets"));
-	return BlendRenderTargets[(CurrentBlendRenderTargetWriteIndex + BlendRenderTargets.Num() - 2) % BlendRenderTargets.Num()];
-}
-
-ULandscapeScratchRenderTarget* FMergeRenderContext::GetValidityRenderTarget(const FName& InTargetLayerName) const
-{
-	check(PerTargetLayerValidityRenderTargets.Contains(InTargetLayerName));
-	return PerTargetLayerValidityRenderTargets[InTargetLayerName];
-}
-
-FTransform FMergeRenderContext::ComputeVisualLogTransform(const FTransform& InTransform) const
-{
-	FTransform ZTransform(CurrentVisualLogOffset / InTransform.GetScale3D()); // the Offset is given in world space so unapply the scale before applying the transform
-	return ZTransform * InTransform;
-}
-
-void FMergeRenderContext::IncrementVisualLogOffset()
-{
-	double VisualLogOffsetIncrement = CVarLandscapeBatchedMergeVisualLogOffsetIncrement.GetValueOnGameThread();
-	CurrentVisualLogOffset.Z += VisualLogOffsetIncrement;
-}
-
-void FMergeRenderContext::ResetVisualLogOffset()
-{
-	CurrentVisualLogOffset = FVector(ForceInitToZero);
-}
-
-void FMergeRenderContext::RenderValidityRenderTargets(UE::Landscape::FRDGBuilderRecorder& RDGBuilderRecorder)
-{
-	using namespace UE::Landscape;
-
-	const FMergeRenderBatch* RenderBatch = GetCurrentRenderBatch();
-	check(RenderBatch != nullptr);
-
-	struct FTextureAndRects
-	{
-		FTextureAndRects(const FName& InTargetLayerName, const FString& InTextureDebugName, FTextureResource* InTextureResource)
-			: TargetLayerName(InTargetLayerName)
-			, TextureDebugName(InTextureDebugName)
-			, TextureResource(InTextureResource)
-		{}
-
-		FName TargetLayerName;
-		FString TextureDebugName;
-		FTextureResource* TextureResource = nullptr;
-		TArray<FUintVector4> Rects;
-	};
-
-	TArray<FTextureAndRects> TexturesAndRects;
-	TexturesAndRects.Reserve(RenderBatch->TargetLayerBitIndices.CountSetBits());
-	ForEachTargetLayer(RenderBatch->TargetLayerBitIndices,
-		[this, &RDGBuilderRecorder, &TexturesAndRects, RenderBatch](int32 InTargetLayerIndex, const FName& InTargetLayerName, ULandscapeLayerInfoObject* InWeightmapLayerInfo)
-		{
-			ULandscapeScratchRenderTarget* ScratchRenderTarget = PerTargetLayerValidityRenderTargets.FindChecked(InTargetLayerName);
-			check(ScratchRenderTarget != nullptr);
-
-			// Make sure the validity mask is entirely cleared first:
-			ScratchRenderTarget->Clear(RDGBuilderRecorder);
-
-			FTextureAndRects& TextureAndRects = TexturesAndRects.Emplace_GetRef(InTargetLayerName, ScratchRenderTarget->GetDebugName(), ScratchRenderTarget->GetRenderTarget2D()->GetResource());
-
-			// Then build a list of quads for marking where the components are valid for this target layer on this batch:
-			const TSet<ULandscapeComponent*>& Components = RenderBatch->TargetLayersToComponents[InTargetLayerIndex];
-			TextureAndRects.Rects.Reserve(Components.Num());
-			for (ULandscapeComponent* Component : Components)
-			{
-				FIntRect ComponentRect = RenderBatch->ComputeSectionRect(Component, /*bInWithDuplicateBorders = */false);
-				TextureAndRects.Rects.Add(FUintVector4(ComponentRect.Min.X, ComponentRect.Min.Y, ComponentRect.Max.X + 1, ComponentRect.Max.Y + 1));
-			}
-
-			ScratchRenderTarget->TransitionTo(ERHIAccess::RTV, RDGBuilderRecorder);
-			return true;
-		});
-
-	auto RDGCommand =
-		[TexturesAndRects = MoveTemp(TexturesAndRects)](FRDGBuilder& GraphBuilder)
-	{
-		for (const FTextureAndRects& TextureAndRects : TexturesAndRects)
-		{
-			FRDGBufferRef RectBuffer = CreateUploadBuffer(GraphBuilder, TEXT("MarkValidityRects"), MakeArrayView(TextureAndRects.Rects));
-			FRDGBufferSRVRef RectBufferSRV = GraphBuilder.CreateSRV(FRDGBufferSRVDesc(RectBuffer, PF_R32G32B32A32_UINT));
-			FRDGTextureRef OutputTexture = GraphBuilder.RegisterExternalTexture(CreateRenderTarget(TextureAndRects.TextureResource->GetTexture2DRHI(), TEXT("ValidityMask")));
-
-			FMarkValidityPSParameters* PassParameters = GraphBuilder.AllocParameters<FMarkValidityPSParameters>();
-			PassParameters->RenderTargets[0] = FRenderTargetBinding(OutputTexture, ERenderTargetLoadAction::ELoad);
-			FGlobalShaderMap* ShaderMap = GetGlobalShaderMap(GMaxRHIFeatureLevel);
-			TShaderRef<FMarkValidityPS> PixelShader = ShaderMap->GetShader<FMarkValidityPS>();
-
-			FPixelShaderUtils::AddRasterizeToRectsPass<FMarkValidityPS>(GraphBuilder,
-				ShaderMap,
-				RDG_EVENT_NAME("MarkValidity(%s) -> %s", *TextureAndRects.TargetLayerName.ToString(), *TextureAndRects.TextureDebugName),
-				PixelShader,
-				PassParameters,
-				/*ViewportSize = */OutputTexture->Desc.Extent,
-				RectBufferSRV,
-				TextureAndRects.Rects.Num(),
-				/*BlendState = */ nullptr,
-				/*RasterizerState = */ nullptr,
-				/*DepthStencilState = */ nullptr,
-				/*StencilRef = */ 0,
-				/*TextureSize = */ OutputTexture->Desc.Extent);
-		}
-		};
-
-	// We need to specify the final state of the external textures to prevent the graph builder from transitioning them to SRVMask :
-	TArray<FRDGBuilderRecorder::FRDGExternalTextureAccessFinal> RDGExternalTextureAccessFinalList;
-	Algo::Transform(TexturesAndRects, RDGExternalTextureAccessFinalList, [](const FTextureAndRects& InTextureAndRects) -> FRDGBuilderRecorder::FRDGExternalTextureAccessFinal
-		{ 
-			return { InTextureAndRects.TextureResource, ERHIAccess::RTV };
-	});
-	RDGBuilderRecorder.EnqueueRDGCommand(RDGCommand, RDGExternalTextureAccessFinalList);
-}
-
-void FMergeRenderContext::RenderExpandedRenderTarget(UE::Landscape::FRDGBuilderRecorder& RDGBuilderRecorder)
-{
-	using namespace UE::Landscape;
-
-	const FMergeRenderBatch* RenderBatch = GetCurrentRenderBatch();
-	check(RenderBatch != nullptr);
-
-	TArray<FUintVector4> SourceRects, DestinationRects;
-	{
-		TArray<FIntRect> SourceInclusiveRects, DestinationInclusiveRects;
-		RenderBatch->ComputeAllSubsectionRects(SourceInclusiveRects, DestinationInclusiveRects);
-		// ComputeAllSubsectionRects returns inclusive bounds while AddRasterizeToRectsPass requires exclusive bounds : 
-		Algo::Transform(SourceInclusiveRects, SourceRects, [](const FIntRect& InRect) { return FUintVector4(InRect.Min.X, InRect.Min.Y, InRect.Max.X + 1, InRect.Max.Y + 1); });
-		Algo::Transform(DestinationInclusiveRects, DestinationRects, [](const FIntRect& InRect) { return FUintVector4(InRect.Min.X, InRect.Min.Y, InRect.Max.X + 1, InRect.Max.Y + 1); });
-	}
-
-	ULandscapeScratchRenderTarget* WriteRT = GetBlendRenderTargetWrite();
-	ULandscapeScratchRenderTarget* ReadRT = GetBlendRenderTargetRead();
-	WriteRT->TransitionTo(ERHIAccess::RTV, RDGBuilderRecorder);
-	ReadRT->TransitionTo(ERHIAccess::SRVMask, RDGBuilderRecorder);
-
-	FSceneInterface* SceneInterface = GetLandscape()->GetWorld()->Scene;
-
-	auto RDGCommand =
-		[ SourceRects = MoveTemp(SourceRects)
-		, DestinationRects = MoveTemp(DestinationRects)
-		, OutputResource = WriteRT->GetRenderTarget2D()->GetResource()
-		, SourceResource = ReadRT->GetRenderTarget2D()->GetResource()
-		, SceneInterface] (FRDGBuilder& GraphBuilder)
-	{
-		FRDGBufferRef RectBuffer = CreateUploadBuffer(GraphBuilder, TEXT("ExpandRects"), MakeArrayView(DestinationRects));
-		FRDGBufferSRVRef RectBufferSRV = GraphBuilder.CreateSRV(FRDGBufferSRVDesc(RectBuffer, PF_R32G32B32A32_UINT));
-
-		FRDGBufferRef RectUVBuffer = CreateUploadBuffer(GraphBuilder, TEXT("ExpandRectsUVs"), MakeArrayView(SourceRects));
-		FRDGBufferSRVRef RectUVBufferSRV = GraphBuilder.CreateSRV(FRDGBufferSRVDesc(RectUVBuffer, PF_R32G32B32A32_UINT));
-
-		FRDGTextureRef OutputTexture = GraphBuilder.RegisterExternalTexture(CreateRenderTarget(OutputResource->GetTexture2DRHI(), TEXT("OutputTexture")));
-		FRDGTextureRef SourceTexture = GraphBuilder.RegisterExternalTexture(CreateRenderTarget(SourceResource->GetTexture2DRHI(), TEXT("SourceTexture")));
-
-		// TODO [jonathan.bard] this is just a rhi validation error for unoptimized shaders... once validation is made to not issue those errors, we can remove this
-		// Create a SceneView to please the shader bindings, but it's unused in practice 
-		FSceneViewFamilyContext ViewFamily(FSceneViewFamily::ConstructionValues(nullptr, SceneInterface, FEngineShowFlags(ESFIM_Game)).SetTime(FGameTime::GetTimeSinceAppStart()));
-		FSceneViewInitOptions ViewInitOptions;
-		ViewInitOptions.ViewFamily = &ViewFamily;
-		ViewInitOptions.SetViewRectangle(FIntRect(0, 0, 1, 1)); // Use a dummy rect to avoid a check(slow)
-			GetRendererModule().CreateAndInitSingleView(GraphBuilder.RHICmdList, &ViewFamily, &ViewInitOptions);
-		const FSceneView* View = ViewFamily.Views[0];
-
-		FCopyQuadsPSParameters* PassParameters = GraphBuilder.AllocParameters<FCopyQuadsPSParameters>();
-		PassParameters->RenderTargets[0] = FRenderTargetBinding(OutputTexture, ERenderTargetLoadAction::ELoad);
-		PassParameters->PS.View = View->ViewUniformBuffer;
-		PassParameters->PS.InSourceTexture = SourceTexture;
-
-		FGlobalShaderMap* ShaderMap = GetGlobalShaderMap(GMaxRHIFeatureLevel);
-		TShaderRef<FCopyQuadsPS> PixelShader = ShaderMap->GetShader<FCopyQuadsPS>();
-
-		FPixelShaderUtils::AddRasterizeToRectsPass<FCopyQuadsPS>(GraphBuilder,
-			ShaderMap,
-			RDG_EVENT_NAME("CopyQuadsPS"),
-			PixelShader,
-			PassParameters,
-			/*ViewportSize = */OutputTexture->Desc.Extent,
-			RectBufferSRV,
-			DestinationRects.Num(),
-			/*BlendState = */nullptr,
-			/*RasterizerState = */nullptr,
-			/*DepthStencilState = */nullptr,
-			/*StencilRef = */0,
-			/*TextureSize = */SourceTexture->Desc.Extent,
-			RectUVBufferSRV);
-		};
-
-	// We need to specify the final state of the external textures to prevent the graph builder from transitioning them to SRVMask (even those that end up as SRVMask at the end of this command,
-	//  because they will likely be part of another RDGCommand down the line so we need to maintain an accurate picture of every external texture ever involved in the recorded command so that 
-	//  we can set a proper access when the recorder is flushed (and the FRDGBuilder, executed) :
-	TArray<FRDGBuilderRecorder::FRDGExternalTextureAccessFinal> RDGExternalTextureAccessFinalList =
-	{
-		{ WriteRT->GetRenderTarget()->GetResource(), ERHIAccess::RTV },
-		{ ReadRT->GetRenderTarget()->GetResource(), ERHIAccess::SRVMask }
-	};
-	RDGBuilderRecorder.EnqueueRDGCommand(RDGCommand, RDGExternalTextureAccessFinalList);
-}
-
-void FMergeRenderContext::GenericBlendLayer(const FBlendParams& InBlendParams, UE::Landscape::FRDGBuilderRecorder& RDGBuilderRecorder)
-{
-	using namespace UE::Landscape;
-	using namespace UE::Landscape::EditLayers;
-
-	const FMergeRenderBatch* RenderBatch = GetCurrentRenderBatch();
-	check(RenderBatch != nullptr);
-
-	// In this step, we perform the blend : merge the edit layer with the previous layers in the stack
-	CycleBlendRenderTargets(RDGBuilderRecorder);
-	ULandscapeScratchRenderTarget* WriteRT = GetBlendRenderTargetWrite();
-	ULandscapeScratchRenderTarget* CurrentLayerReadRT = GetBlendRenderTargetRead();
-	ULandscapeScratchRenderTarget* PreviousLayersReadRT = GetBlendRenderTargetReadPrevious();
-
-	WriteRT->TransitionTo(ERHIAccess::RTV, RDGBuilderRecorder);
-	CurrentLayerReadRT->TransitionTo(ERHIAccess::SRVMask, RDGBuilderRecorder);
-	PreviousLayersReadRT->TransitionTo(ERHIAccess::SRVMask, RDGBuilderRecorder);
-
-	if (IsHeightmapMerge())
-	{
-		auto RDGCommand =
-			[ HeightmapBlendParams = InBlendParams.HeightmapBlendParams
-			, OutputResource = WriteRT->GetRenderTarget2D()->GetResource()
-			, OutputResourceName = WriteRT->GetDebugName()
-			, CurrentEditLayerResource = CurrentLayerReadRT->GetRenderTarget2D()->GetResource()
-			, PreviousEditLayersResource = PreviousLayersReadRT->GetRenderTarget2D()->GetResource()
-			, EffectiveTextureSize = RenderBatch->GetRenderTargetResolution(/*bInWithDuplicateBorders = */false)]
-			(FRDGBuilder& GraphBuilder)
-			{
-				FRDGTextureRef OutputTextureRef = GraphBuilder.RegisterExternalTexture(CreateRenderTarget(OutputResource->TextureRHI, TEXT("OutputTexture")));
-				FRDGTextureRef CurrentEditLayerTextureRef = GraphBuilder.RegisterExternalTexture(CreateRenderTarget(CurrentEditLayerResource->TextureRHI, TEXT("CurrentEditLayerTexture")));
-				FRDGTextureRef PreviousEditLayersTextureRef = GraphBuilder.RegisterExternalTexture(CreateRenderTarget(PreviousEditLayersResource->TextureRHI, TEXT("PreviousEditLayersTexture")));
-
-				FLandscapeEditLayersHeightmapsMergeEditLayerPS::FParameters* PSParams = GraphBuilder.AllocParameters<FLandscapeEditLayersHeightmapsMergeEditLayerPS::FParameters>();
-				PSParams->RenderTargets[0] = FRenderTargetBinding(OutputTextureRef, ERenderTargetLoadAction::ENoAction);
-				PSParams->InEditLayerBlendMode = static_cast<uint32>(HeightmapBlendParams.BlendMode);
-				PSParams->InEditLayerAlpha = HeightmapBlendParams.Alpha;
-				PSParams->InCurrentEditLayerHeightmap = GraphBuilder.CreateSRV(FRDGTextureSRVDesc::Create(CurrentEditLayerTextureRef));
-				PSParams->InPreviousEditLayersHeightmap = GraphBuilder.CreateSRV(FRDGTextureSRVDesc::Create(PreviousEditLayersTextureRef));
-
-				FLandscapeEditLayersHeightmapsMergeEditLayerPS::MergeEditLayerPS(
-					RDG_EVENT_NAME("MergeEditLayer(Height) -> %s", *OutputResourceName), GraphBuilder, PSParams, EffectiveTextureSize);
-			};
-
-
-		// We need to specify the final state of the external textures to prevent the graph builder from transitioning them to SRVMask (even those that end up as SRVMask at the end of this command,
-		//  because they will likely be part of another RDGCommand down the line so we need to maintain an accurate picture of every external texture ever involved in the recorded command so that 
-		//  we can set a proper access when the recorder is flushed (and the FRDGBuilder, executed) :
-		TArray<FRDGBuilderRecorder::FRDGExternalTextureAccessFinal> RDGExternalTextureAccessFinalList =
-		{
-			{ WriteRT->GetRenderTarget()->GetResource(), ERHIAccess::RTV },
-			{ CurrentLayerReadRT->GetRenderTarget()->GetResource(), ERHIAccess::SRVMask },
-			{ PreviousLayersReadRT->GetRenderTarget()->GetResource(), ERHIAccess::SRVMask },
-		};
-		RDGBuilderRecorder.EnqueueRDGCommand(RDGCommand, RDGExternalTextureAccessFinalList);
-	}
-	else
-	{
-		ForEachTargetLayer(RenderBatch->TargetLayerBitIndices,
-			[WriteRT, CurrentLayerReadRT, PreviousLayersReadRT, &RenderBatch, &InBlendParams, &RDGBuilderRecorder](int32 InTargetLayerIndex, const FName& InTargetLayerName, ULandscapeLayerInfoObject* InWeightmapLayerInfo)
-			{
-				// By default, use passthrough mode so that each layer gets at least copied into the WriteRT :
-				const FWeightmapBlendParams* TargetLayerBlendParams = &FWeightmapBlendParams::GetDefaultPassthroughBlendParams();
-				// Special case for visibility which is always "Additive" :
-				if (InTargetLayerName == UMaterialExpressionLandscapeVisibilityMask::ParameterName)
-				{
-					TargetLayerBlendParams = &FWeightmapBlendParams::GetDefaultAdditiveBlendParams();
-				}
-				else if (const FWeightmapBlendParams* FoundTargetLayerBlendParams = InBlendParams.WeightmapBlendParams.Find(InTargetLayerName))
-				{
-					TargetLayerBlendParams = FoundTargetLayerBlendParams;
-				}
-
-				// TODO [jonathan.bard] : we could render several layers at once via MRT (up to MaxSimultaneousRenderTargets)
-				auto RDGCommand =
-					[ TargetLayerBlendParams = *TargetLayerBlendParams
-					, InTargetLayerIndex
-					, InTargetLayerName
-					, OutputResource = WriteRT->GetRenderTarget2DArray()->GetResource()
-					, OutputResourceName = WriteRT->GetDebugName()
-					, CurrentEditLayerResource = CurrentLayerReadRT->GetRenderTarget2DArray()->GetResource()
-					, PreviousEditLayersResource = PreviousLayersReadRT->GetRenderTarget2DArray()->GetResource()
-					, EffectiveTextureSize = RenderBatch->GetRenderTargetResolution(/*bInWithDuplicateBorders = */false)]
-					(FRDGBuilder& GraphBuilder)
-					{
-						FRDGTextureRef OutputTextureRef = GraphBuilder.RegisterExternalTexture(CreateRenderTarget(OutputResource->TextureRHI, TEXT("OutputTexture")));
-						FRDGTextureRef CurrentEditLayerTextureRef = GraphBuilder.RegisterExternalTexture(CreateRenderTarget(CurrentEditLayerResource->TextureRHI, TEXT("CurrentEditLayerTexture")));
-						FRDGTextureSRVRef CurrentEditLayerTextureSRVRef = GraphBuilder.CreateSRV(FRDGTextureSRVDesc::Create(CurrentEditLayerTextureRef));
-						FRDGTextureRef PreviousEditLayersTextureRef = GraphBuilder.RegisterExternalTexture(CreateRenderTarget(PreviousEditLayersResource->TextureRHI, TEXT("PreviousEditLayersTexture")));
-						FRDGTextureSRVRef PreviousEditLayersTextureSRVRef = GraphBuilder.CreateSRV(FRDGTextureSRVDesc::Create(PreviousEditLayersTextureRef));
-
-						FLandscapeEditLayersWeightmapsMergeEditLayerPS::FParameters* PSParams = GraphBuilder.AllocParameters<FLandscapeEditLayersWeightmapsMergeEditLayerPS::FParameters>();
-						PSParams->RenderTargets[0] = FRenderTargetBinding(OutputTextureRef, ERenderTargetLoadAction::ENoAction, /*InMipIndex = */0, /*InArraySlice = */InTargetLayerIndex);
-						PSParams->InTargetLayerIndex = InTargetLayerIndex;
-						PSParams->InEditLayerTargetLayerBlendMode = static_cast<uint32>(TargetLayerBlendParams.BlendMode);
-						PSParams->InEditLayerAlpha = TargetLayerBlendParams.Alpha;
-						PSParams->InCurrentEditLayerWeightmaps = CurrentEditLayerTextureSRVRef;
-						PSParams->InPreviousEditLayersWeightmaps = PreviousEditLayersTextureSRVRef;
-
-						FLandscapeEditLayersWeightmapsMergeEditLayerPS::MergeEditLayerPS(
-							RDG_EVENT_NAME("MergeEditLayer(%s) -> %s", *InTargetLayerName.ToString(), *OutputResourceName), GraphBuilder, PSParams, EffectiveTextureSize);
-					};
-
-				// We need to specify the final state of the external textures to prevent the graph builder from transitioning them to SRVMask (even those that end up as SRVMask at the end of this command,
-				//  because they will likely be part of another RDGCommand down the line so we need to maintain an accurate picture of every external texture ever involved in the recorded command so that 
-				//  we can set a proper access when the recorder is flushed (and the FRDGBuilder, executed) :
-				TArray<FRDGBuilderRecorder::FRDGExternalTextureAccessFinal> RDGExternalTextureAccessFinalList =
-				{
-					{ WriteRT->GetRenderTarget()->GetResource(), ERHIAccess::RTV },
-					{ CurrentLayerReadRT->GetRenderTarget()->GetResource(), ERHIAccess::SRVMask },
-					{ PreviousLayersReadRT->GetRenderTarget()->GetResource(), ERHIAccess::SRVMask },
-				};
-				RDGBuilderRecorder.EnqueueRDGCommand(RDGCommand, RDGExternalTextureAccessFinalList);
-
-				return true;
-			});
-
-	}
-}
-
-void FMergeRenderContext::Render(TFunction<void(const FOnRenderBatchTargetGroupDoneParams& /*InParams*/, UE::Landscape::FRDGBuilderRecorder& /*RDGBuilderRecorder*/)> OnBatchTargetGroupDone)
-{
-	using namespace UE::Landscape;
-	using namespace UE::Landscape::EditLayers;
-
-	TRACE_CPUPROFILER_EVENT_SCOPE(FMergeRenderContext::Render);
-
-	check(CurrentRenderBatchIndex == INDEX_NONE);
-
-	const bool bVisualLogMergeProcess = CVarLandscapeBatchedMergeVisualLogShowMergeProcess.GetValueOnGameThread();
-
-	const FTransform& LandscapeTransform = Landscape->GetTransform();
-	// For visual logging, start at the top of the landscape's bounding box :
-	FVector LandscapeTopPosition(0.0, 0.0, MaxLocalHeight);
-	FTransform LandscapeWorldTransformForVisLog = FTransform(LandscapeTopPosition) * LandscapeTransform;
-
-	AllocateResources();
-
-	// Command recorder when accumulating render commands for several consecutive RenderLayer calls when ERenderFlags::RenderMode_Recorded is used 
-	FRDGBuilderRecorder RDGBuilderRecorder;
-
-	// Kick start the blend render targets : 
-	CycleBlendRenderTargets(RDGBuilderRecorder);
-
-	const int32 NumBatches = RenderBatches.Num();
-	for (CurrentRenderBatchIndex = 0; CurrentRenderBatchIndex < NumBatches; ++CurrentRenderBatchIndex)
-	{
-		const FMergeRenderBatch& RenderBatch = RenderBatches[CurrentRenderBatchIndex];
-		FString RenderBatchDebugName = FString::Format(TEXT("Render Batch [{0}] : ({1},{2})->({3},{4})"), { CurrentRenderBatchIndex, RenderBatch.SectionRect.Min.X, RenderBatch.SectionRect.Min.Y, RenderBatch.SectionRect.Max.X, RenderBatch.SectionRect.Max.Y });
-		RHI_BREADCRUMB_EVENT_GAMETHREAD("%s", RenderBatchDebugName);
-		TRACE_CPUPROFILER_EVENT_SCOPE(RenderBatch);
-
-		checkf((RenderBatch.RenderSteps.Num() >= 1) && (RenderBatch.RenderSteps.Last().Type == FMergeRenderStep::EType::SignalBatchMergeGroupDone),
-			TEXT("Any batch should end with a SignalBatchMergeGroupDone step and there should be at least another step prior to that, otherwise, the batch is just useless."));
-
-		AllocateBatchResources(RenderBatch);
-
-		IncrementVisualLogOffset();
-
-		// Drop a visual log showing the area covered by this batch : 
-		UE_IFVLOG(
-			if (IsVisualLogEnabled() && bVisualLogMergeProcess)
-			{
-				// Pick a new color for each batch : 
-				uint32 Hash = PointerHash(&RenderBatch);
-				uint8 * HashElement = reinterpret_cast<uint8*>(&Hash);
-				FColor Color(HashElement[0], HashElement[1], HashElement[2]);
-
-				UE_VLOG_OBOX(Landscape, LogLandscape, Log, FBox(FVector(RenderBatch.SectionRect.Min) - FVector(0.5, 0.5, 0.0), FVector(RenderBatch.SectionRect.Max) - FVector(0.5, 0.5, 0.0)),
-					ComputeVisualLogTransform(LandscapeWorldTransformForVisLog).ToMatrixWithScale(), Color.WithAlpha(GetVisualLogAlpha()),
-					TEXT("%s"), *FString::Format(TEXT("{0}\nBatch.SectionRect=([{1},{2}],[{3},{4}])"), { *RenderBatchDebugName, RenderBatch.SectionRect.Min.X, RenderBatch.SectionRect.Min.Y, RenderBatch.SectionRect.Max.X, RenderBatch.SectionRect.Max.Y }));
-			});
-
-		const int32 NumRenderSteps = RenderBatch.RenderSteps.Num();
-		// Current index of RenderLayer (for debugging purposes)
-		int32 RenderLayerStepIndex = 0;
-		// Index of RenderLayer at which we started recording the current render command sequence
-		int32 RenderCommandStartLayerStepIndex = INDEX_NONE;
-		// Indicates how many successful and consecutive RenderLayer steps (i.e. something has been rendered) have occurred (valid until the next BlendLayer step). 
-		//  It's useful for BlendLayer steps as it allows to skip the separate blend step if nothing was rendered prior to it :
-		int32 NumSuccessfulRenderLayerStepsUntilBlendLayerStep = 0;
-		for (int32 RenderStepIndex = 0; RenderStepIndex < NumRenderSteps; ++RenderStepIndex)
-		{
-			const FMergeRenderStep& RenderStep = RenderBatch.RenderSteps[RenderStepIndex];
-			TScriptInterface<ILandscapeEditLayerRenderer> Renderer = RenderStep.RendererState.GetRenderer();
-
-			TArray<FName> TargetLayerGroupLayerNames;
-			TArray<ULandscapeLayerInfoObject*> TargetLayerGroupLayerInfos;
-			TArray<FComponentMergeRenderInfo> SortedComponentMergeRenderInfos;
-
-			FTransform RenderAreaWorldTransform;
-			FIntRect RenderAreaSectionRect;
-
-			// Compute some data for the actual render steps : 
-			if ((RenderStep.Type == FMergeRenderStep::EType::RenderLayer)
-				|| (RenderStep.Type == FMergeRenderStep::EType::BlendLayer)
-				|| (RenderStep.Type == FMergeRenderStep::EType::SignalBatchMergeGroupDone))
-			{
-				TargetLayerGroupLayerNames = ConvertTargetLayerBitIndicesToNames(RenderStep.TargetLayerGroupBitIndices);
-				TargetLayerGroupLayerInfos = bIsHeightmapMerge ? TArray<ULandscapeLayerInfoObject*> { nullptr } : ConvertTargetLayerBitIndicesToLayerInfos(RenderStep.TargetLayerGroupBitIndices);
-
-			    // Compute all necessary info about the components affected by this renderer at this step
-			    SortedComponentMergeRenderInfos.Reserve(RenderStep.ComponentsToRender.Num());
-			    Algo::Transform(RenderStep.ComponentsToRender, SortedComponentMergeRenderInfos, [MinComponentKey = RenderBatch.MinComponentKey](ULandscapeComponent* InComponent)
-			    {
-				    FComponentMergeRenderInfo ComponentMergeRenderInfo;
-				    ComponentMergeRenderInfo.Component = InComponent;
-    
-				    const FIntPoint ComponentKey = InComponent->GetComponentKey();
-				    const FIntPoint LocalComponentKey = ComponentKey - MinComponentKey;
-				    check((LocalComponentKey.X >= 0) && (LocalComponentKey.Y >= 0));
-				    ComponentMergeRenderInfo.ComponentKeyInRenderArea = LocalComponentKey;
-				    // Area in the render target for this component : 
-				    ComponentMergeRenderInfo.ComponentRegionInRenderArea = FIntRect(LocalComponentKey * InComponent->ComponentSizeQuads, (LocalComponentKey + 1) * InComponent->ComponentSizeQuads);
-    
-				    return ComponentMergeRenderInfo;
-			    });
-			    SortedComponentMergeRenderInfos.Sort();
-			}
-
-			// Compute some additional data for the actual render steps : 
-			if ((RenderStep.Type == FMergeRenderStep::EType::BeginRenderLayerGroup)
-				|| (RenderStep.Type == FMergeRenderStep::EType::EndRenderLayerGroup)
-				|| (RenderStep.Type == FMergeRenderStep::EType::RenderLayer)
-				|| (RenderStep.Type == FMergeRenderStep::EType::BlendLayer)
-				|| (RenderStep.Type == FMergeRenderStep::EType::SignalBatchMergeGroupDone))
-			{
-				// TODO[jonathan.bard] offset the world transform to account for the half-pixel offset?
-				//RenderParams.RenderAreaWorldTransform = FTransform(LandscapeTransform.GetRotation(), LandscapeTransform.GetTranslation() + FVector(RenderBatch.SectionRect.Min) /** (double)ComponentSizeQuads - FVector(0.5, 0.5, 0)*/, LandscapeTransform.GetScale3D());
-
-				// TODO [jonathan.bard] : this is more of a Batch world transform/section rect at the moment. Shall we have a RenderAreaWorldTransform/RenderAreaSectionRect in FRenderParams and a BatchRenderAreaWorldTransform in FMergeRenderBatch?
-				//  because currently the old BP brushes work with FMergeRenderBatch data (i.e. 1 transform for the batch and a section rect for the entire batch) but eventually, renderers might be interested in just their Render step context, 
-				//  that is : 1 matrix corresponding to the bottom left corner of their list of components to render?
-				RenderAreaWorldTransform = FTransform(FVector(RenderBatch.SectionRect.Min)) * LandscapeTransform;
-				RenderAreaSectionRect = RenderBatch.SectionRect;
-			}
-
-			switch (RenderStep.Type)
-			{
-			case FMergeRenderStep::EType::BeginRenderCommandRecorder:
-			{
-				// Start recording a new sequence of RDG commands in order to let the upcoming FMergeRenderStep::EType::RenderLayer_Recorded steps push their render thread-based "operations" (lambdas)
-				checkf(RDGBuilderRecorder.IsEmpty() && !RDGBuilderRecorder.IsRecording(), TEXT("There shouldn't be any pending command being recorded when starting a new render command"));
-				checkf(RenderCommandStartLayerStepIndex == INDEX_NONE, TEXT("RenderCommandStartLayerStepIndex should be invalid as no render command recording should be active"));
-
-				// Remember the render layer step index at which we started recording the RDG render command : 
-				RenderCommandStartLayerStepIndex = RenderLayerStepIndex;
-				check(RenderCommandStartLayerStepIndex >= 0);
-				RDGBuilderRecorder.StartRecording();
-				break;
-			}
-			case FMergeRenderStep::EType::EndRenderCommandRecorder:
-			{
-				checkf(RDGBuilderRecorder.IsRecording(), TEXT("There should be a pending command being recorded when starting a ending a render command"));
-				checkf(RenderCommandStartLayerStepIndex != INDEX_NONE, TEXT("We should have initiated the render command recording with a begin operation, which should set a valid RenderCommandStartLayerStepIndex"));
-				check(RenderCommandStartLayerStepIndex <= RenderLayerStepIndex);
-
-				// This is where we actually push the current render command with all the render thread operations that have been accumulated on the render command context : 
-				FRDGEventName RDGEventName = (RenderCommandStartLayerStepIndex == RenderLayerStepIndex - 1) 
-					? RDG_EVENT_NAME("Step [%i]", RenderCommandStartLayerStepIndex)
-					: RDG_EVENT_NAME("Steps [%i-%i]", RenderCommandStartLayerStepIndex, RenderLayerStepIndex - 1);
-				RDGBuilderRecorder.StopRecordingAndFlush(MoveTemp(RDGEventName));
-
-				// We've flushed the render command, we can reset the starting render layer step index : 
-				RenderCommandStartLayerStepIndex = INDEX_NONE;
-				break;
-			}
-			case FMergeRenderStep::EType::BeginRenderLayerGroup:
-			case FMergeRenderStep::EType::EndRenderLayerGroup:
-			{
-				const ERenderFlags RenderFlags = Renderer->GetRenderFlags(this);
-				check((Renderer != nullptr) 
-					&& EnumHasAnyFlags(RenderFlags, ERenderFlags::RenderMode_Mask) 
-					&& EnumHasAllFlags(RenderFlags, ERenderFlags::BlendMode_SeparateBlend | ERenderFlags::RenderLayerGroup_SupportsGrouping));
-
-				if (RenderStep.Type == FMergeRenderStep::EType::BeginRenderLayerGroup)
-				{
-					// We start a new render layer group, so let's start tracking the number of successful RenderLayer steps within this group :
-					NumSuccessfulRenderLayerStepsUntilBlendLayerStep = 0;
-					ILandscapeEditLayerRenderer::FRenderParams RenderParams(this, TargetLayerGroupLayerNames, TargetLayerGroupLayerInfos, RenderStep.RendererState, SortedComponentMergeRenderInfos, 
-						RenderAreaWorldTransform, RenderAreaSectionRect, NumSuccessfulRenderLayerStepsUntilBlendLayerStep);
-					Renderer->BeginRenderLayerGroup(RenderParams, RDGBuilderRecorder);
-				}
-                else
-                {
-					ILandscapeEditLayerRenderer::FRenderParams RenderParams(this, TargetLayerGroupLayerNames, TargetLayerGroupLayerInfos, RenderStep.RendererState, SortedComponentMergeRenderInfos,
-						RenderAreaWorldTransform, RenderAreaSectionRect, NumSuccessfulRenderLayerStepsUntilBlendLayerStep);
-					Renderer->EndRenderLayerGroup(RenderParams, RDGBuilderRecorder);
-                }
-				break;
-			}
-			case FMergeRenderStep::EType::RenderLayer:
-			case FMergeRenderStep::EType::BlendLayer:
-			{
-				TRACE_CPUPROFILER_EVENT_SCOPE(RenderAndBlendLayer);
-
-				check((Renderer != nullptr) && EnumHasAnyFlags(Renderer->GetRenderFlags(this), ERenderFlags::RenderMode_Mask));
-
-				const bool bIsRecordedStep = EnumHasAnyFlags(RenderStep.RenderFlags, ERenderFlags::RenderMode_Recorded);
-				const bool bIsSeparateBlendStep = (RenderStep.Type == FMergeRenderStep::EType::BlendLayer);
-
-				FString RenderStepProfilingEventName = FString::Format(TEXT("Step [{0}] ({1}): {2} {3}"), { RenderLayerStepIndex, *ConvertTargetLayerNamesToString(TargetLayerGroupLayerNames), bIsSeparateBlendStep ? TEXT("Blend") : TEXT("Render"), *Renderer->GetEditLayerRendererDebugName() });
-
-				// Drop some visual cues to help understand how each renderer is applied :
-				UE_IFVLOG(
-				if (IsVisualLogEnabled() && !bIsSeparateBlendStep && bVisualLogMergeProcess)
-					{
-						FTransform RenderAreaWorldTransformForVisLog = FTransform(FVector(RenderBatch.SectionRect.Min)) * LandscapeWorldTransformForVisLog;
-						IncrementVisualLogOffset();
-						UE_VLOG_LOCATION(Landscape, LogLandscape, Log, ComputeVisualLogTransform(RenderAreaWorldTransformForVisLog).GetTranslation(), 10.0f, FColor::Red, TEXT("%s"), *RenderStepProfilingEventName);
-						UE_VLOG_WIREOBOX(Landscape, LogLandscape, Log, FBox(FVector(RenderBatch.SectionRect.Min) - FVector(0.5, 0.5, 0.0), FVector(RenderBatch.SectionRect.Max) - FVector(0.5, 0.5, 0.0)),
-							ComputeVisualLogTransform(LandscapeWorldTransformForVisLog).ToMatrixWithScale(), FColor::White, TEXT(""));
-
-						// Draw each component's bounds rendered by this renderer : 
-						for (const FComponentMergeRenderInfo& ComponentMergeRenderInfo : SortedComponentMergeRenderInfos)
-						{
-							UE_VLOG_WIREOBOX(Landscape, LogLandscape, Log, FBox(FVector(ComponentMergeRenderInfo.ComponentRegionInRenderArea.Min), FVector(ComponentMergeRenderInfo.ComponentRegionInRenderArea.Max)),
-								ComputeVisualLogTransform(RenderAreaWorldTransformForVisLog).ToMatrixWithScale(), FColor::White, TEXT(""));
-						}
-					});
-		
-				ILandscapeEditLayerRenderer::FRenderParams RenderParams(this, TargetLayerGroupLayerNames, TargetLayerGroupLayerInfos, RenderStep.RendererState, SortedComponentMergeRenderInfos, 
-					RenderAreaWorldTransform, RenderAreaSectionRect, NumSuccessfulRenderLayerStepsUntilBlendLayerStep);
-				auto RenderOrBlend = [&RDGBuilderRecorder, &RenderStepProfilingEventName, &NumSuccessfulRenderLayerStepsUntilBlendLayerStep, &RenderParams, bIsSeparateBlendStep, Renderer]()
-					{
-						if (bIsSeparateBlendStep)
-						{
-							// Skip the blend if nothing was ever rendered : 
-							if (NumSuccessfulRenderLayerStepsUntilBlendLayerStep > 0)
-							{
-								Renderer->BlendLayer(RenderParams, RDGBuilderRecorder);
-							}
-							// The blend has occurred, we can now stop tracking the number of successful RenderLayer steps :
-							NumSuccessfulRenderLayerStepsUntilBlendLayerStep = 0;
-						}
-						else
-						{
-							const bool bHasStepRenderedSomething = Renderer->RenderLayer(RenderParams, RDGBuilderRecorder);
-							NumSuccessfulRenderLayerStepsUntilBlendLayerStep += bHasStepRenderedSomething ? 1 : 0;
-						}
-					};
-    
-				    if (bIsRecordedStep)
-				    {
-					    checkf(RDGBuilderRecorder.IsRecording(), TEXT("(Render/Blend)Layer_Recorded must be preceded by a BeginRenderCommandRecorder which should create a command recorder"));
-					    RDG_RENDER_COMMAND_RECORDER_BREADCRUMB_EVENT(RDGBuilderRecorder, "%s", *RenderStepProfilingEventName);
-						RenderOrBlend();
-				    }
-				    else
-				    {
-					    checkf(RDGBuilderRecorder.IsEmpty() && !RDGBuilderRecorder.IsRecording(), TEXT("(Render/Blend)Layer_Immediate should be preceded by a EndRenderCommandRecorder which should finalize a command recorder and destroy it"));
-					        RHI_BREADCRUMB_EVENT_GAMETHREAD("%s", RenderStepProfilingEventName);
-					    RenderOrBlend();
-				    }
-    
-				    ++RenderLayerStepIndex;
-				    break;
-			    }
-			    case FMergeRenderStep::EType::SignalBatchMergeGroupDone:
-			    {
-				    TRACE_CPUPROFILER_EVENT_SCOPE(MergeGroupDone);
-				    RHI_BREADCRUMB_EVENT_GAMETHREAD("Step [%i] (%s) : Render Group Done", RenderLayerStepIndex, ConvertTargetLayerNamesToString(TargetLayerGroupLayerNames));
-    
-				    checkf(RDGBuilderRecorder.IsEmpty() && !RDGBuilderRecorder.IsRecording(), TEXT("SignalBatchMergeGroupDone should be preceded by a EndRenderCommandRecorder which should finalize a command recorder and destroy it"));
-
-				    // The last render target we wrote to is the one containing the batch group's merge result : 
-				    FOnRenderBatchTargetGroupDoneParams Params(this, TargetLayerGroupLayerNames, TargetLayerGroupLayerInfos, SortedComponentMergeRenderInfos);
-				    OnBatchTargetGroupDone(Params, RDGBuilderRecorder);
-				    break;
-			    }
-			    default:
-				    checkf(false, TEXT("Invalid render step type"));
-			}
-		}
-		
-		checkf(RDGBuilderRecorder.IsEmpty() && !RDGBuilderRecorder.IsRecording(), TEXT("We should not have any pending render command recorder at the end of render"));
-
-		FreeBatchResources(RenderBatch);
-	}
-
-	FreeResources();
-}
-
-const FMergeRenderBatch* FMergeRenderContext::GetCurrentRenderBatch() const
-{
-	return RenderBatches.IsValidIndex(CurrentRenderBatchIndex) ? &RenderBatches[CurrentRenderBatchIndex] : nullptr;
-}
-
-
-bool FMergeRenderContext::IsValid() const
-{
-	return !RenderBatches.IsEmpty();
-}
-
-#if ENABLE_VISUAL_LOG
-
-int32 FMergeRenderContext::GetVisualLogAlpha()
-{
-	return FMath::Clamp(CVarLandscapeBatchedMergeVisualLogAlpha.GetValueOnGameThread(), 0.0f, 1.0f) * 255;
-}
-
-bool FMergeRenderContext::IsVisualLogEnabled() const
-{
-	switch (CVarLandscapeBatchedMergeVisualLogShowMergeType.GetValueOnGameThread())
-	{
-	case 0: // Disabled
-		return false;
-	case 1: // Heightmaps only
-		return bIsHeightmapMerge;
-	case 2: // Weightmaps only
-		return !bIsHeightmapMerge;
-	case 3: // Both
-		return true;
-	default:
-		return false;
-	}
-}
-
-#endif // ENABLE_VISUAL_LOG
-
-
-// ----------------------------------------------------------------------------------
-
-FBox FOOBox2D::BuildAABB() const
-{
-	return FBox::BuildAABB(Transform.GetTranslation(), Transform.TransformVector(FVector(Extents, 0.0)));
-}
-
-
-// ----------------------------------------------------------------------------------
-
-FIntRect FInputWorldArea::GetLocalComponentKeys(const FIntPoint& InComponentKey) const
-{
-	check(Type == EType::LocalComponent); 
-	return LocalArea + InComponentKey;
-}
-
-FIntRect FInputWorldArea::GetSpecificComponentKeys() const
-{
-	check(Type == EType::SpecificComponent);
-	return LocalArea + SpecificComponentKey;
-}
-
-FBox FInputWorldArea::ComputeWorldAreaAABB(const FTransform& InLandscapeTransform, const FBox& InLandscapeLocalBounds, const FTransform& InComponentTransform, const FBox& InComponentLocalBounds) const
-{
-	switch (Type)
-	{
-		case EType::Infinite:
-		{
-			return InLandscapeLocalBounds.TransformBy(InLandscapeTransform);
-		}
-		case EType::LocalComponent:
-		{
-			return InComponentLocalBounds.TransformBy(InComponentTransform);
-		}
-		case EType::SpecificComponent:
-		{
-			FVector ComponentLocalSize = InComponentLocalBounds.GetSize();
-			FIntRect LocalAreaCoordinates(SpecificComponentKey + LocalArea.Min, SpecificComponentKey + LocalArea.Max);
-			FBox LocalAreaBounds = FBox(
-				FVector(LocalAreaCoordinates.Min.X * ComponentLocalSize.X, LocalAreaCoordinates.Min.Y * ComponentLocalSize.Y, InComponentLocalBounds.Min.Z),
-				FVector(LocalAreaCoordinates.Max.X * ComponentLocalSize.X, LocalAreaCoordinates.Max.Y * ComponentLocalSize.Y, InComponentLocalBounds.Max.Z));
-			return LocalAreaBounds.TransformBy(InComponentTransform);
-		}
-		case EType::OOBox:
-		{
-			return OOBox2D.BuildAABB();
-		}
-		default:
-			check(false);
-	}
-
-	return FBox();
-}
-
-FOOBox2D FInputWorldArea::ComputeWorldAreaOOBB(const FTransform& InLandscapeTransform, const FBox& InLandscapeLocalBounds, const FTransform& InComponentTransform, const FBox& InComponentLocalBounds) const
-{
-	switch (Type)
-	{
-	case EType::Infinite:
-	{
-		FVector Center, Extents;
-		InLandscapeLocalBounds.GetCenterAndExtents(Center, Extents);
-		FTransform LandscapeTransformAtCenter = InLandscapeTransform;
-		LandscapeTransformAtCenter.SetTranslation(InLandscapeTransform.TransformVector(Center));
-		return FOOBox2D(LandscapeTransformAtCenter, FVector2D(Extents));
-	}
-	case EType::LocalComponent:
-	{
-		FVector Center, Extents;
-		InComponentLocalBounds.GetCenterAndExtents(Center, Extents);
-		FTransform ComponentTransformAtCenter = InComponentTransform;
-		ComponentTransformAtCenter.SetTranslation(InComponentTransform.TransformVector(Center));
-		return FOOBox2D(ComponentTransformAtCenter, FVector2D(Extents));
-	}
-	case EType::SpecificComponent:
-	{
-		FVector ComponentLocalSize = InComponentLocalBounds.GetSize();
-		FIntRect LocalAreaCoordinates(SpecificComponentKey + LocalArea.Min, SpecificComponentKey + LocalArea.Max);
-		FBox LocalAreaBounds = FBox(
-			FVector(LocalAreaCoordinates.Min.X * ComponentLocalSize.X, LocalAreaCoordinates.Min.Y * ComponentLocalSize.Y, InComponentLocalBounds.Min.Z),
-			FVector(LocalAreaCoordinates.Max.X * ComponentLocalSize.X, LocalAreaCoordinates.Max.Y * ComponentLocalSize.Y, InComponentLocalBounds.Max.Z));
-		FVector Center, Extents;
-		LocalAreaBounds.GetCenterAndExtents(Center, Extents);
-		FTransform ComponentTransformAtCenter = InComponentTransform;
-		ComponentTransformAtCenter.SetTranslation(InComponentTransform.TransformVector(Center));
-		return FOOBox2D(ComponentTransformAtCenter, FVector2D(Extents));
-	}
-	case EType::OOBox:
-	{
-		return OOBox2D;
-	}
-	default:
-		check(false);
-	}
-
-	return FOOBox2D();
-}
-
-
-// ----------------------------------------------------------------------------------
-
-FBox FOutputWorldArea::ComputeWorldAreaAABB(const FTransform& InComponentTransform, const FBox& InComponentLocalBounds) const
-{
-	switch (Type)
-	{
-	case EType::LocalComponent:
-	{
-		return InComponentLocalBounds.TransformBy(InComponentTransform);
-	}
-	case EType::SpecificComponent:
-	{
-		FVector ComponentLocalSize = InComponentLocalBounds.GetSize();
-		FBox LocalAreaBounds = FBox(
-			FVector(SpecificComponentKey.X * ComponentLocalSize.X, SpecificComponentKey.Y * ComponentLocalSize.Y, InComponentLocalBounds.Min.Z),
-			FVector((SpecificComponentKey.X + 1) * ComponentLocalSize.X, (SpecificComponentKey.Y + 1) * ComponentLocalSize.Y, InComponentLocalBounds.Max.Z));
-		return LocalAreaBounds.TransformBy(InComponentTransform);
-	}
-	case EType::OOBox:
-	{
-		return OOBox2D.BuildAABB();
-	}
-	default:
-		check(false);
-	}
-
-	return FBox();
-}
-
-FOOBox2D FOutputWorldArea::ComputeWorldAreaOOBB(const FTransform& InComponentTransform, const FBox& InComponentLocalBounds) const
-{
-	switch (Type)
-	{
-	case EType::LocalComponent:
-	{
-		FVector Center, Extents;
-		InComponentLocalBounds.GetCenterAndExtents(Center, Extents);
-		FTransform ComponentTransformAtCenter = InComponentTransform;
-		ComponentTransformAtCenter.SetTranslation(InComponentTransform.TransformVector(Center));
-		return FOOBox2D(ComponentTransformAtCenter, FVector2D(Extents));
-	}
-	case EType::SpecificComponent:
-	{
-		FVector ComponentLocalSize = InComponentLocalBounds.GetSize();
-		FBox LocalAreaBounds = FBox(
-			FVector(SpecificComponentKey.X * ComponentLocalSize.X, SpecificComponentKey.Y * ComponentLocalSize.Y, InComponentLocalBounds.Min.Z),
-			FVector((SpecificComponentKey.X + 1) * ComponentLocalSize.X, (SpecificComponentKey.Y + 1) * ComponentLocalSize.Y, InComponentLocalBounds.Max.Z));
-		FVector Center, Extents;
-		LocalAreaBounds.GetCenterAndExtents(Center, Extents);
-		FTransform ComponentTransformAtCenter = InComponentTransform;
-		ComponentTransformAtCenter.SetTranslation(InComponentTransform.TransformVector(Center));
-		return FOOBox2D(ComponentTransformAtCenter, FVector2D(Extents));
-	}
-	case EType::OOBox:
-	{
-		return OOBox2D;
-	}
-	default:
-		check(false);
-	}
-
-	return FOOBox2D();
-}
-
-
-// ----------------------------------------------------------------------------------
-
-bool FComponentMergeRenderInfo::operator<(const FComponentMergeRenderInfo& InOther) const
-{
-	// Sort by X / Y so that the order in which we render them is always consistent : 
-	if (ComponentRegionInRenderArea.Min.Y < InOther.ComponentRegionInRenderArea.Min.Y)
-	{
-		return true;
-	}
-	else if (ComponentRegionInRenderArea.Min.Y == InOther.ComponentRegionInRenderArea.Min.Y)
-	{
-		return (ComponentRegionInRenderArea.Min.X < InOther.ComponentRegionInRenderArea.Min.X);
-	}
-	return false;
-}
-
-
-// ----------------------------------------------------------------------------------
-
-FMergeRenderParams::FMergeRenderParams(TArray<ULandscapeComponent*> InComponentsToMerge, const TArrayView<FEditLayerRendererState>& InEditLayerRendererStates, const TSet<FName>& InWeightmapLayerNames, bool bInRequestAllLayers)
-	: ComponentsToMerge(InComponentsToMerge)
-	, EditLayerRendererStates(InEditLayerRendererStates)
-	, WeightmapLayerNames(InWeightmapLayerNames)
-	, bRequestAllLayers(bInRequestAllLayers)
-{
-}
-
-#endif // WITH_EDITOR
-
-} // namespace UE::Landscape::EditLayers
-
-#undef LOCTEXT_NAMESPACE
+ILandscapeEditLayerRenderer::FRenderParams::FRenderParams(
+	UE::Landscape::EditLayers::FMergeRenderContext* InMergeRenderContext, 
+	const TArrayView<FName>& InTargetLayerGroupLayerNames, 
+	const TArrayView<ULandscapeLayerInfoObject*>& InTargetLayerGroupLayerInfos, 
+	const UE::Landscape::EditLayers::FEditLayerRendererState& InRendererState, 
+	const TArrayView<UE::Landscape::EditLayers::FComponentMergeRenderInfo>& InSortedComponentMergeRenderInfos, 
+	const FTransform& InRenderAreaWorldTransform, 
+	const FIntRect& InRenderAreaSectionRect, 
+	int32 InNumSuccessfulRenderLayerStepsUntilBlendLayerStep) 
+	: MergeRenderContext(InMergeRenderContext)
+	, TargetLayerGroupLayerNames(InTargetLayerGroupLayerNames)
+	, TargetLayerGroupLayerInfos(InTargetLayerGroupLayerInfos)
+	, RendererState(InRendererState)
+	, SortedComponentMergeRenderInfos(InSortedComponentMergeRenderInfos)
+	, RenderAreaWorldTransform(InRenderAreaWorldTransform)
+	, RenderAreaSectionRect(InRenderAreaSectionRect)
+	, NumSuccessfulRenderLayerStepsUntilBlendLayerStep(InNumSuccessfulRenderLayerStepsUntilBlendLayerStep)
+{
+}
+
+#endif // WITH_EDITOR
\ No newline at end of file
diff --git a/Engine/Source/Runtime/Landscape/Private/LandscapeEditLayerRendererPrivate.h b/Engine/Source/Runtime/Landscape/Private/LandscapeEditLayerRendererPrivate.h
--- a/Engine/Source/Runtime/Landscape/Private/LandscapeEditLayerRendererPrivate.h
+++ b/Engine/Source/Runtime/Landscape/Private/LandscapeEditLayerRendererPrivate.h
@@ -7,8 +7,6 @@
 
 #include "LandscapeEditLayerRendererPrivate.generated.h"
 
-class ILandscapeEditLayerRenderer;
-
 
 /**
 * Edit layer renderer added at the bottom of the stack to provide the default value for every requested target layer (heightmaps and weightmaps)
diff --git a/Engine/Source/Runtime/Landscape/Private/LandscapeEditLayers.cpp b/Engine/Source/Runtime/Landscape/Private/LandscapeEditLayers.cpp
--- a/Engine/Source/Runtime/Landscape/Private/LandscapeEditLayers.cpp
+++ b/Engine/Source/Runtime/Landscape/Private/LandscapeEditLayers.cpp
@@ -17,6 +17,8 @@ LandscapeEditLayers.cpp: Landscape editing layers mode
 #include "LandscapeEditReadback.h"
 #include "LandscapeEditResources.h"
 #include "LandscapeEditResourcesSubsystem.h"
+#include "LandscapeEditLayerRenderer.h"
+#include "LandscapeEditLayerMergeRenderContext.h"
 #include "LandscapeEditLayerRendererPrivate.h"
 #include "LandscapeEditTypes.h"
 #include "LandscapeUtils.h"
@@ -6209,7 +6211,7 @@ namespace UE::Landscape::EditLayers::Private
 					if (!bIsHeightmapMerge)
 					{
 						TArray<FName> TargetLayerNames = RenderContext->ConvertTargetLayerBitIndicesToNames(InComponentRenderInfo.TargetLayerBitIndices);
-						Message += FString::Printf(TEXT("\n%s"), *ConvertTargetLayerNamesToString(TargetLayerNames));
+						Message += FString::Printf(TEXT("\n%s"), *UE::Landscape::ConvertTargetLayerNamesToString(TargetLayerNames));
 					}
 				}
 
@@ -6308,7 +6310,7 @@ namespace UE::Landscape::EditLayers::Private
 		FString LogMessage = InRendererRenderInfo.RendererState.GetRenderer()->GetEditLayerRendererDebugName();
 		if (!bInIsHeightmapMerge)
 		{
-			LogMessage += FString::Printf(TEXT("\n%s%s%s"), *ConvertTargetLayerNamesToString(InRenderItemTargetLayerNames), !bInAffectsOutputLayerBitIndices ? TEXT("(") : TEXT(""), !bInAffectsOutputLayerBitIndices ? TEXT(")") : TEXT(""));
+			LogMessage += FString::Printf(TEXT("\n%s%s%s"), *UE::Landscape::ConvertTargetLayerNamesToString(InRenderItemTargetLayerNames), !bInAffectsOutputLayerBitIndices ? TEXT("(") : TEXT(""), !bInAffectsOutputLayerBitIndices ? TEXT(")") : TEXT(""));
 		}
 
 		FTransform BaseTransform = FTransform(FVector(0.0, 0.0, InRendererRenderInfo.VisualLogOffsetLocalSpace + OutputLocalOffsetLocalSpace)) * InLandscapeTransform;
diff --git a/Engine/Source/Runtime/Landscape/Private/LandscapeUtils.cpp b/Engine/Source/Runtime/Landscape/Private/LandscapeUtils.cpp
--- a/Engine/Source/Runtime/Landscape/Private/LandscapeUtils.cpp
+++ b/Engine/Source/Runtime/Landscape/Private/LandscapeUtils.cpp
@@ -493,6 +493,13 @@ ULandscapeMaterialInstanceConstant* CreateLandscapeLayerThumbnailMIC(FMaterialUp
 	return MaterialInstance;
 }
 
+FString ConvertTargetLayerNamesToString(const TArrayView<const FName>& InTargetLayerNames)
+{
+	TArray<FString> TargetLayerStrings;
+	Algo::Transform(InTargetLayerNames, TargetLayerStrings, [](FName InTargetLayerName) { return InTargetLayerName.ToString(); });
+	return *FString::Join(TargetLayerStrings, TEXT(","));
+}
+
 #endif // WITH_EDITOR
 
 } // end namespace UE::Landscape
\ No newline at end of file
