Split landscape edit layer batched merge code into several smaller files

diff --git a/Engine/Source/Runtime/Landscape/Public/LandscapeEditLayerRenderer.h b/Engine/Source/Runtime/Landscape/Public/LandscapeEditLayerRenderer.h
--- a/Engine/Source/Runtime/Landscape/Public/LandscapeEditLayerRenderer.h
+++ b/Engine/Source/Runtime/Landscape/Public/LandscapeEditLayerRenderer.h
@@ -5,1155 +5,40 @@
 #include "CoreMinimal.h"
 #include "LandscapeEditTypes.h"
 #include "LandscapeUtils.h"
-#include "Engine/TextureRenderTarget2D.h"
-#include "EngineDefines.h"
-#include "RHIAccess.h"
-#include "StructUtils/InstancedStruct.h"
+#include "LandscapeEditLayerRendererState.h"
+#include "LandscapeEditLayerTargetTypeState.h"
+#include "LandscapeEditLayerTypes.h"
 
-#include "LandscapeEditLayerRenderer.generated.h"
-
-class ALandscape;
-class ULandscapeComponent;
-class ULandscapeInfo;
-class ULandscapeLayerInfoObject;
-class ULandscapeScratchRenderTarget;
-class ULandscapeEditLayerRenderer;
-class ILandscapeEditLayerRenderer;
-class FRDGBuilder;
-class FRHICommandListImmediate;
-struct FLandscapeEditLayerRenderBlackboardItem;
-
-namespace UE::Landscape
-{
-	class FRDGBuilderRecorder;
-} // namespace UE::Landscape
-
-// TODO [jonathan.bard] : Move some classes into their own file, it's getting messy in there
-
-
-// ----------------------------------------------------------------------------------
-
-namespace UE::Landscape::EditLayers
-{
-
-// ----------------------------------------------------------------------------------
-
-FString ConvertTargetLayerNamesToString(const TArrayView<const FName>& InTargetLayerNames);
-
-
-// ----------------------------------------------------------------------------------
-
-// Must match EHeightmapBlendMode in LandscapeEditLayersHeightmaps.usf
-enum class EHeightmapBlendMode : uint32
-{
-	Additive = 0,
-	LegacyAlphaBlend, // In legacy alpha blend, the layer stores the height with premultiplied alpha (legacy landscape splines)
-	AlphaBlend,
-
-	Num,
-};
-
-// Must match EHeightmapAlphaFlags in LandscapeCommon.ush
-enum class EHeightmapAlphaFlags : uint8
-{
-	None = 0, // aka Additive
-	Additive = 0, // The height is considered to be an offset (positive or negative)
-	Min = (1 << 0), // Only lower the existing landscape values
-	Max = (1 << 1), // Only raise the existing landscape values
-	AlphaBlend = (Min | Max), // Raise or lower the existing landscape values
-};
-ENUM_CLASS_FLAGS(EHeightmapAlphaFlags);
-
-// Must match EWeightmapBlendMode in LandscapeEditLayersWeightmaps.usf
-enum class EWeightmapBlendMode : uint32
-{
-	None = 0, // aka Additive
-	Additive = 0,
-	Subtractive,
-	Passthrough,
-	AlphaBlend, 
-
-	Num,
-};
-
-// Must match EWeightmapAlphaFlags in LandscapeCommon.ush
-enum class EWeightmapAlphaFlags : uint8
-{
-	None = 0, // aka Additive
-	Additive = 0, // The weight is considered to be an offset (positive)
-	Min = (1 << 0), // Only retain the min between the weight and the existing landscape weight value
-	Max = (1 << 1), // Only retain the max between the weight and the existing landscape weight value
-	AlphaBlend = (Min | Max), // Full alpha blending of the weight against the existing landscape weight value
-};
-ENUM_CLASS_FLAGS(EWeightmapAlphaFlags);
-
-// Must match EWeightmapTargetLayerFlags in LandscapeEditLayersWeightmaps.usf
-enum class EWeightmapTargetLayerFlags : uint32
-{
-	IsVisibilityLayer = (1 << 0), // This target layer is the visibility layer
-	IsWeightBlended = (1 << 1), // Blend the target layer's value with all the other target layers weights
-
-	None = 0
-};
-ENUM_CLASS_FLAGS(EWeightmapTargetLayerFlags);
-
-// Must match FWeightmapTargetLayerInfo in LandscapeEditLayersWeightmaps.usf
-struct FWeightmapTargetLayerInfo
-{
-	EWeightmapTargetLayerFlags Flags = EWeightmapTargetLayerFlags::None; // Additional info about this target layer
-};
-
-// Defines how heightmaps should be blended (see GenericBlendLayer)
-struct FHeightmapBlendParams
-{
-	FHeightmapBlendParams() = default;
-	FHeightmapBlendParams(EHeightmapBlendMode InBlendMode)
-		: BlendMode(InBlendMode)
-	{}
-
-	EHeightmapBlendMode BlendMode = EHeightmapBlendMode::Additive;
-	float Alpha = 1.0f;
-};
-
-// Defines how weightmaps should be blended (see GenericBlendLayer)
-struct FWeightmapBlendParams
-{
-	FWeightmapBlendParams() = default;
-	FWeightmapBlendParams(EWeightmapBlendMode InBlendMode)
-		: BlendMode(InBlendMode)
-	{}
-
-	static const FWeightmapBlendParams& GetDefaultPassthroughBlendParams()
-	{
-		static const FWeightmapBlendParams BlendParams(EWeightmapBlendMode::Passthrough);
-		return BlendParams;
-	}
-
-	static const FWeightmapBlendParams& GetDefaultAdditiveBlendParams()
-	{
-		static const FWeightmapBlendParams BlendParams(EWeightmapBlendMode::Additive);
-		return BlendParams;
-	}
-
-	EWeightmapBlendMode BlendMode = EWeightmapBlendMode::Passthrough;
-	float Alpha = 1.0f;
-};
-
-// Defines heightmaps+weightmaps blending params (see GenericBlendLayer)
-//  There should be as many blend params as there are target layers to be blended in the blend operation, the others are simply passthrough
-struct FBlendParams
-{
-	FHeightmapBlendParams HeightmapBlendParams;
-	TMap<FName, FWeightmapBlendParams> WeightmapBlendParams;
-};
-
-
-// ----------------------------------------------------------------------------------
-
-#if WITH_EDITOR
-
-/**
- * Global info about the landscape being merged that can be used throughout the merge. 
- */
-class FMergeContext
-{
-public:
-	friend class ::ALandscape;
-
-	FMergeContext(ALandscape* InLandscape, bool bInIsHeightmapMerge, bool bInSkipProceduralRenderers);
-	virtual ~FMergeContext() = default;
-	FMergeContext(const FMergeContext& Other) = default;
-	FMergeContext(FMergeContext&& Other) = default;
-	FMergeContext& operator=(const FMergeContext& Other) = default;
-	FMergeContext& operator=(FMergeContext&& Other) = default;
-
-	LANDSCAPE_EDIT_LAYERS_BATCHED_MERGE_EXPERIMENTAL
-	inline ALandscape* GetLandscape() const { return Landscape; }
-
-	LANDSCAPE_EDIT_LAYERS_BATCHED_MERGE_EXPERIMENTAL
-	inline ULandscapeInfo* GetLandscapeInfo() const { return LandscapeInfo; }
-	
-	/**
-	 * @return all target layer names declared in this context, including ones that are potentially invalid (i.e. no layer info associated)
-	 */
-	LANDSCAPE_EDIT_LAYERS_BATCHED_MERGE_EXPERIMENTAL
-	inline const TArray<FName>& GetAllTargetLayerNames() const { return AllTargetLayerNames; }
-
-	/**
-	 * @return all target layer names that have a valid layer info object. Will return the height layer in the case of a heightmap merge, even though there's no layer info for those : 
-	 */
-	LANDSCAPE_EDIT_LAYERS_BATCHED_MERGE_EXPERIMENTAL
-	TArray<FName> GetValidTargetLayerNames() const;
-
-	/**
-	 * @return all layer info objects declared in this context (can contain null entries for invalid layers or in the case of a heightmap merge)
-	 */
-	LANDSCAPE_EDIT_LAYERS_BATCHED_MERGE_EXPERIMENTAL
-	inline const TArray<ULandscapeLayerInfoObject*>& GetAllWeightmapLayerInfos() const { return AllWeightmapLayerInfos; }
-
-	/**
-	 * @return true if this merge context is used for a heightmap merge, false otherwise (visibility / weightmap)
-	 */
-	LANDSCAPE_EDIT_LAYERS_BATCHED_MERGE_EXPERIMENTAL
-	inline bool IsHeightmapMerge() const { return bIsHeightmapMerge; }
-
-	/**
-	 * @return true if this merge context is used for a heightmap merge, false otherwise (visibility / weightmap)
-	 */
-	LANDSCAPE_EDIT_LAYERS_BATCHED_MERGE_EXPERIMENTAL
-	inline bool ShouldSkipProceduralRenderers() const { return bSkipProceduralRenderers; }
-
-	/**
-	 * @return true if the target layer name is registered to this merge context and has an associated layer info object
-	 */
-	LANDSCAPE_EDIT_LAYERS_BATCHED_MERGE_EXPERIMENTAL
-	LANDSCAPE_API int32 IsValidTargetLayerName(const FName& InName) const;
-
-	/**
-	 * @return true if the target layer name has an associated layer info object. Asserts if the layer name isn't known to the merge context
-	 */
-	LANDSCAPE_EDIT_LAYERS_BATCHED_MERGE_EXPERIMENTAL
-	LANDSCAPE_API int32 IsValidTargetLayerNameChecked(const FName& InName) const;
-
-	/**
-	 * @return true if the target layer name has an associated layer info object. Asserts if the layer index isn't known to the merge context
-	 */
-	LANDSCAPE_EDIT_LAYERS_BATCHED_MERGE_EXPERIMENTAL
-	LANDSCAPE_API int32 IsValidTargetLayerIndex(int32 InIndex) const;
-
-	/**
-	 * @return the index of this target layer name in this context, INDEX_NONE if not found
-	 */
-	LANDSCAPE_EDIT_LAYERS_BATCHED_MERGE_EXPERIMENTAL
-	LANDSCAPE_API int32 GetTargetLayerIndexForName(const FName& InName) const;
-
-	/**
-	 * @return the index of this target layer name in this context. Asserts if the layer name isn't known to the merge context
-	 */
-	LANDSCAPE_EDIT_LAYERS_BATCHED_MERGE_EXPERIMENTAL
-	LANDSCAPE_API int32 GetTargetLayerIndexForNameChecked(const FName& InName) const;
-
-	/**
-	 * @return the target layer name corresponding to the index in this context, NAME_None if not found
-	 */
-	LANDSCAPE_EDIT_LAYERS_BATCHED_MERGE_EXPERIMENTAL
-	LANDSCAPE_API FName GetTargetLayerNameForIndex(int32 InIndex) const;
-
-	/**
-	 * @return the target layer name corresponding to the index in this context, NAME_None if not found. Asserts if the layer index isn't known to the merge context
-	 */
-	LANDSCAPE_EDIT_LAYERS_BATCHED_MERGE_EXPERIMENTAL
-	LANDSCAPE_API FName GetTargetLayerNameForIndexChecked(int32 InIndex) const;
-
-	/**
-	 * @return the index of this landscape layer info in this context, INDEX_NONE if not found
-	 */
-	LANDSCAPE_EDIT_LAYERS_BATCHED_MERGE_EXPERIMENTAL
-	LANDSCAPE_API int32 GetTargetLayerIndexForLayerInfo(ULandscapeLayerInfoObject* InLayerInfo) const;
-
-	/**
-	 * @return the index of this landscape layer info in this context. Asserts if the layer info isn't known to the merge context
-	 */
-	LANDSCAPE_EDIT_LAYERS_BATCHED_MERGE_EXPERIMENTAL
-	LANDSCAPE_API int32 GetTargetLayerIndexForLayerInfoChecked(ULandscapeLayerInfoObject* InLayerInfo) const;
-
-	/**
-	 * @return the landscape layer object corresponding to the target layer name in this context, nullptr if not found
-	 */
-	LANDSCAPE_EDIT_LAYERS_BATCHED_MERGE_EXPERIMENTAL
-	LANDSCAPE_API ULandscapeLayerInfoObject* GetTargetLayerInfoForName(const FName& InName) const;
-
-	/**
-	 * @return the landscape layer object corresponding to the target layer name in this context. Asserts if the layer name isn't known to the merge context
-	 */
-	LANDSCAPE_EDIT_LAYERS_BATCHED_MERGE_EXPERIMENTAL
-	LANDSCAPE_API ULandscapeLayerInfoObject* GetTargetLayerInfoForNameChecked(const FName& InName) const;
-
-	/**
-	 * @return the landscape layer object corresponding to the index in this context, nullptr if not found. Asserts if the layer index isn't known to the merge context
-	 */
-	LANDSCAPE_EDIT_LAYERS_BATCHED_MERGE_EXPERIMENTAL
-	LANDSCAPE_API ULandscapeLayerInfoObject* GetTargetLayerInfoForIndex(int32 InIndex) const;
-
-	/**
-	 * @return a TBitArray<> where, for each layer from InTargetLayerNames that is found in this context, the corresponding bit is set
-	 */
-	LANDSCAPE_EDIT_LAYERS_BATCHED_MERGE_EXPERIMENTAL
-	LANDSCAPE_API TBitArray<> ConvertTargetLayerNamesToBitIndices(TConstArrayView<FName> InTargetLayerNames) const;
-
-	/**
-	 * @return a TBitArray<> where, for each layer from InTargetLayerNames that is found in this context, the corresponding bit is set. Asserts if one of the layer names isn't known to the merge context
-	 */
-	LANDSCAPE_EDIT_LAYERS_BATCHED_MERGE_EXPERIMENTAL
-	LANDSCAPE_API TBitArray<> ConvertTargetLayerNamesToBitIndicesChecked(TConstArrayView<FName> InTargetLayerNames) const;
-
-	/**
-	 * @return a list of target layer names for each set bit in InTargetLayerBitIndices. Asserts if InTargetLayerBitIndices isn't the same size as AllTargetLayerNames
-	 */
-	LANDSCAPE_EDIT_LAYERS_BATCHED_MERGE_EXPERIMENTAL
-	LANDSCAPE_API TArray<FName> ConvertTargetLayerBitIndicesToNames(const TBitArray<>& InTargetLayerBitIndices) const;
-
-	/**
-	 * @return a list of landscape layer info objects for each set bit in InTargetLayerBitIndices. Asserts if InTargetLayerBitIndices isn't the same size as AllTargetLayerNames
-	 */
-	LANDSCAPE_EDIT_LAYERS_BATCHED_MERGE_EXPERIMENTAL
-	LANDSCAPE_API TArray<ULandscapeLayerInfoObject*> ConvertTargetLayerBitIndicesToLayerInfos(const TBitArray<>& InTargetLayerBitIndices) const;
-
-	/**
-	 * Runs the given function for each all valid target layer in the bit indices in parameters, with the possibility of early exit
-	 * Most easily used with a lambda as follows:
-	 * ForEachTargetLayer([](int32 InTargetLayerIndex, FName InTargetLayerName) -> bool
-	 * {
-	 *     return continueLoop ? true : false;
-	 * });
-	 */
-	LANDSCAPE_EDIT_LAYERS_BATCHED_MERGE_EXPERIMENTAL
-	LANDSCAPE_API void ForEachTargetLayer(const TBitArray<>& InTargetLayerBitIndices, TFunctionRef<bool(int32 /*InTargetLayerIndex*/, const FName& /*InTargetLayerName*/, ULandscapeLayerInfoObject* /*InWeightmapLayerInfo*/)> Fn) const;
-
-	/** Same as ForEachTargetLayer but skips over invalid target layers */
-	LANDSCAPE_EDIT_LAYERS_BATCHED_MERGE_EXPERIMENTAL
-	LANDSCAPE_API void ForEachValidTargetLayer(TFunctionRef<bool(int32 /*InTargetLayerIndex*/, const FName& /*InTargetLayerName*/, ULandscapeLayerInfoObject* /*InWeightmapLayerInfo*/)> Fn) const;
-
-	/**
-	 * @return the index corresponding to the visibility target layer in AllTargetLayerNames, INDEX_NONE if it's not in the list of layers
-	 */
-	inline int32 GetVisibilityTargetLayerIndex() const { return VisibilityTargetLayerIndex; }
-
-	/**
-	 * @return a bit mask the size of AllTargetLayerNames, where all bits are set to zero except the one corresponding to the visibility target layer
-	 */
-	inline const TBitArray<>& GetVisibilityTargetLayerMask() const { return VisibilityTargetLayerMask; }
-
-	/** Same as GetVisibilityTargetLayerMask() but negated */
-	inline const TBitArray<>& GetNegatedVisibilityTargetLayerMask() const { return NegatedVisibilityTargetLayerMask; }
-
-	/**
-	 * @return a bit mask for all target layers from AllTargetLayerNames that have a valid landscape layer info object
-	 */
-	inline const TBitArray<>& GetValidTargetLayerBitIndices() const { return ValidTargetLayerBitIndices; }
-
-	/**
-	 * @return a bit mask the size of AllTargetLayerNames, where all bits are set to bBitValue
-	 */
-	inline TBitArray<> BuildTargetLayerBitIndices(bool bInBitValue) const;
-
-protected:
-	// COMMENT [jonathan.bard] : purposefully don't use ELandscapeToolTargetType here as ELandscapeToolTargetType::Weightmap and ELandscapeToolTargetType::Visibility have to be processed together (because of weightmap packing, 
-	//  which means a visibility weightmap could be another channel of a texture which contains weightmap up to 3 other weightmaps, so we have to resolve the 4 channels altogether). 
-	//  Note: this could change when SUPPORTS_LANDSCAPE_EDITORONLY_UBER_MATERIAL is done...
-	/** Type of merge being requested */
-	bool bIsHeightmapMerge = false;
-
-	/** Allows to skip all non-default (i.e. persistent edit layer) layer renderers for this merge. Mostly there for backwards compatibility and debugging purposes. */
-	bool bSkipProceduralRenderers = false;
-
-	/** Landscape being merged */
-	ALandscape* Landscape = nullptr;
-
-	/** Landscape info associated with the landscape being merged */
-	ULandscapeInfo* LandscapeInfo = nullptr;
-
-	/** List of all target layer names that can possibly be rendered on this landscape (even invalid ones). The index of the target layer in that list is important :
-	  it allows to use a TBitArray instead of a TSet to designate a list of target layers all throughout the merge, which greatly accelerates all the intersection tests being made.
-	  If merging heightmaps, it contains a single name (that is only meaningful for debug display). */
-	TArray<FName> AllTargetLayerNames;
-
-	/** Bit mask that corresponds to only the visibility layer */
-	TBitArray<> VisibilityTargetLayerMask;
-
-	/** Bit mask that corresponds to all target layers except the visibility layer */
-	TBitArray<> NegatedVisibilityTargetLayerMask;
-
-	/** Per-target layer info. Same size as AllTargetLayerNames (only meaningful when rendering weightmaps) */
-	TArray<ULandscapeLayerInfoObject*> AllWeightmapLayerInfos;
-
-	/** List of valid target layers. If a target layer name is present here, it's because it has a valid landscape layer info object. Each bit in that bit array corresponds to an entry in AllTargetLayerNames */
-	TBitArray<> ValidTargetLayerBitIndices;
-
-	// Index of the visibility layer in AllTargetLayerNames (always valid for weightmap merges)
-	int32 VisibilityTargetLayerIndex = INDEX_NONE;
-};
-
-
-// ----------------------------------------------------------------------------------
-
-/** 
- * FEditLayerTargetTypeState fully describes the state of an edit layer renderer wrt its target types. It's named after the enum "ELandscapeToolTargetType" in order to tell
- *  whether the renderer's heightmaps and/or visibility and/or weightmaps are enabled (and if so, which weightmap is enabled exactly)  
- *  It is meant to be provided by the edit layer renderer's GetRendererStateInfo implementation.
- */
-class FEditLayerTargetTypeState
-{
-	friend class ILandscapeEditLayerRenderer;
-	friend class FEditLayerRendererState;
-
-private:
-	// Private constructor : either use the constructors taking a merge context in parameter or use GetDummyTargetTypeState()
-	FEditLayerTargetTypeState() = default;
-
-public:
-	LANDSCAPE_EDIT_LAYERS_BATCHED_MERGE_EXPERIMENTAL
-	static const FEditLayerTargetTypeState& GetDummyTargetTypeState();
-
-	/**
-	 * Constructs an empty FEditLayerTargetTypeState (all target types turned off, no active weightmaps)
-	 * 
-	 * @param InMergeContext context containing all sorts of information related to this merge operation
-	 */
-	LANDSCAPE_EDIT_LAYERS_BATCHED_MERGE_EXPERIMENTAL
-	LANDSCAPE_API FEditLayerTargetTypeState(const FMergeContext* InMergeContext);
-
-	/**
-	 * Initializes a FEditLayerTargetTypeState with the target types passed in parameter 
-	 * 
-	 * @param InMergeContext context containing all sorts of information related to this merge operation
-	 * @param InTargetTypeMask target type mask supported by this target type state (no active weightmaps)
-	 */
-	LANDSCAPE_EDIT_LAYERS_BATCHED_MERGE_EXPERIMENTAL
-	LANDSCAPE_API FEditLayerTargetTypeState(const FMergeContext* InMergeContext, ELandscapeToolTargetTypeFlags InTargetTypeMask);
-
-	/**
-	 * Initializes a FEditLayerTargetTypeState with the target types and the weightmaps passed in parameter (by name)
-	 * 
-	 * @param InMergeContext context containing all sorts of information related to this merge operation
-	 * @param InTargetTypeMask target type mask supported by this target type state
-	 * @param InSupportedWeightmaps names of weightmaps supported by this target type state 
-	 * @param bInChecked, if true, will validate that every weightmap requested is part of the merge context (if false, only the known weightmaps will be marked as supported)
-	 */
-	LANDSCAPE_EDIT_LAYERS_BATCHED_MERGE_EXPERIMENTAL
-	LANDSCAPE_API FEditLayerTargetTypeState(const FMergeContext* InMergeContext, ELandscapeToolTargetTypeFlags InTargetTypeMask, const TArrayView<const FName>& InSupportedWeightmaps, bool bInChecked);
-
-	/**
-	 * Initializes a FEditLayerTargetTypeState with the target types and (optionally) the weightmaps passed in parameter (by bit index)
-	 * 
-	 * @param InMergeContext context containing all sorts of information related to this merge operation
-	 * @param InTargetTypeMask target type mask supported by this target type state
-	 * @param InSupportedWeightmaps bit indices of the weightmaps supported by this target type state (corresponds to the FMergeContext's AllTargetLayerNames)
-	 * Asserts if one of the layer indices isn't valid for the merge context (except when InSupportedWeightmapLayerIndices is == TBitArray<>)
-	 */
-	LANDSCAPE_EDIT_LAYERS_BATCHED_MERGE_EXPERIMENTAL
-	LANDSCAPE_API FEditLayerTargetTypeState(const FMergeContext* InMergeContext, ELandscapeToolTargetTypeFlags InTargetTypeMask, const TBitArray<>& InSupportedWeightmapLayerIndices);
-	
-	/**
-	 * Indicates whether a given target type is currently active in this state
-	 * @param InTargetType the requested target type (heightmap/weightmap/visibility)
-	 * @param InWeightmapLayerName (optional) is the requested weightmap, only relevant for the ELandscapeToolTargetType::Weightmap case
-	 */
-	LANDSCAPE_EDIT_LAYERS_BATCHED_MERGE_EXPERIMENTAL
-	LANDSCAPE_API bool IsActive(ELandscapeToolTargetType InTargetType, const FName& InWeightmapLayerName = NAME_None) const;
-
-		/**
-	 * Indicates whether a given target type is currently active in this state
-	 * @param InTargetType the requested target type (heightmap/weightmap/visibility)
-	 * @param InWeightmapLayerName (optional) is the requested weightmap, only relevant for the ELandscapeToolTargetType::Weightmap case. Asserts if the layer name isn't valid for the merge context (except for NAME_None)
-	 */
-	LANDSCAPE_EDIT_LAYERS_BATCHED_MERGE_EXPERIMENTAL
-	LANDSCAPE_API bool IsActiveChecked(ELandscapeToolTargetType InTargetType, const FName& InWeightmapLayerName = NAME_None) const;
-
-	/**
-	 * Indicates whether a given target type is currently active in this state
-	 * @param InTargetType is the requested target type (heightmap/weightmap/visibility)
-	 * @param int32 (optional) is the requested weightmap index (in FMergeContex), only relevant for the ELandscapeToolTargetType::Weightmap case. Asserts if the layer index isn't valid for the merge context (except for INDEX_NONE)
-	 */
-	LANDSCAPE_EDIT_LAYERS_BATCHED_MERGE_EXPERIMENTAL
-	LANDSCAPE_API bool IsActive(ELandscapeToolTargetType InTargetType, int32 InWeightmapLayerIndex = INDEX_NONE) const;
-
-	/** Returns the currently active weightmaps : 
-	 *  - If Weightmap is amongst the target types it will return all the weightmaps
-	 *  - Additionally, if Visibility is amongst the target types, it will also return the visibility weightmap
-	 *  - If neither Weightmap nor Visibility is amongst the target types, it will return an empty array
-	 */
-	LANDSCAPE_EDIT_LAYERS_BATCHED_MERGE_EXPERIMENTAL
-	LANDSCAPE_API TArray<FName> GetActiveWeightmaps() const;
-
-	/** Returns the currently active weightmaps :
-	 *  - If Weightmap is amongst the target types it will return all the weightmaps
-	 *  - Additionally, if Visibility is amongst the target types, it will also return the visibility weightmap
-	 *  - If neither Weightmap nor Visibility is amongst the target types, it will return an empty array
-	 */
-	LANDSCAPE_EDIT_LAYERS_BATCHED_MERGE_EXPERIMENTAL
-	LANDSCAPE_API TBitArray<> GetActiveWeightmapBitIndices() const;
-
-	/** Returns the target type mask (i.e. same as ELandscapeToolTargetType, but as bit flags) */
-	LANDSCAPE_EDIT_LAYERS_BATCHED_MERGE_EXPERIMENTAL
-	LANDSCAPE_API ELandscapeToolTargetTypeFlags GetTargetTypeMask() const { return TargetTypeMask; }
-
-	/** Sets the target type mask (i.e. same as ELandscapeToolTargetType, but as bit flags) */
-	LANDSCAPE_EDIT_LAYERS_BATCHED_MERGE_EXPERIMENTAL
-	LANDSCAPE_API void SetTargetTypeMask(ELandscapeToolTargetTypeFlags InTargetTypeMask);
-
-	/** Adds the target type in parameter to the mask of active target types */
-	LANDSCAPE_EDIT_LAYERS_BATCHED_MERGE_EXPERIMENTAL
-	LANDSCAPE_API void AddTargetType(ELandscapeToolTargetType InTargetType);
-
-	/** Appends the target type mask in parameter to the mask of active target types */
-	LANDSCAPE_EDIT_LAYERS_BATCHED_MERGE_EXPERIMENTAL
-	LANDSCAPE_API void AddTargetTypeMask(ELandscapeToolTargetTypeFlags InTargetTypeMask);
-
-	/** Removes a single target type from the mask of active target types */
-	LANDSCAPE_EDIT_LAYERS_BATCHED_MERGE_EXPERIMENTAL
-	LANDSCAPE_API void RemoveTargetType(ELandscapeToolTargetType InTargetType);
-
-	/** Removes the target type mask in parameter from the mask of active target types */
-	LANDSCAPE_EDIT_LAYERS_BATCHED_MERGE_EXPERIMENTAL
-	LANDSCAPE_API void RemoveTargetTypeMask(ELandscapeToolTargetTypeFlags InTargetTypeMask);
-
-	/** Adds a weightmap to the list of active weightmaps (make sure ELandscapeToolTargetType::Weightmap is amongst the supported target types) */
-	LANDSCAPE_EDIT_LAYERS_BATCHED_MERGE_EXPERIMENTAL
-	LANDSCAPE_API void AddWeightmap(const FName& InWeightmapLayerName);
-
-	/** Adds a weightmap to the list of active weightmaps (make sure ELandscapeToolTargetType::Weightmap is amongst the supported target types). Asserts if the layer name isn't known to the merge context.  */
-	LANDSCAPE_EDIT_LAYERS_BATCHED_MERGE_EXPERIMENTAL
-	LANDSCAPE_API void AddWeightmapChecked(const FName& InWeightmapLayerName);
-
-	/** Adds a weightmap to the list of active weightmaps (make sure ELandscapeToolTargetType::Weightmap is amongst the supported target types). Asserts if the layer index isn't valid for the merge context */
-	LANDSCAPE_EDIT_LAYERS_BATCHED_MERGE_EXPERIMENTAL
-	LANDSCAPE_API void AddWeightmap(int32 InWeightmapLayerIndex);
-
-	/** Removes a weightmap from the list of active weightmaps */
-	LANDSCAPE_EDIT_LAYERS_BATCHED_MERGE_EXPERIMENTAL
-	LANDSCAPE_API void RemoveWeightmap(const FName& InWeightmapLayerName);
-
-	/** Removes a weightmap from the list of active weightmaps. Asserts if the layer name isn't known to the merge context. */
-	LANDSCAPE_EDIT_LAYERS_BATCHED_MERGE_EXPERIMENTAL
-	LANDSCAPE_API void RemoveWeightmapChecked(const FName& InWeightmapLayerName);
-
-	/** Removes a weightmap from the list of active weightmaps. Asserts if the layer index isn't valid for the merge context. */
-	LANDSCAPE_EDIT_LAYERS_BATCHED_MERGE_EXPERIMENTAL
-	LANDSCAPE_API void RemoveWeightmap(int32 InWeightmapLayerIndex);
-
-	/**
-	 * Returns the "intersection" (AND operation) between the target type state and the one in parameter. e.g. if state S0 is:
-	 *  - (---------|Visibility|Weightmap) with active weightmaps (A|B|C|-), and state S1 is:
-	 *  - (Heightmap|----------|Weightmap) with active weightmaps (-|-|C|D), then state S0.Intersect(S1) is:
-	 *  - (---------|----------|Weightmap) with active weightmaps (-|-|C|-)
-	 */
-	FEditLayerTargetTypeState Intersect(const FEditLayerTargetTypeState& InOther) const;
-
-	bool operator == (const FEditLayerTargetTypeState& InOther) const;
-
-	FString ToString() const;
-
-private:
-	/** Global context being used for this merge : contains generic information about the landscape, its available layer names, etc. */
-	const FMergeContext* MergeContext = nullptr;
-
-	/** Bitmask of the target types that are supported  */
-	ELandscapeToolTargetTypeFlags TargetTypeMask = ELandscapeToolTargetTypeFlags::None;
-	
-	/** 
-	 * List of weightmaps that are supported for the ELandscapeToolTargetType::Weightmap/ELandscapeToolTargetType::Visibility type. 
-	 *  Each bit in that bit array corresponds to an entry in FMergeContext's AllTargetLayerNames
-	 */
-	TBitArray<> WeightmapTargetLayerBitIndices;
-};
-
-
-// ----------------------------------------------------------------------------------
-
-/**
- * FEditLayerRendererState describes the entire state of an edit layer renderer : what it is capable of doing (SupportedTargetTypeState, immutable) and what it is currently doing (EnabledTargetTypeState, mutable)
- *  These states are provided by IEditLayerRendererProvider in order to describe both what the renderer can do and what it currently does by default. e.g. a disabled edit layer supports rendering heightmaps but
- *  its enabled state for the heightmap target type is false. This way, the user can selectively enable it at merge time without altering the entire landscape's state (i.e. just for the purpose of a specific merge render)
- *  A target type must be both supported and enabled on a given edit layer renderer in order for this renderer to render anything. 
- *  It also describes the target layer groups this renderer needs when rendering its weightmap (i.e. which weightmap needs to be rendered with which weightmaps : e.g. for weight blending)
- */
-class FEditLayerRendererState
-{
-private:
-	// Private constructor : either use the constructors taking a merge context in parameter or use GetDummyRendererState()
-	FEditLayerRendererState() = default;
-
-public:
-	LANDSCAPE_EDIT_LAYERS_BATCHED_MERGE_EXPERIMENTAL
-	static const FEditLayerRendererState& GetDummyRendererState();
-
-	LANDSCAPE_EDIT_LAYERS_BATCHED_MERGE_EXPERIMENTAL
-	LANDSCAPE_API explicit FEditLayerRendererState(const FMergeContext* InMergeContext);
-	LANDSCAPE_EDIT_LAYERS_BATCHED_MERGE_EXPERIMENTAL
-	LANDSCAPE_API explicit FEditLayerRendererState(const FMergeContext* InMergeContext, TScriptInterface<ILandscapeEditLayerRenderer> InRenderer);
-
-	/** Returns the edit layer renderer which this state relates to */
-	LANDSCAPE_EDIT_LAYERS_BATCHED_MERGE_EXPERIMENTAL
-	LANDSCAPE_API TScriptInterface<ILandscapeEditLayerRenderer> GetRenderer() const { return Renderer; }
-
-	/** Returns a mask of all target types / weightmaps supported by this renderer */
-	LANDSCAPE_EDIT_LAYERS_BATCHED_MERGE_EXPERIMENTAL
-	LANDSCAPE_API const FEditLayerTargetTypeState& GetSupportedTargetTypeState() const { return SupportedTargetTypeState; }
-
-	/** Returns a mask of all target types / weightmaps enabled on this renderer */
-	LANDSCAPE_EDIT_LAYERS_BATCHED_MERGE_EXPERIMENTAL
-	LANDSCAPE_API const FEditLayerTargetTypeState& GetEnabledTargetTypeState() const { return EnabledTargetTypeState; }
-
-	/** Returns a mask of all target types / weightmaps supported and enabled on this renderer */
-	LANDSCAPE_EDIT_LAYERS_BATCHED_MERGE_EXPERIMENTAL
-	LANDSCAPE_API FEditLayerTargetTypeState GetActiveTargetTypeState() const { return ActiveTargetTypeState; }
-
-	/** Mutates the EnabledTargetTypeState by adding the target type in parameter to the mask of active target types */
-	LANDSCAPE_EDIT_LAYERS_BATCHED_MERGE_EXPERIMENTAL
-	LANDSCAPE_API void EnableTargetType(ELandscapeToolTargetType InTargetType);
-	
-	/** Mutates the EnabledTargetTypeState by adding the target type mask in parameter to the mask of active target types */
-	LANDSCAPE_EDIT_LAYERS_BATCHED_MERGE_EXPERIMENTAL
-	LANDSCAPE_API void EnableTargetTypeMask(ELandscapeToolTargetTypeFlags InTargetTypeMask);
-
-	/** Mutates the EnabledTargetTypeState by removing the target type in parameter from the mask of active target types */
-	LANDSCAPE_EDIT_LAYERS_BATCHED_MERGE_EXPERIMENTAL
-	LANDSCAPE_API void DisableTargetType(ELandscapeToolTargetType InTargetType);
-
-	/** Mutates the EnabledTargetTypeState by removing the target type mask in parameter from the mask of active target types */
-	LANDSCAPE_EDIT_LAYERS_BATCHED_MERGE_EXPERIMENTAL
-	LANDSCAPE_API void DisableTargetTypeMask(ELandscapeToolTargetTypeFlags InTargetTypeMask);
-
-	/**
-	 * Mutates the EnabledTargetTypeState by adding the weightmap in parameter to the list of enabled weightmaps, if it is known to the merge context
-	 */
-	LANDSCAPE_EDIT_LAYERS_BATCHED_MERGE_EXPERIMENTAL
-	LANDSCAPE_API void EnableWeightmap(ELandscapeToolTargetType InTargetType, const FName& InWeightmapLayerName);
-
-	/**
-	 * Mutates the EnabledTargetTypeState by adding the weightmap in parameter to the list of enabled weightmaps. Asserts if the layer name isn't known to the merge context
-	 */
-	LANDSCAPE_EDIT_LAYERS_BATCHED_MERGE_EXPERIMENTAL
-	LANDSCAPE_API void EnableWeightmapChecked(ELandscapeToolTargetType InTargetType, const FName& InWeightmapLayerName);
-
-	/** 
-	 * Mutates the EnabledTargetTypeState by adding the weightmap in parameter to the list of enabled weightmaps. Asserts if the layer index isn't known to the merge context
-	 */
-	LANDSCAPE_EDIT_LAYERS_BATCHED_MERGE_EXPERIMENTAL
-	LANDSCAPE_API void EnableWeightmap(ELandscapeToolTargetType InTargetType, int32 InWeightmapLayerIndex);
-
-	/**
-	 * Mutates the EnabledTargetTypeState by removing the weightmap in parameter from the list of enabled weightmaps, if it is known to the merge context
-	 */
-	LANDSCAPE_EDIT_LAYERS_BATCHED_MERGE_EXPERIMENTAL
-	LANDSCAPE_API void DisableWeightmap(const FName& InWeightmapLayerName);
-
-	/**
-	 * Mutates the EnabledTargetTypeState by removing the weightmap in parameter from the list of enabled weightmaps. Asserts if the layer index isn't known to the merge context
-	 */
-	LANDSCAPE_EDIT_LAYERS_BATCHED_MERGE_EXPERIMENTAL
-	LANDSCAPE_API void DisableWeightmapChecked(const FName& InWeightmapLayerName);
-
-	/** 
-	 * Mutates the EnabledTargetTypeState by removing the weightmap in parameter from the list of enabled weightmaps. Asserts if the layer index isn't known to the merge context
-	 */
-	LANDSCAPE_EDIT_LAYERS_BATCHED_MERGE_EXPERIMENTAL
-	LANDSCAPE_API void DisableWeightmap(int32 InWeightmapLayerIndex);
-
-	/** Returns the mask of target types that are both supported and enabled by this renderer */
-
-	LANDSCAPE_EDIT_LAYERS_BATCHED_MERGE_EXPERIMENTAL
-	LANDSCAPE_API ELandscapeToolTargetTypeFlags GetActiveTargetTypeMask() const;
-
-	/**
-	 * Indicates whether a given target type and weightmap layer name is currently supported and enabled by this renderer. 
-	 * @param InTargetType the requested target type (heightmap/weightmap/visibility)
-	 * @param InWeightmapLayerName (optional) is the requested weightmap, only relevant for the ELandscapeToolTargetType::Weightmap case
-	 */
-	LANDSCAPE_EDIT_LAYERS_BATCHED_MERGE_EXPERIMENTAL
-	LANDSCAPE_API bool IsTargetActive(ELandscapeToolTargetType InTargetType, const FName& InWeightmapLayerName = NAME_None) const;
-
-	/**
-	 * Indicates whether a given target type and weightmap layer name is currently supported and enabled by this renderer. Asserts if the layer name isn't known to the merge context (except if NAME_None)
-	 * @param InTargetType the requested target type (heightmap/weightmap/visibility)
-	 * @param InWeightmapLayerName (optional) is the requested weightmap, only relevant for the ELandscapeToolTargetType::Weightmap case
-	 */
-	LANDSCAPE_EDIT_LAYERS_BATCHED_MERGE_EXPERIMENTAL
-	LANDSCAPE_API bool IsTargetActiveChecked(ELandscapeToolTargetType InTargetType, const FName& InWeightmapLayerName = NAME_None) const;
-
-	/**
-	 * Indicates whether a given target type and weightmap layer name is currently supported and enabled by this renderer. Asserts if the layer index isn't known to the merge context (except if INDEX_NONE)
-	 * @param InTargetType the requested target type (heightmap/weightmap/visibility)
-	 * @param InWeightmapLayerIndex (optional) is the requested weightmap's index in the merge index, only relevant for the ELandscapeToolTargetType::Weightmap case
-	 */
-	LANDSCAPE_EDIT_LAYERS_BATCHED_MERGE_EXPERIMENTAL
-	LANDSCAPE_API bool IsTargetActive(ELandscapeToolTargetType InTargetType, int32 InWeightmapLayerIndex = INDEX_NONE) const;
-
-	/**
-	 * Returns a list of all weightmaps supported and enabled by this renderer (only relevant for ELandscapeToolTargetType::Weightmap (and ELandscapeToolTargetType::Visibility))
-	 */
-	LANDSCAPE_EDIT_LAYERS_BATCHED_MERGE_EXPERIMENTAL
-	LANDSCAPE_API TArray<FName> GetActiveTargetWeightmaps() const;
-
-	/** 
-	 * Returns a list of all weightmaps supported and enabled by this renderer (only relevant for ELandscapeToolTargetType::Weightmap (and ELandscapeToolTargetType::Visibility))
-	 */
-	LANDSCAPE_EDIT_LAYERS_BATCHED_MERGE_EXPERIMENTAL
-	LANDSCAPE_API TBitArray<> GetActiveTargetWeightmapBitIndices() const;
-
-	/** Returns the target layer groups associated with this renderer. A target layer group is a set of target layers (weightmaps) that depend on one another in order to produce the output target layers.
-	 This allows to implement "horizontal blending", where weightmaps can be blended with one another at each step of the landscape edit layers merge algorithm */
-	LANDSCAPE_EDIT_LAYERS_BATCHED_MERGE_EXPERIMENTAL
-	LANDSCAPE_API TArray<TBitArray<>> GetTargetLayerGroups() const { return TargetLayerGroups; }
-
-private:
-	void UpdateActiveTargetTypeState();
-
-private:
-	/** Global context being used for this merge : contains generic information about the landscape, its available layer names, etc. */
-	const FMergeContext* MergeContext = nullptr;
-
-	/** Renderer associated with this state */
-	TScriptInterface<ILandscapeEditLayerRenderer> Renderer;
-
-	/** For debug purposes : this is the same as Renderer->GetEditLayerRendererDebugName() but having a member makes it easier to debug in the watch window */
-	FString DebugName;
-	
-PRAGMA_DISABLE_EXPERIMENTAL_WARNINGS
-	/** Struct that contains the supported target types and the which weightmap are supported by this renderer. Immutable. */
-	FEditLayerTargetTypeState SupportedTargetTypeState = FEditLayerTargetTypeState::GetDummyTargetTypeState();
-	
-	/** Struct that contains the enabled target types and the which weightmap are currently enabled by this renderer. Can be set by the user using EnableTargetType. */
-	FEditLayerTargetTypeState EnabledTargetTypeState = FEditLayerTargetTypeState::GetDummyTargetTypeState();
-
-	/** Intersection of SupportedTargetTypeState and EnabledTargetTypeState (defines what is both supported and enabled on this renderer) */
-	FEditLayerTargetTypeState ActiveTargetTypeState = FEditLayerTargetTypeState::GetDummyTargetTypeState();
-PRAGMA_ENABLE_EXPERIMENTAL_WARNINGS
-
-	/* List of groups of target layers that this renderer requires to be rendered together. All target layers listed in SupportedTargetTypeState must belong to 1 (and 1 only) target layer group.
-	 *  Each target layer group is a bit array for which each bit corresponds to an entry in FMergeContext's AllTargetLayerNames
-	 */
-	TArray<TBitArray<>> TargetLayerGroups;
-};
-
-
-// ----------------------------------------------------------------------------------
-
-/** 
- * Params struct passed to the merge function. It contains everything needed for requesting a given set of target layers (for weightmaps) on a given number of components and for a certain configuration of edit layers
- *  Note that this is what is requested by the caller, but in practice, there might be more renderers (e.g. some might get added e.g. legacy weight-blending, some removed because they turn out to be disabled...)
- *  and more weightmaps being rendered (e.g. a requested weightmap might depend on another one that has not been requested), or less (e.g. a requested weightmap is actually invalid) :
- */
-struct FMergeRenderParams
-{
-	FMergeRenderParams(TArray<ULandscapeComponent*> InComponentsToMerge, const TArrayView<FEditLayerRendererState>& InEditLayerRendererStates, const TSet<FName>& InWeightmapLayerNames = {}, bool bInRequestAllLayers = false);
-
-	/** List of components that need merging */
-	TArray<ULandscapeComponent*> ComponentsToMerge;
-
-	/** Requested states for every edit layer renderer participating to the merge */
-	TArray<FEditLayerRendererState> EditLayerRendererStates;
-	
-	/** List of weightmap layers being requested. */
-	TSet<FName> WeightmapLayerNames;
-
-	/** Ignore the computed list of requested layers and instead request all valid layers */
-	bool bRequestAllLayers = false;
-};
-
-
-// ----------------------------------------------------------------------------------
-
-/** Flags that allow a given ILandscapeEditLayerRenderer to customize the way it renders/blends */
-enum class ERenderFlags : uint8
-{
-	None = 0,
-
-	// Render mode flags
-	RenderMode_Recorded = (1 << 0), // This renderer can record its render commands into a single FRDGBuilder on the render thread (prefer this if possible), via a FRDGBuilderRecorder in "recording" mode. Exclusive with RenderMode_Immediate
-	RenderMode_Immediate = (1 << 1), // This renderer enqueues its render commands immediately either via the FRDGBuilderRecorder in "immediate" mode or just enqueuing render commands the usual way. Exclusive with RenderMode_Recorded
-	RenderMode_Mask = RenderMode_Recorded | RenderMode_Immediate,
-
-	// Blend mode flags
-	BlendMode_SeparateBlend = (1 << 2), // This renderer has a separate render function for blending. When this flag is not set, only RenderLayer is called and is assumed to both render the layer and blend it. When it's set, RenderLayer will be followed by BlendLayer.
-
-	// Render layer group flags
-	RenderLayerGroup_SupportsGrouping = (1 << 3), // This renderer supports being rendered along with others in a series of RenderLayer steps before performing a single BlendLayer. Assumes BlendMode_SeparateBlend
-};
-ENUM_CLASS_FLAGS(ERenderFlags);
+#if UE_ENABLE_INCLUDE_ORDER_DEPRECATED_IN_5_6
+#include "LandscapeEditLayerMergeRenderContext.h"
+#include "LandscapeEditLayerMergeRenderBlackboardItem.h"
+#endif // UE_ENABLE_INCLUDE_ORDER_DEPRECATED_IN_5_6
 
+#include "LandscapeEditLayerRenderer.generated.h"
 
 // ----------------------------------------------------------------------------------
+// Forward declarations
 
-/**
- * Defines an individual render step of the batch merge
- */
-struct FMergeRenderStep
-{
-	/** Describes what kind of operation this render step executes. Despite some of the names in there, those steps all run on the game thread but the _RenderThread ones will defer operations to the render thread. */
-	enum class EType
-	{
-		Invalid,
-
-		BeginRenderCommandRecorder, // Initiates a render command recorder that will batch one or several RenderLayer_Recorded operations, runs on the game thread
-		EndRenderCommandRecorder, // Ends the render command recorder initiated by the last BeginRenderCommandRecorder, runs on the game thread but this is where the render command executing the recorded operations from (several) RenderLayer_Recorded step(s) will be pushed (on the render thread, this is where the FRDBuilder is created and executed)
-
-		BeginRenderLayerGroup, // Initiates a series of RenderLayer steps for renderers that can be rendered one after another without an intermediate BlendLayer. Triggers a call to BeginRenderLayerGroup() on the first renderer in the group 
-		EndRenderLayerGroup, // Ends a series of RenderLayer steps for renderers that can be rendered one after another without an intermediate BlendLayer. Triggers a call to EndRenderLayerGroup() on the last renderer in the group 
-
-		RenderLayer, // Performs the rendering of a target layer group on an edit layer on a given world region (i.e. in a batch), runs on the game thread. Use the command recorder to perform the render operations. In RenderMode_Recorded, these will be delayed until EndRenderCommandRecorder
-		BlendLayer, // Performs the rendering of a target layer group on an edit layer on a given world region (i.e. in a batch), runs on the game thread. Use the command recorder to perform the render operations. In RenderMode_Recorded, these will be delayed until EndRenderCommandRecorder
-
-		SignalBatchMergeGroupDone, // Final step when rendering a target layer group on a given world region (i.e. in a batch) : runs on the game thread and allows to retrieve the result of the merge and do something with it (e.g. resolve the final textures)
-	};
-
-	FMergeRenderStep(EType InType)
-		: Type(InType)
-	{}
-
-
-	FMergeRenderStep(EType InType, const TBitArray<>& InTargetLayerGroupBitIndices, const TArrayView<ULandscapeComponent*>& InComponentsToRender)
-		: Type(InType)
-		, TargetLayerGroupBitIndices(InTargetLayerGroupBitIndices)
-		, ComponentsToRender(InComponentsToRender)
-	{}
-
-	FMergeRenderStep(EType InType, ERenderFlags InRenderFlags, const FEditLayerRendererState& InRendererState, const TBitArray<>& InTargetLayerGroupBitIndices, const TArrayView<ULandscapeComponent*>& InComponentsToRender)
-		: Type(InType)
-		, RenderFlags(InRenderFlags)
-		, RendererState(InRendererState)
-		, TargetLayerGroupBitIndices(InTargetLayerGroupBitIndices)
-		, ComponentsToRender(InComponentsToRender)
-	{}
-
-	UE_DEPRECATED(5.6, "Use the other constructors")
-	FMergeRenderStep(EType InType, const FEditLayerRendererState& InRendererState, const TBitArray<>& InTargetLayerGroupBitIndices, const TArrayView<ULandscapeComponent*>& InComponentsToRender)
-		: FMergeRenderStep(InType, ERenderFlags::None, InRendererState, InTargetLayerGroupBitIndices, InComponentsToRender)
-	{}
-
-	/** Type of operation for this step */
-	EType Type = EType::Invalid;
-	
-	/** The render flags corresponding to this state */
-	ERenderFlags RenderFlags = ERenderFlags::None;
-
-PRAGMA_DISABLE_EXPERIMENTAL_WARNINGS
-	/** Renderer state to be used this step. This includes the renderer as well as its precise step (e.g. which weightmap are supported? which are enabled?) */
-	FEditLayerRendererState RendererState = FEditLayerRendererState::GetDummyRendererState();
-PRAGMA_ENABLE_EXPERIMENTAL_WARNINGS
-
-	/** List of target layers being involved in this step. Each bit in that bit array corresponds to an entry in FMergeContext's AllTargetLayerNames */
-	TBitArray<> TargetLayerGroupBitIndices;
-
-	/** List of components involved in this step */
-	TArray<ULandscapeComponent*> ComponentsToRender;
-
-	UE_DEPRECATED(5.6, "Renamed : use TargetLayerGroupBitIndices")
-	TBitArray<> RenderGroupBitIndices;
-};
-
-
-// ----------------------------------------------------------------------------------
+class ULandscapeLayerInfoObject;
 
-/** 
- * Defines an individual render batch when merging the landscape. A batch corresponds to a target layer group on a world's region, i.e. a set of weightmaps (or the heightmap) to render on a 
- *  portion of the world. Each batch is composed of a series of render steps. 
- */
-struct FMergeRenderBatch
+namespace UE::Landscape
 {
-	bool operator < (const FMergeRenderBatch& InOther) const;
-
-	ALandscape* Landscape = nullptr;
-
-	/** Section of the landscape being covered by this batch (in landscape vertex coordinates, inclusive bounds) */
-	FIntRect SectionRect;
-
-	/** Resolution of the render target needed for this batch (including duplicate borders) */
-	// TODO [jonathan.bard] : rename EffectiveResolution and make private maybe?
-	FIntPoint Resolution = FIntPoint(ForceInitToZero);
-	
-	FIntPoint MinComponentKey = FIntPoint(MAX_int32, MAX_int32);
-	FIntPoint MaxComponentKey = FIntPoint(MIN_int32, MIN_int32);
-	
-	/** Sequential list of rendering operations that need to be performed to fully render this batch*/
-	TArray<FMergeRenderStep> RenderSteps;
-	
-	/** List of all components involved in this batch */
-	TSet<ULandscapeComponent*> ComponentsToRender;
-
-	/** List of all target layers being rendered in this batch (i.e. bitwise OR of all of the render steps' TargetLayerGroupBitIndices). Each bit in that bit array corresponds to an entry in FMergeRenderContext's AllTargetLayerNames  */
-	TBitArray<> TargetLayerBitIndices;
-
-	/** List of components involved in this batch and the target layers they're writing to (redundant with ComponentsToRender but we keep the latter for convenience)
-	 (each bit corresponds to a target layer name in FMergeContext's AllTargetLayerNames) */
-	TMap<ULandscapeComponent*, TBitArray<>> ComponentToTargetLayerBitIndices;
-
-	/** Reverse lookup of ComponentToTargetLayerBitIndices : one entry per element in ComponentToTargetLayerBitIndices, each entry containing all of the components involved in this merge for this target layer */
-	TArray<TSet<ULandscapeComponent*>> TargetLayersToComponents;
-
-public:
-	FIntPoint GetRenderTargetResolution(bool bInWithDuplicateBorders) const;
-	
-	/**
-	 * Find the area in the render batch render target corresponding to each of the subsections of this component 
-	 * @param InComponent Which component to compute the subsection rects for 
-	 * @param OutSubsectionRects List of (up to 4) subsection rects when *not* taking into account duplicate borders (inclusive bounds)
-	 * @param OutSubsectionRectsWithDuplicateBorders List of (up to 4) subsection rects when taking into account duplicate borders (inclusive bounds)
-	 * @param bInWithDuplicateBorders indicates whether the rect coordinates should include the duplicated column/row at the end of the subsection or not
-	 * 
-	 * @return number of subsections
-	 */
-	int32 ComputeSubsectionRects(ULandscapeComponent* InComponent, TArray<FIntRect, TInlineAllocator<4>>& OutSubsectionRects, TArray<FIntRect, TInlineAllocator<4>>& OutSubsectionRectsWithDuplicateBorders) const;
-
-	/**
-	 * Find the area in the render batch render target corresponding to this component 
-	 * @param InComponent Which component to compute the section rect for 
-	 * @param bInWithDuplicateBorders indicates whether the rect coordinates should include the duplicated columns/rows at the end of each subsection or not
-	 * 
-	 * @Return the section rect (inclusive bounds)
-	 */
-	FIntRect ComputeSectionRect(ULandscapeComponent* InComponent, bool bInWithDuplicateBorders) const;
-
-	/**
-	 * Compute the rects corresponding to the sub-sections that need to be read from and written to when expanding the render target (inclusive bounds)
-	 * @param OutSubsectionRects List of all subsection rects when *not* taking into account duplicate borders (inclusive bounds)
-	 * @param OutSubsectionRectsWithDuplicateBorders List of all subsection rects when taking into account duplicate borders (inclusive bounds)
-	 */
-	void ComputeAllSubsectionRects(TArray<FIntRect>& OutSubsectionRects, TArray<FIntRect>& OutSubsectionRectsWithDuplicateBorders) const;
-};
-
-
-// ----------------------------------------------------------------------------------
+	class FRDGBuilderRecorder;
+} // namespace UE::Landscape
 
-/**
- * Utility struct for attaching some information that pertains to a given landscape component in the context of a batch render
- */
-struct FComponentMergeRenderInfo
+namespace UE::Landscape::EditLayers
 {
-	/** Component to render */
-	ULandscapeComponent* Component = nullptr;
-
-	/** Texture region that corresponds to this component in the render area's render target */
-	FIntRect ComponentRegionInRenderArea;
-
-	/** Index of the component in the render area's render target */
-	FIntPoint ComponentKeyInRenderArea = FIntPoint(ForceInitToZero);
-
-	bool operator < (const FComponentMergeRenderInfo& InOther) const;
-};
+	struct FComponentMergeRenderInfo;
+	class FMergeRenderContext;
+} // namespace UE::Landscape::EditLayers
 
 
 // ----------------------------------------------------------------------------------
 
-/** 
- * Utility class that contains everything necessary to perform the batched merge : scratch render targets, list of batches, etc.
- */
-class FMergeRenderContext : public FMergeContext
+namespace UE::Landscape::EditLayers
 {
-public:
-	friend class ::ALandscape;
-
-	FMergeRenderContext(const FMergeContext& InMergeContext);
-	virtual ~FMergeRenderContext();
-	FMergeRenderContext(const FMergeRenderContext& Other) = default;
-	FMergeRenderContext(FMergeRenderContext&& Other) = default;
-	FMergeRenderContext& operator=(const FMergeRenderContext& Other) = default;
-	FMergeRenderContext& operator=(FMergeRenderContext&& Other) = default;
-
-	bool IsValid() const;
-
-	/**
-	 * Cycle between the 3 render targets used for blending:
-	 *  Write becomes Read -> Read becomes ReadPrevious -> ReadPrevious becomes Write
-	 *	The new Write RT will be transitioned to state == InDesiredWriteAccess (if != ERHIAccess::None)
-	 *  The new Read RT will be transitioned to state ERHIAccess::SRVMask
-	 *  The new ReadPrevious RT will stay in state ERHIAccess::SRVMask
-	 * 
-	 * @param InDesiredWriteAccess RHI state of the write RT after the operation (if InDesiredWriteAccess != ERHIAccess::None)
-	 * @param RDGBuilderRecorder recorder to append operations to a single FRDGBuilder, if the the recorder in "recording" mode. In "immediate" mode, a render command will be enqueued immediately
-	 */
-	LANDSCAPE_EDIT_LAYERS_BATCHED_MERGE_EXPERIMENTAL
-	LANDSCAPE_API void CycleBlendRenderTargets(UE::Landscape::FRDGBuilderRecorder& RDGBuilderRecorder);
-
-	LANDSCAPE_EDIT_LAYERS_BATCHED_MERGE_EXPERIMENTAL
-	LANDSCAPE_API ULandscapeScratchRenderTarget* GetBlendRenderTargetWrite() const;
-	LANDSCAPE_EDIT_LAYERS_BATCHED_MERGE_EXPERIMENTAL
-	LANDSCAPE_API ULandscapeScratchRenderTarget* GetBlendRenderTargetRead() const;
-	LANDSCAPE_EDIT_LAYERS_BATCHED_MERGE_EXPERIMENTAL
-	LANDSCAPE_API ULandscapeScratchRenderTarget* GetBlendRenderTargetReadPrevious() const;
-	LANDSCAPE_EDIT_LAYERS_BATCHED_MERGE_EXPERIMENTAL
-	LANDSCAPE_API ULandscapeScratchRenderTarget* GetValidityRenderTarget(const FName& InTargetLayerName) const;
-	
-	struct FOnRenderBatchTargetGroupDoneParams
-	{
-		FOnRenderBatchTargetGroupDoneParams(FMergeRenderContext* InMergeRenderContext, 
-			const TArrayView<FName>& InTargetLayerGroupLayerNames,
-			const TArrayView<ULandscapeLayerInfoObject*>& InTargetLayerGroupLayerInfos,
-			const TArrayView<FComponentMergeRenderInfo>& InSortedComponentMergeRenderInfos)
-			: MergeRenderContext(InMergeRenderContext)
-			, TargetLayerGroupLayerNames(InTargetLayerGroupLayerNames)
-			, TargetLayerGroupLayerInfos(InTargetLayerGroupLayerInfos)
-			, SortedComponentMergeRenderInfos(InSortedComponentMergeRenderInfos)
-		{}
-
-		/** Render context : this is still active in this step and can be used for doing additional renders in the blend render targets, etc. */
-		FMergeRenderContext* MergeRenderContext = nullptr;
-
-		/** List of target layers being involved in this step */
-		TArray<FName> TargetLayerGroupLayerNames;
-
-		/** List of target layer info objects being involved in this step (same size as TargetLayerGroupLayerNames) */
-		TArray<ULandscapeLayerInfoObject*> TargetLayerGroupLayerInfos;
-
-		/** Additional info about the components that have been processed in this batch render */
-		const TArray<FComponentMergeRenderInfo> SortedComponentMergeRenderInfos;
-
-		UE_DEPRECATED(5.6, "Renamed : use TargetLayerGroupLayerNames")
-		TArray<FName> RenderGroupTargetLayerNames;
-		UE_DEPRECATED(5.6, "Renamed : use TargetLayerGroupLayerInfos")
-		TArray<ULandscapeLayerInfoObject*> RenderGroupTargetLayerInfos;
-		UE_DEPRECATED(5.6, "Removed : use MergeRenderContext->GetCurrentRenderBatch()")
-		const FMergeRenderBatch* RenderBatch = nullptr;
-	};
-	void Render(TFunction<void(const FOnRenderBatchTargetGroupDoneParams& /*InParams*/, UE::Landscape::FRDGBuilderRecorder& /*RDGBuilderRecorder*/)> OnRenderBatchTargetGroupDone);
-
-	LANDSCAPE_EDIT_LAYERS_BATCHED_MERGE_EXPERIMENTAL
-	inline FIntPoint GetMaxNeededResolution() const { return MaxNeededResolution; }
-	LANDSCAPE_EDIT_LAYERS_BATCHED_MERGE_EXPERIMENTAL
-	inline const TArray<FMergeRenderBatch>& GetRenderBatches() const { return RenderBatches; }
-	LANDSCAPE_EDIT_LAYERS_BATCHED_MERGE_EXPERIMENTAL
-	LANDSCAPE_API const FMergeRenderBatch* GetCurrentRenderBatch() const;
-
-	LANDSCAPE_EDIT_LAYERS_BATCHED_MERGE_EXPERIMENTAL
-	LANDSCAPE_API FTransform ComputeVisualLogTransform(const FTransform& InTransform) const;
-	LANDSCAPE_EDIT_LAYERS_BATCHED_MERGE_EXPERIMENTAL
-	LANDSCAPE_API void IncrementVisualLogOffset();
-	LANDSCAPE_EDIT_LAYERS_BATCHED_MERGE_EXPERIMENTAL
-	LANDSCAPE_API void ResetVisualLogOffset();
-
-#if ENABLE_VISUAL_LOG
-	LANDSCAPE_EDIT_LAYERS_BATCHED_MERGE_EXPERIMENTAL
-	LANDSCAPE_API static int32 GetVisualLogAlpha();
-	LANDSCAPE_EDIT_LAYERS_BATCHED_MERGE_EXPERIMENTAL
-	LANDSCAPE_API bool IsVisualLogEnabled() const;
-#endif // ENABLE_VISUAL_LOG
-
-	/** 
-	 * Render the stencil render targets for each target layer in this merge for this batch 
-	 * 
-	 * @param RDGBuilderRecorder recorder to append operations to a single FRDGBuilder, if the the recorder in "recording" mode. In "immediate" mode, a render command will be enqueued immediately
-	 */
-	void RenderValidityRenderTargets(UE::Landscape::FRDGBuilderRecorder& RDGBuilderRecorder);
-
-	/** 
-	 * Duplicates the vertex data from the (sub-)sections of the batch, assuming GetBlendRenderTargetRead() is the RT that is read from and GetBlendRenderTargetWrite() the one that is written to 
-	 * 
-	 * @param RDGBuilderRecorder recorder to append operations to a single FRDGBuilder, if the the recorder in "recording" mode. In "immediate" mode, a render command will be enqueued immediately
-	 */
-	void RenderExpandedRenderTarget(UE::Landscape::FRDGBuilderRecorder& RDGBuilderRecorder);
-
-	/** 
-	 * Performs a generic blend assuming GetBlendRenderTargetWrite() is the RT that contains the layer to blend and GetBlendRenderTargetRead() the one that contains the result of the merge up until this layer.
-	 *  It will cycle the RTs such as after this call, GetBlendRenderTargetWrite() will contain the merge result
-	 * 
-	 * @param InBlendParams defines how heightmaps/weightmaps should be blended. In the case of weightmaps, there must be as may entries in WeightmapBlendParams as there are layers in the target layers group
-	 * @param RDGBuilderRecorder recorder to append operations to a single FRDGBuilder, if the the recorder in "recording" mode. In "immediate" mode, a render command will be enqueued immediately
-	 */
-	LANDSCAPE_EDIT_LAYERS_BATCHED_MERGE_EXPERIMENTAL
-	LANDSCAPE_API void GenericBlendLayer(const FBlendParams& InBlendParams, UE::Landscape::FRDGBuilderRecorder& RDGBuilderRecorder);
-
-	const TBitArray<>& GetFinalTargetLayerBitIndices() const 
-	{ 
-		return FinalTargetLayerBitIndices; 
-	}
-
-	/** 
-	 * @return true if there is at least one blackboard item of type T in the context's list 
-	 */
-	template <typename T, typename = std::enable_if_t<std::is_base_of_v<FLandscapeEditLayerRenderBlackboardItem, std::decay_t<T>>>>
-	LANDSCAPE_EDIT_LAYERS_BATCHED_MERGE_EXPERIMENTAL
-	bool HasBlackboardItem() const;
-
-	/** 
-	 * Create a new blackboard item of type T and adds it to the context's list
-	 * 
-	 * @return the newly-created blackboard item
-	 */
-	template <typename T, typename... TArgs, typename = std::enable_if_t<std::is_base_of_v<FLandscapeEditLayerRenderBlackboardItem, std::decay_t<T>>>>
-	LANDSCAPE_EDIT_LAYERS_BATCHED_MERGE_EXPERIMENTAL
-	T& AddBlackboardItem(TArgs&&... InArgs);
-
-	/**
-	 * @return the first blackboard item of type T from the context's list or nullptr if there isn't one
-	 */
-	template <typename T, typename = std::enable_if_t<std::is_base_of_v<FLandscapeEditLayerRenderBlackboardItem, std::decay_t<T>>>>
-	LANDSCAPE_EDIT_LAYERS_BATCHED_MERGE_EXPERIMENTAL
-	T* TryGetBlackboardItem();
-
-	/**
-	 * @return the first blackboard item of type T from the context's list. Asserts if there isn't one
-	 */
-	template <typename T, typename = std::enable_if_t<std::is_base_of_v<FLandscapeEditLayerRenderBlackboardItem, std::decay_t<T>>>>
-	LANDSCAPE_EDIT_LAYERS_BATCHED_MERGE_EXPERIMENTAL
-	T& GetBlackboardItem();
-
-	/**
-	 * @return the first blackboard item of type T from the context's list or create a new one if there isn't one
-	 */
-	template <typename T, typename... TArgs, typename = std::enable_if_t<std::is_base_of_v<FLandscapeEditLayerRenderBlackboardItem, std::decay_t<T>>>>
-	LANDSCAPE_EDIT_LAYERS_BATCHED_MERGE_EXPERIMENTAL
-	T& GetOrCreateBlackboardItem(TArgs&&... InArgs);
-
-	/**
-	 * @return all blackboard items of type T from the context's list
-	 */
-	template <typename T, typename = std::enable_if_t<std::is_base_of_v<FLandscapeEditLayerRenderBlackboardItem, std::decay_t<T>>>>
-	LANDSCAPE_EDIT_LAYERS_BATCHED_MERGE_EXPERIMENTAL
-	TArray<T*> GetBlackboardItems();
-
-	/**
-	 * @return all blackboard items currently stored in the context
-	 */
-	LANDSCAPE_EDIT_LAYERS_BATCHED_MERGE_EXPERIMENTAL
-	const TArray<TInstancedStruct<FLandscapeEditLayerRenderBlackboardItem>>& GetAllBlackboardItems() const { return BlackboardItems; }
-
-private:
-	/** Allocate all needed render targets for this merge */
-	void AllocateResources();
-	/** Free all render targets used in this merge */
-	void FreeResources();
-	/** Allocate all needed render targets for this batch */
-	void AllocateBatchResources(const FMergeRenderBatch& InRenderBatch);
-	/** Free all render targets used in this batch */
-	void FreeBatchResources(const FMergeRenderBatch& InRenderBatch);
-
-private:
-	// Blending is pretty much what we only do during the merge. It requires 3 render targets : 1 that we write to and therefore use as RTV (Write) and 2 that we read from and therefore use 
-	//  as SRV (one that contains the layer to merge, the other the accumulated result so far) : Previous(SRV) + Current(SRV) --> Write(RTV)
-	static constexpr int32 NumBlendRenderTargets = 3; 
-
-	/** Render targets that are used throughout the blending operations (they could be texture arrays in the case of multiple weightmaps) */
-	TStaticArray<ULandscapeScratchRenderTarget*, NumBlendRenderTargets> BlendRenderTargets;
-	int32 CurrentBlendRenderTargetWriteIndex = -1; 
-
-	/** Final list of target layer names being involved in this merge context. If a target layer name is present here, it's because it's a valid target layer and it needs to be rendered
-	  because it has been requested or one of the target layers that have been requested needs it to be present (e.g. weight-blending). Each bit in that bit array corresponds to an entry in AllTargetLayerNames */
-	TBitArray<> FinalTargetLayerBitIndices;
-
-	/** Render targets storing the validity of each pixel wrt the target layer (i.e. stencil-like buffer, but stored as a RT to let users access it as a standard texture) : 
-	 Useful when sampling neighbors to know whether the data there corresponds to a valid neighbor.*/
-	TMap<FName, ULandscapeScratchRenderTarget*> PerTargetLayerValidityRenderTargets;
-
-	/** Maximum resolution needed by a given batch in this context (means we won't ever need more than this size of a render target during the whole merge */
-	FIntPoint MaxNeededResolution = FIntPoint(ForceInit);
-
-	/** Maximum number of slices needed by a given batch / target layer group in this context (means we won't ever need more than this number of slices in a given render target texture array during the whole merge */
-	int32 MaxNeededNumSlices = 0;
-
-	/** Successive batches of components being processed by this context. Each batch should be self-contained so that we won't ever need to keep more than one in memory (VRAM)
-	  in order to compute the info we need (heightmap/weightmaps), i.e. nothing should ever depend on 2 different batches. */
-	TArray<FMergeRenderBatch> RenderBatches;
-
-	/** Current batch being rendered */
-	int32 CurrentRenderBatchIndex = INDEX_NONE;
-
-	/** Offset for visual debugging */
-	FVector CurrentVisualLogOffset = FVector(ForceInitToZero);
-
-	/** Maximum height of all components to render in local space. */
-	double MaxLocalHeight = DBL_MIN;
-
-	/** List of components involved in this merge and the target layers they're writing to (each bit corresponds to the target layer name in AllTargetLayerNames) */
-	TMap<ULandscapeComponent*, TBitArray<>> ComponentToTargetLayerBitIndices;
-
-	/** Reverse lookup of ComponentToTargetLayerBitIndices : one entry per element in ComponentToTargetLayerBitIndices, each entry containing all of the components involved in this merge for this target layer */
-	TArray<TSet<ULandscapeComponent*>> TargetLayersToComponents;
-
-	/** Generic data struct to store some data to pass around throughout the context's lifetime. Basically, this is a mini-RTTI system that e.g. allows to store/pass around specific data between renderers in a generic way */
-	TArray<TInstancedStruct<FLandscapeEditLayerRenderBlackboardItem>> BlackboardItems;
-};
 
+#if WITH_EDITOR
 
 // ----------------------------------------------------------------------------------
 
@@ -1330,8 +215,8 @@ private:
 
 // ----------------------------------------------------------------------------------
 
-/** 
- * Interface to implement to be able to provide an ordered list of renderers to the landscape. 
+/**
+ * Interface to implement to be able to provide an ordered list of renderers to the landscape.
  */
 class IEditLayerRendererProvider
 {
@@ -1344,26 +229,13 @@ public:
 	 * @return a list of renderer states (i.e. a ILandscapeEditLayerRenderer and its current state) to be processed in that order by the merge operation
 	 */
 	LANDSCAPE_EDIT_LAYERS_BATCHED_MERGE_EXPERIMENTAL
-	LANDSCAPE_API virtual TArray<FEditLayerRendererState> GetEditLayerRendererStates(const FMergeContext* InMergeContext) PURE_VIRTUAL(IEditLayerRendererProvider::GetEditLayerRendererStates, return {}; );
+		LANDSCAPE_API virtual TArray<FEditLayerRendererState> GetEditLayerRendererStates(const FMergeContext* InMergeContext) PURE_VIRTUAL(IEditLayerRendererProvider::GetEditLayerRendererStates, return {}; );
 };
-#endif // WITH_EDITOR
-
-} //namespace UE::Landscape::EditLayers
 
+#endif // WITH_EDITOR
 
-// ----------------------------------------------------------------------------------
+} // namespace UE::Landscape::EditLayers
 
-/**
- * Base class for various blackboard items that can be attached to FMergeRenderContext in order to transfer temporary information
- *  between renderers. It's a USTRUCT only to benefit from the RTTI services provided by TInstancedStruct
- */
-USTRUCT(meta = (Abstract))
-struct FLandscapeEditLayerRenderBlackboardItem
-{
-	GENERATED_BODY()
-	
-	virtual ~FLandscapeEditLayerRenderBlackboardItem() = default;
-};
 
 // ----------------------------------------------------------------------------------
 
@@ -1398,16 +270,7 @@ public:
 			const TArrayView<UE::Landscape::EditLayers::FComponentMergeRenderInfo>& InSortedComponentMergeRenderInfos, 
 			const FTransform& InRenderAreaWorldTransform, 
 			const FIntRect& InRenderAreaSectionRect,
-			int32 InNumSuccessfulRenderLayerStepsUntilBlendLayerStep)
-			: MergeRenderContext(InMergeRenderContext)
-			, TargetLayerGroupLayerNames(InTargetLayerGroupLayerNames)
-			, TargetLayerGroupLayerInfos(InTargetLayerGroupLayerInfos)
-			, RendererState(InRendererState)
-			, SortedComponentMergeRenderInfos(InSortedComponentMergeRenderInfos)
-			, RenderAreaWorldTransform(InRenderAreaWorldTransform)
-			, RenderAreaSectionRect(InRenderAreaSectionRect)
-			, NumSuccessfulRenderLayerStepsUntilBlendLayerStep(InNumSuccessfulRenderLayerStepsUntilBlendLayerStep)
-		{}
+			int32 InNumSuccessfulRenderLayerStepsUntilBlendLayerStep);
 
 		/** Merge context */
 		UE::Landscape::EditLayers::FMergeRenderContext* MergeRenderContext = nullptr;
@@ -1560,89 +423,3 @@ public:
 
 #endif // WITH_EDITOR
 };
-
-// TODO [jonathan.bard] : move to .inl file :
-namespace UE::Landscape::EditLayers
-{
-
-#if WITH_EDITOR
-
-template <typename T, typename>
-bool UE::Landscape::EditLayers::FMergeRenderContext::HasBlackboardItem() const
-{
-	for (const TInstancedStruct<FLandscapeEditLayerRenderBlackboardItem>& Item : BlackboardItems)
-	{
-		if (const T* TypedItem = Item.template GetPtr<T>())
-		{
-			return true;
-		}
-	}
-	return false;
-}
-
-template <typename T, typename... TArgs, typename>
-T& UE::Landscape::EditLayers::FMergeRenderContext::AddBlackboardItem(TArgs&&... InArgs)
-{
-	T* Item = BlackboardItems.Add_GetRef(TInstancedStruct<T>::Make(Forward<TArgs>(InArgs)...)).template GetMutablePtr<T>();
-	return *Item;
-}
-
-template <typename T, typename>
-T* UE::Landscape::EditLayers::FMergeRenderContext::TryGetBlackboardItem()
-{
-	for (TInstancedStruct<FLandscapeEditLayerRenderBlackboardItem>& Item : BlackboardItems)
-	{
-		if (T* TypedItem = Item.template GetMutablePtr<T>())
-		{
-			return TypedItem;
-		}
-	}
-	return nullptr;
-}
-
-template <typename T, typename>
-T& UE::Landscape::EditLayers::FMergeRenderContext::GetBlackboardItem()
-{
-	for (TInstancedStruct<FLandscapeEditLayerRenderBlackboardItem>& Item : BlackboardItems)
-	{
-		if (T* TypedItem = Item.template GetMutablePtr<T>())
-		{
-			return *TypedItem;
-		}
-	}
-	check(false);
-
-	// this is absolutely wrong, but we shouldn't ever reach that line and we need to return something anyway : 
-	static FLandscapeEditLayerRenderBlackboardItem DummyItem;
-	return *reinterpret_cast<T*>(&DummyItem);
-}
-
-template <typename T, typename... TArgs, typename>
-T& UE::Landscape::EditLayers::FMergeRenderContext::GetOrCreateBlackboardItem(TArgs&&... InArgs)
-{
-	T* Item = TryGetBlackboardItem<T>();
-	if (Item == nullptr)
-	{
-		Item = BlackboardItems.Add_GetRef(TInstancedStruct<T>::Make(Forward<TArgs>(InArgs)...)).template GetMutablePtr<T>();
-	}
-	check(Item != nullptr);
-	return *Item;
-}
-
-template <typename T, typename>
-TArray<T*> UE::Landscape::EditLayers::FMergeRenderContext::GetBlackboardItems()
-{
-	TArray<T*> Items;
-	for (TInstancedStruct<FLandscapeEditLayerRenderBlackboardItem>& Item : BlackboardItems)
-	{
-		if (T* TypedItem = Item.template GetMutablePtr<T>())
-		{
-			Items.Add(TypedItem);
-		}
-	}
-	return Items;
-}
-
-#endif // WITH_EDITOR
-
-} //namespace UE::Landscape::EditLayers
\ No newline at end of file
diff --git a/Engine/Source/Runtime/Landscape/Public/LandscapeUtils.h b/Engine/Source/Runtime/Landscape/Public/LandscapeUtils.h
--- a/Engine/Source/Runtime/Landscape/Public/LandscapeUtils.h
+++ b/Engine/Source/Runtime/Landscape/Public/LandscapeUtils.h
@@ -335,6 +335,9 @@ LANDSCAPE_API UMaterialInstance* CreateToolLandscapeMaterialInstanceConstant(UMa
 /** Create a thumbnail material for a given layer. Can return nullptr if the option to disable landscape thumbnails has been turned on */
 LANDSCAPE_API ULandscapeMaterialInstanceConstant* CreateLandscapeLayerThumbnailMIC(FMaterialUpdateContext& MaterialUpdateContext, UMaterialInterface* LandscapeMaterial, FName LayerName);
 
+/** Concatenates the target layer names in parameter into a string */
+LANDSCAPE_API FString ConvertTargetLayerNamesToString(const TArrayView<const FName>& InTargetLayerNames);
+
 #endif // WITH_EDITOR
 
 } // end namespace UE::Landscape
\ No newline at end of file
