Non-unity build fix

#rnx

diff --git a/Engine/Source/Runtime/Landscape/Private/LandscapeEditLayerRendererPrivate.cpp b/Engine/Source/Runtime/Landscape/Private/LandscapeEditLayerRendererPrivate.cpp
--- a/Engine/Source/Runtime/Landscape/Private/LandscapeEditLayerRendererPrivate.cpp
+++ b/Engine/Source/Runtime/Landscape/Private/LandscapeEditLayerRendererPrivate.cpp
@@ -6,6 +6,7 @@
 #include "Algo/ForEach.h"
 #include "Algo/Transform.h"
 #include "LandscapeEditResourcesSubsystem.h"
+#include "LandscapeEditLayerMergeRenderContext.h"
 #include "LandscapeLayerInfoObject.h"
 #include "LandscapeInfo.h"
 #include "RHIAccess.h"
