Removed unneeded localization macro

#rnx

diff --git a/Engine/Source/Runtime/Landscape/Private/LandscapeEditLayerMergeRenderContext.cpp b/Engine/Source/Runtime/Landscape/Private/LandscapeEditLayerMergeRenderContext.cpp
--- a/Engine/Source/Runtime/Landscape/Private/LandscapeEditLayerMergeRenderContext.cpp
+++ b/Engine/Source/Runtime/Landscape/Private/LandscapeEditLayerMergeRenderContext.cpp
@@ -25,8 +25,6 @@
 #include "TextureResource.h"
 #include "VisualLogger/VisualLogger.h"
 
-#define LOCTEXT_NAMESPACE "LandscapeEditLayerRenderer"
-
 extern TAutoConsoleVariable<float> CVarLandscapeBatchedMergeVisualLogOffsetIncrement;
 extern TAutoConsoleVariable<int32> CVarLandscapeEditLayersClearBeforeEachWriteToScratch;
 extern TAutoConsoleVariable<int32> CVarLandscapeBatchedMergeVisualLogShowMergeType;
@@ -1278,5 +1276,3 @@ FMergeRenderParams::FMergeRenderParams(TArray<ULandscapeComponent*> InComponents
 #endif // WITH_EDITOR
 
 } // namespace UE::Landscape::EditLayers
-
-#undef LOCTEXT_NAMESPACE
