Fixed Mac compilation
#rnx

diff --git a/Engine/Source/Runtime/Landscape/Public/LandscapeEditLayerMergeContext.h b/Engine/Source/Runtime/Landscape/Public/LandscapeEditLayerMergeContext.h
--- a/Engine/Source/Runtime/Landscape/Public/LandscapeEditLayerMergeContext.h
+++ b/Engine/Source/Runtime/Landscape/Public/LandscapeEditLayerMergeContext.h
@@ -204,7 +204,8 @@ public:
 	/**
 	 * @return a bit mask the size of AllTargetLayerNames, where all bits are set to bBitValue
 	 */
-	inline TBitArray<> BuildTargetLayerBitIndices(bool bInBitValue) const;
+	LANDSCAPE_EDIT_LAYERS_BATCHED_MERGE_EXPERIMENTAL
+	LANDSCAPE_API TBitArray<> BuildTargetLayerBitIndices(bool bInBitValue) const;
 
 protected:
 	// COMMENT [jonathan.bard] : purposefully don't use ELandscapeToolTargetType here as ELandscapeToolTargetType::Weightmap and ELandscapeToolTargetType::Visibility have to be processed together (because of weightmap packing, 
