Add cvar "landscape.RuntimeVirtualTextureRenderWithQuad" which defaults off.
When enabled the "Virtual Texture Render With Quad" setting on landscape actors is overriden to be forced on.

diff --git a/Engine/Source/Runtime/Landscape/Private/LandscapeRender.cpp b/Engine/Source/Runtime/Landscape/Private/LandscapeRender.cpp
--- a/Engine/Source/Runtime/Landscape/Private/LandscapeRender.cpp
+++ b/Engine/Source/Runtime/Landscape/Private/LandscapeRender.cpp
@@ -227,6 +227,15 @@ static FAutoConsoleVariableRef CLandscapeAllowNanitePerClusterDisplacementDisabl
 	ECVF_RenderThreadSafe
 );
 
+bool GLandscapeRuntimeVirtualTextureRenderWithQuad = false;
+FAutoConsoleVariableRef CVarLandscapeRuntimeVirtualTextureRenderWithQuad(
+	TEXT("landscape.RuntimeVirtualTextureRenderWithQuad"),
+	GLandscapeRuntimeVirtualTextureRenderWithQuad,
+	TEXT("Force use of a single quad when rendering landscape to runtime virtual texture pages"),
+	FConsoleVariableDelegate::CreateStatic(&OnCVarNeedingRenderStateInvalidationChanged),
+	ECVF_RenderThreadSafe | ECVF_Scalability
+);
+
 #if WITH_EDITOR
 extern TAutoConsoleVariable<int32> CVarLandscapeShowDirty;
 #endif
@@ -1544,7 +1553,7 @@ FLandscapeComponentSceneProxy::FLandscapeComponentSceneProxy(ULandscapeComponent
 	LODSettings.LastLODScreenSizeSquared = LODScreenRatioSquared[LastLOD];
 	LODSettings.ForcedLOD = ForcedLOD;
 
-	const bool bVirtualTextureRenderWithQuad = InComponent->GetLandscapeProxy()->bVirtualTextureRenderWithQuad;
+	const bool bVirtualTextureRenderWithQuad = InComponent->GetLandscapeProxy()->bVirtualTextureRenderWithQuad || GLandscapeRuntimeVirtualTextureRenderWithQuad;
 	const bool bVirtualTextureRenderWithQuadHQ = InComponent->GetLandscapeProxy()->bVirtualTextureRenderWithQuadHQ;
 	VirtualTexturePerPixelHeight = bVirtualTextureRenderWithQuad ? bVirtualTextureRenderWithQuadHQ ? 2 : 1 : 0;
 
