Fixed bug introduced with the last landscape edit layer merge refactor : blending was writing to / reading from an incorrect  slice index in the texture array.

Misc:
struct ILandcapeEditLayerRenderer::FRenderParams has been deprecated in favor of UE::Landscape::EditLayers::FRenderParams, in order to be able to forward-declare it

#rnx

diff --git a/Engine/Plugins/Experimental/LandscapePatch/Source/LandscapePatch/Private/LandscapeCircleHeightPatch.cpp b/Engine/Plugins/Experimental/LandscapePatch/Source/LandscapePatch/Private/LandscapeCircleHeightPatch.cpp
--- a/Engine/Plugins/Experimental/LandscapePatch/Source/LandscapePatch/Private/LandscapeCircleHeightPatch.cpp
+++ b/Engine/Plugins/Experimental/LandscapePatch/Source/LandscapePatch/Private/LandscapeCircleHeightPatch.cpp
@@ -82,7 +82,7 @@ bool ULandscapeCircleHeightPatch::CanGroupRenderLayerWith(TScriptInterface<ILand
 	return OtherRenderer->IsA<ULandscapeCircleHeightPatch>();
 }
 
-bool ULandscapeCircleHeightPatch::RenderLayer(FRenderParams& RenderParams, UE::Landscape::FRDGBuilderRecorder& RDGBuilderRecorder)
+bool ULandscapeCircleHeightPatch::RenderLayer(UE::Landscape::EditLayers::FRenderParams& RenderParams, UE::Landscape::FRDGBuilderRecorder& RDGBuilderRecorder)
 {
 	using namespace UE::Landscape::PatchUtil;
 	using namespace UE::Landscape::EditLayers;
@@ -126,7 +126,7 @@ bool ULandscapeCircleHeightPatch::RenderLayer(FRenderParams& RenderParams, UE::L
 	return bDidRenderSomething;
 }
 
-void ULandscapeCircleHeightPatch::BlendLayer(FRenderParams& RenderParams, UE::Landscape::FRDGBuilderRecorder& RDGBuilderRecorder)
+void ULandscapeCircleHeightPatch::BlendLayer(UE::Landscape::EditLayers::FRenderParams& RenderParams, UE::Landscape::FRDGBuilderRecorder& RDGBuilderRecorder)
 {
 	using namespace UE::Landscape::EditLayers;
 
@@ -142,10 +142,10 @@ void ULandscapeCircleHeightPatch::BlendLayer(FRenderParams& RenderParams, UE::La
 	}
 
 	// Then perform the generic blend : 
-	RenderParams.MergeRenderContext->GenericBlendLayer(BlendParams, RDGBuilderRecorder);
+	RenderParams.MergeRenderContext->GenericBlendLayer(BlendParams, RenderParams, RDGBuilderRecorder);
 }
 
-bool ULandscapeCircleHeightPatch::ApplyCirclePatch(bool bInPerformBlending, FRenderParams* RenderParams, UE::Landscape::FRDGBuilderRecorder& RDGBuilderRecorder, bool bIsVisibilityLayer,
+bool ULandscapeCircleHeightPatch::ApplyCirclePatch(bool bInPerformBlending, UE::Landscape::EditLayers::FRenderParams* RenderParams, UE::Landscape::FRDGBuilderRecorder& RDGBuilderRecorder, bool bIsVisibilityLayer,
 	FTextureResource* InMergedLandscapeTextureResource, int32 LandscapeTextureSliceIndex,
 	const FIntPoint& DestinationResolution, const FTransform& HeightmapCoordsToWorld, ERHIAccess OutputAccess)
 {
diff --git a/Engine/Plugins/Experimental/LandscapePatch/Source/LandscapePatch/Private/LandscapeTexturePatch.cpp b/Engine/Plugins/Experimental/LandscapePatch/Source/LandscapePatch/Private/LandscapeTexturePatch.cpp
--- a/Engine/Plugins/Experimental/LandscapePatch/Source/LandscapePatch/Private/LandscapeTexturePatch.cpp
+++ b/Engine/Plugins/Experimental/LandscapePatch/Source/LandscapePatch/Private/LandscapeTexturePatch.cpp
@@ -154,7 +154,7 @@ bool ULandscapeTexturePatch::CanGroupRenderLayerWith(TScriptInterface<ILandscape
 	return OtherRenderer->IsA<ULandscapeTexturePatch>();
 }
 
-bool ULandscapeTexturePatch::RenderLayer(FRenderParams& RenderParams, UE::Landscape::FRDGBuilderRecorder& RDGBuilderRecorder)
+bool ULandscapeTexturePatch::RenderLayer(UE::Landscape::EditLayers::FRenderParams& RenderParams, UE::Landscape::FRDGBuilderRecorder& RDGBuilderRecorder)
 {
 	using namespace UE::Landscape;
 	using namespace UE::Landscape::PatchUtil;
@@ -249,7 +249,7 @@ bool ULandscapeTexturePatch::RenderLayer(FRenderParams& RenderParams, UE::Landsc
 	return bDidRenderSomething;
 }//end ULandscapeTexturePatch::RenderLayer
 
-void ULandscapeTexturePatch::BlendLayer(FRenderParams& RenderParams, UE::Landscape::FRDGBuilderRecorder& RDGBuilderRecorder)
+void ULandscapeTexturePatch::BlendLayer(UE::Landscape::EditLayers::FRenderParams& RenderParams, UE::Landscape::FRDGBuilderRecorder& RDGBuilderRecorder)
 {
 	using namespace UE::Landscape::EditLayers;
 
@@ -270,7 +270,7 @@ void ULandscapeTexturePatch::BlendLayer(FRenderParams& RenderParams, UE::Landsca
 	}
 
 	// Then perform the generic blend : 
-	RenderParams.MergeRenderContext->GenericBlendLayer(BlendParams, RDGBuilderRecorder);
+	RenderParams.MergeRenderContext->GenericBlendLayer(BlendParams, RenderParams, RDGBuilderRecorder);
 }
 
 // Legacy path, which gets the entire heightmap.
@@ -358,7 +358,7 @@ UTextureRenderTarget2D* ULandscapeTexturePatch::RenderLayer_Native(const FLandsc
 	}
 }
 
-UTextureRenderTarget2D* ULandscapeTexturePatch::ApplyToHeightmap(bool bInPerformBlending, FRenderParams* RenderParams, UE::Landscape::FRDGBuilderRecorder& RDGBuilderRecorder, UTextureRenderTarget2D* InCombinedResult,
+UTextureRenderTarget2D* ULandscapeTexturePatch::ApplyToHeightmap(bool bInPerformBlending, UE::Landscape::EditLayers::FRenderParams* RenderParams, UE::Landscape::FRDGBuilderRecorder& RDGBuilderRecorder, UTextureRenderTarget2D* InCombinedResult,
 	const FTransform& LandscapeHeightmapToWorld, bool& bOutDidRenderSomething, ERHIAccess OutputAccess)
 {
 	using namespace UE::Landscape;
@@ -500,7 +500,7 @@ UTextureRenderTarget2D* ULandscapeTexturePatch::ApplyToHeightmap(bool bInPerform
 	return InCombinedResult;
 }
 
-void ULandscapeTexturePatch::ApplyToWeightmap(bool bInPerformBlending, FRenderParams* RenderParams, UE::Landscape::FRDGBuilderRecorder& RDGBuilderRecorder, ULandscapeWeightPatchTextureInfo* PatchInfo,
+void ULandscapeTexturePatch::ApplyToWeightmap(bool bInPerformBlending, UE::Landscape::EditLayers::FRenderParams* RenderParams, UE::Landscape::FRDGBuilderRecorder& RDGBuilderRecorder, ULandscapeWeightPatchTextureInfo* PatchInfo,
 	FTextureResource* InMergedLandscapeTextureResource, int32 LandscapeTextureSliceIndex,
 	const FIntPoint& LandscapeTextureResolution, const FTransform& LandscapeHeightmapToWorld, bool& bOutDidRenderSomething, ERHIAccess OutputAccess)
 {
diff --git a/Engine/Plugins/Experimental/LandscapePatch/Source/LandscapePatch/Public/LandscapeCircleHeightPatch.h b/Engine/Plugins/Experimental/LandscapePatch/Source/LandscapePatch/Public/LandscapeCircleHeightPatch.h
--- a/Engine/Plugins/Experimental/LandscapePatch/Source/LandscapePatch/Public/LandscapeCircleHeightPatch.h
+++ b/Engine/Plugins/Experimental/LandscapePatch/Source/LandscapePatch/Public/LandscapeCircleHeightPatch.h
@@ -43,8 +43,8 @@ public:
 	virtual TArray<FEditLayerRenderItem> GetRenderItems(const UE::Landscape::EditLayers::FMergeContext* InMergeContext) const override;
 	virtual UE::Landscape::EditLayers::ERenderFlags GetRenderFlags(const UE::Landscape::EditLayers::FMergeContext* InMergeContext) const override;
 	virtual bool CanGroupRenderLayerWith(TScriptInterface<ILandscapeEditLayerRenderer> InOtherRenderer) const override;
-	virtual bool RenderLayer(FRenderParams& RenderParams, UE::Landscape::FRDGBuilderRecorder& RDGBuilderRecorder) override;
-	virtual void BlendLayer(FRenderParams& RenderParams, UE::Landscape::FRDGBuilderRecorder& RDGBuilderRecorder) override;
+	virtual bool RenderLayer(UE::Landscape::EditLayers::FRenderParams& RenderParams, UE::Landscape::FRDGBuilderRecorder& RDGBuilderRecorder) override;
+	virtual void BlendLayer(UE::Landscape::EditLayers::FRenderParams& RenderParams, UE::Landscape::FRDGBuilderRecorder& RDGBuilderRecorder) override;
 #endif
 
 	virtual bool CanAffectHeightmap() const override { return !bEditVisibility; }
@@ -73,7 +73,7 @@ protected:
 
 private:
 #if WITH_EDITOR
-	bool ApplyCirclePatch(bool bInPerformBlending, FRenderParams* RenderParams, UE::Landscape::FRDGBuilderRecorder& RDGBuilderRecorder, bool bIsVisibilityLayer,
+	bool ApplyCirclePatch(bool bInPerformBlending, UE::Landscape::EditLayers::FRenderParams* RenderParams, UE::Landscape::FRDGBuilderRecorder& RDGBuilderRecorder, bool bIsVisibilityLayer,
 		FTextureResource* InMergedLandscapeTextureResource, int32 LandscapeTextureSliceIndex,
 		const FIntPoint& DestinationResolution, const FTransform& HeightmapCoordsToWorld, ERHIAccess OutputAccess = ERHIAccess::None);
 #endif // WITH_EDITOR
diff --git a/Engine/Plugins/Experimental/LandscapePatch/Source/LandscapePatch/Public/LandscapePatchManager.h b/Engine/Plugins/Experimental/LandscapePatch/Source/LandscapePatch/Public/LandscapePatchManager.h
--- a/Engine/Plugins/Experimental/LandscapePatch/Source/LandscapePatch/Public/LandscapePatchManager.h
+++ b/Engine/Plugins/Experimental/LandscapePatch/Source/LandscapePatch/Public/LandscapePatchManager.h
@@ -48,7 +48,7 @@ public:
 		UE::Landscape::EditLayers::FEditLayerTargetTypeState& OutEnabledTargetTypeState,
 		TArray<TBitArray<>>& OutTargetLayerGroups) const override {}
 	virtual TArray<UE::Landscape::EditLayers::FEditLayerRenderItem> GetRenderItems(const UE::Landscape::EditLayers::FMergeContext* InMergeContext) const override { return {}; }
-	virtual bool RenderLayer(ILandscapeEditLayerRenderer::FRenderParams& RenderParams, UE::Landscape::FRDGBuilderRecorder& RDGBuilderRecorder) override { return false; };
+	virtual bool RenderLayer(UE::Landscape::EditLayers::FRenderParams& RenderParams, UE::Landscape::FRDGBuilderRecorder& RDGBuilderRecorder) override { return false; };
 #endif // WITH_EDITOR
 
 	// Adds the brush to the given landscape, removing it from any previous one. This differs from SetOwningLandscape
diff --git a/Engine/Plugins/Experimental/LandscapePatch/Source/LandscapePatch/Public/LandscapeTexturePatch.h b/Engine/Plugins/Experimental/LandscapePatch/Source/LandscapePatch/Public/LandscapeTexturePatch.h
--- a/Engine/Plugins/Experimental/LandscapePatch/Source/LandscapePatch/Public/LandscapeTexturePatch.h
+++ b/Engine/Plugins/Experimental/LandscapePatch/Source/LandscapePatch/Public/LandscapeTexturePatch.h
@@ -240,8 +240,8 @@ public:
 	virtual TArray<FEditLayerRenderItem> GetRenderItems(const UE::Landscape::EditLayers::FMergeContext* InMergeContext) const override;
 	virtual UE::Landscape::EditLayers::ERenderFlags GetRenderFlags(const UE::Landscape::EditLayers::FMergeContext* InMergeContext) const override;
 	virtual bool CanGroupRenderLayerWith(TScriptInterface<ILandscapeEditLayerRenderer> InOtherRenderer) const override;
-	virtual bool RenderLayer(FRenderParams& RenderParams, UE::Landscape::FRDGBuilderRecorder& RDGBuilderRecorder) override;
-	virtual void BlendLayer(FRenderParams& RenderParams, UE::Landscape::FRDGBuilderRecorder& RDGBuilderRecorder) override;
+	virtual bool RenderLayer(UE::Landscape::EditLayers::FRenderParams& RenderParams, UE::Landscape::FRDGBuilderRecorder& RDGBuilderRecorder) override;
+	virtual void BlendLayer(UE::Landscape::EditLayers::FRenderParams& RenderParams, UE::Landscape::FRDGBuilderRecorder& RDGBuilderRecorder) override;
 
 	// ULandscapePatchComponent
 	UTextureRenderTarget2D* RenderLayer_Native(const FLandscapeBrushParameters& InParameters, const FTransform& HeightmapToWorld) override;
@@ -661,10 +661,10 @@ private:
 	void TransitionHeightSourceModeInternal(ELandscapeTexturePatchSourceMode OldMode, ELandscapeTexturePatchSourceMode NewMode);
 	FLandscapeHeightPatchConvertToNativeParams GetHeightConvertToNativeParams() const;
 
-	UTextureRenderTarget2D* ApplyToHeightmap(bool bInPerformBlending, FRenderParams* RenderParams, UE::Landscape::FRDGBuilderRecorder& RDGBuilderRecorder, UTextureRenderTarget2D* InCombinedResult,
+	UTextureRenderTarget2D* ApplyToHeightmap(bool bInPerformBlending, UE::Landscape::EditLayers::FRenderParams* RenderParams, UE::Landscape::FRDGBuilderRecorder& RDGBuilderRecorder, UTextureRenderTarget2D* InCombinedResult,
 		const FTransform& LandscapeHeightmapToWorld, bool& bOutDidRenderSomething, ERHIAccess OutputAccess = ERHIAccess::None);
 
-	void ApplyToWeightmap(bool bInPerformBlending, FRenderParams* RenderParams, UE::Landscape::FRDGBuilderRecorder& RDGBuilderRecorder, ULandscapeWeightPatchTextureInfo* PatchInfo,
+	void ApplyToWeightmap(bool bInPerformBlending, UE::Landscape::EditLayers::FRenderParams* RenderParams, UE::Landscape::FRDGBuilderRecorder& RDGBuilderRecorder, ULandscapeWeightPatchTextureInfo* PatchInfo,
 		FTextureResource* InMergedLandscapeTextureResource, int32 LandscapeTextureSliceIndex, const FIntPoint& LandscapeTextureResolution, 
 		const FTransform& LandscapeHeightmapToWorld, bool& bOutDidRenderSomething, ERHIAccess OutputAccess = ERHIAccess::None);
 
diff --git a/Engine/Source/Runtime/Landscape/Classes/LandscapeEditLayer.h b/Engine/Source/Runtime/Landscape/Classes/LandscapeEditLayer.h
--- a/Engine/Source/Runtime/Landscape/Classes/LandscapeEditLayer.h
+++ b/Engine/Source/Runtime/Landscape/Classes/LandscapeEditLayer.h
@@ -201,9 +201,9 @@ public:
 	LANDSCAPE_EDIT_LAYERS_BATCHED_MERGE_EXPERIMENTAL
 	LANDSCAPE_API virtual TArray<UE::Landscape::EditLayers::FEditLayerRenderItem> GetRenderItems(const UE::Landscape::EditLayers::FMergeContext* InMergeContext) const override;
 	LANDSCAPE_EDIT_LAYERS_BATCHED_MERGE_EXPERIMENTAL
-	LANDSCAPE_API virtual bool RenderLayer(FRenderParams& RenderParams, UE::Landscape::FRDGBuilderRecorder& RDGBuilderRecorder) override;
+	LANDSCAPE_API virtual bool RenderLayer(UE::Landscape::EditLayers::FRenderParams& RenderParams, UE::Landscape::FRDGBuilderRecorder& RDGBuilderRecorder) override;
 	LANDSCAPE_EDIT_LAYERS_BATCHED_MERGE_EXPERIMENTAL
-	LANDSCAPE_API virtual void BlendLayer(FRenderParams& RenderParams, UE::Landscape::FRDGBuilderRecorder& RDGBuilderRecorder) override;
+	LANDSCAPE_API virtual void BlendLayer(UE::Landscape::EditLayers::FRenderParams& RenderParams, UE::Landscape::FRDGBuilderRecorder& RDGBuilderRecorder) override;
 	LANDSCAPE_EDIT_LAYERS_BATCHED_MERGE_EXPERIMENTAL
 	LANDSCAPE_API virtual FString GetEditLayerRendererDebugName() const override;
 	//~ End ILandscapeEditLayerRenderer implementation
diff --git a/Engine/Source/Runtime/Landscape/Private/LandscapeBlueprintBrushBase.cpp b/Engine/Source/Runtime/Landscape/Private/LandscapeBlueprintBrushBase.cpp
--- a/Engine/Source/Runtime/Landscape/Private/LandscapeBlueprintBrushBase.cpp
+++ b/Engine/Source/Runtime/Landscape/Private/LandscapeBlueprintBrushBase.cpp
@@ -347,7 +347,7 @@ TArray<UE::Landscape::EditLayers::FEditLayerRenderItem> ALandscapeBlueprintBrush
 	return { FEditLayerRenderItem(EnabledTargetTypeState, InputWorldArea, OutputWorldArea, /*bModifyExistingWeightmapsOnly = */false)};
 }
 
-bool ALandscapeBlueprintBrushBase::RenderLayer(ILandscapeEditLayerRenderer::FRenderParams& RenderParams, UE::Landscape::FRDGBuilderRecorder& RDGBuilderRecorder)
+bool ALandscapeBlueprintBrushBase::RenderLayer(UE::Landscape::EditLayers::FRenderParams& RenderParams, UE::Landscape::FRDGBuilderRecorder& RDGBuilderRecorder)
 {
 	using namespace UE::Landscape;
 
diff --git a/Engine/Source/Runtime/Landscape/Private/LandscapeEditLayerMergeRenderContext.cpp b/Engine/Source/Runtime/Landscape/Private/LandscapeEditLayerMergeRenderContext.cpp
--- a/Engine/Source/Runtime/Landscape/Private/LandscapeEditLayerMergeRenderContext.cpp
+++ b/Engine/Source/Runtime/Landscape/Private/LandscapeEditLayerMergeRenderContext.cpp
@@ -198,6 +198,30 @@ IMPLEMENT_GLOBAL_SHADER(FLandscapeEditLayersWeightmapsMergeEditLayerPS, "/Engine
 
 #if WITH_EDITOR
 
+
+FRenderParams::FRenderParams(
+	FMergeRenderContext* InMergeRenderContext,
+	const TArrayView<FName>& InTargetLayerGroupLayerNames,
+	const TArrayView<ULandscapeLayerInfoObject*>& InTargetLayerGroupLayerInfos,
+	const FEditLayerRendererState& InRendererState,
+	const TArrayView<FComponentMergeRenderInfo>& InSortedComponentMergeRenderInfos,
+	const FTransform& InRenderAreaWorldTransform,
+	const FIntRect& InRenderAreaSectionRect,
+	int32 InNumSuccessfulRenderLayerStepsUntilBlendLayerStep)
+	: MergeRenderContext(InMergeRenderContext)
+	, TargetLayerGroupLayerNames(InTargetLayerGroupLayerNames)
+	, TargetLayerGroupLayerInfos(InTargetLayerGroupLayerInfos)
+	, RendererState(InRendererState)
+	, SortedComponentMergeRenderInfos(InSortedComponentMergeRenderInfos)
+	, RenderAreaWorldTransform(InRenderAreaWorldTransform)
+	, RenderAreaSectionRect(InRenderAreaSectionRect)
+	, NumSuccessfulRenderLayerStepsUntilBlendLayerStep(InNumSuccessfulRenderLayerStepsUntilBlendLayerStep)
+{
+}
+
+
+// ----------------------------------------------------------------------------------
+
 bool FMergeRenderBatch::operator<(const FMergeRenderBatch& InOther) const
 {
 	// Sort by coordinates for making debugging more "logical" : 
@@ -661,10 +685,9 @@ void FMergeRenderContext::RenderExpandedRenderTarget(UE::Landscape::FRDGBuilderR
 	RDGBuilderRecorder.EnqueueRDGCommand(RDGCommand, RDGExternalTextureAccessFinalList);
 }
 
-void FMergeRenderContext::GenericBlendLayer(const FBlendParams& InBlendParams, UE::Landscape::FRDGBuilderRecorder& RDGBuilderRecorder)
+void FMergeRenderContext::GenericBlendLayer(const FBlendParams& InBlendParams, FRenderParams& RenderParams, UE::Landscape::FRDGBuilderRecorder& RDGBuilderRecorder)
 {
 	using namespace UE::Landscape;
-	using namespace UE::Landscape::EditLayers;
 
 	const FMergeRenderBatch* RenderBatch = GetCurrentRenderBatch();
 	check(RenderBatch != nullptr);
@@ -720,8 +743,11 @@ void FMergeRenderContext::GenericBlendLayer(const FBlendParams& InBlendParams, U
 	else
 	{
 		ForEachTargetLayer(RenderBatch->TargetLayerBitIndices,
-			[WriteRT, CurrentLayerReadRT, PreviousLayersReadRT, &RenderBatch, &InBlendParams, &RDGBuilderRecorder](int32 InTargetLayerIndex, const FName& InTargetLayerName, ULandscapeLayerInfoObject* InWeightmapLayerInfo)
+			[WriteRT, CurrentLayerReadRT, PreviousLayersReadRT, &RenderBatch, &RenderParams, &InBlendParams, &RDGBuilderRecorder](int32 InTargetLayerIndex, const FName& InTargetLayerName, ULandscapeLayerInfoObject* InWeightmapLayerInfo)
 			{
+				int32 TargetLayerIndexInGroup = RenderParams.TargetLayerGroupLayerNames.Find(InTargetLayerName);
+				check(TargetLayerIndexInGroup != INDEX_NONE);
+
 				// By default, use passthrough mode so that each layer gets at least copied into the WriteRT :
 				const FWeightmapBlendParams* TargetLayerBlendParams = &FWeightmapBlendParams::GetDefaultPassthroughBlendParams();
 				// Special case for visibility which is always "Additive" :
@@ -737,7 +763,7 @@ void FMergeRenderContext::GenericBlendLayer(const FBlendParams& InBlendParams, U
 				// TODO [jonathan.bard] : we could render several layers at once via MRT (up to MaxSimultaneousRenderTargets)
 				auto RDGCommand =
 					[ TargetLayerBlendParams = *TargetLayerBlendParams
-					, InTargetLayerIndex
+					, TargetLayerIndexInGroup
 					, InTargetLayerName
 					, OutputResource = WriteRT->GetRenderTarget2DArray()->GetResource()
 					, OutputResourceName = WriteRT->GetDebugName()
@@ -753,8 +779,8 @@ void FMergeRenderContext::GenericBlendLayer(const FBlendParams& InBlendParams, U
 						FRDGTextureSRVRef PreviousEditLayersTextureSRVRef = GraphBuilder.CreateSRV(FRDGTextureSRVDesc::Create(PreviousEditLayersTextureRef));
 
 						FLandscapeEditLayersWeightmapsMergeEditLayerPS::FParameters* PSParams = GraphBuilder.AllocParameters<FLandscapeEditLayersWeightmapsMergeEditLayerPS::FParameters>();
-						PSParams->RenderTargets[0] = FRenderTargetBinding(OutputTextureRef, ERenderTargetLoadAction::ENoAction, /*InMipIndex = */0, /*InArraySlice = */InTargetLayerIndex);
-						PSParams->InTargetLayerIndex = InTargetLayerIndex;
+						PSParams->RenderTargets[0] = FRenderTargetBinding(OutputTextureRef, ERenderTargetLoadAction::ENoAction, /*InMipIndex = */0, /*InArraySlice = */TargetLayerIndexInGroup);
+						PSParams->InTargetLayerIndex = TargetLayerIndexInGroup;
 						PSParams->InEditLayerTargetLayerBlendMode = static_cast<uint32>(TargetLayerBlendParams.BlendMode);
 						PSParams->InEditLayerAlpha = TargetLayerBlendParams.Alpha;
 						PSParams->InCurrentEditLayerWeightmaps = CurrentEditLayerTextureSRVRef;
@@ -784,7 +810,6 @@ void FMergeRenderContext::GenericBlendLayer(const FBlendParams& InBlendParams, U
 void FMergeRenderContext::Render(TFunction<void(const FOnRenderBatchTargetGroupDoneParams& /*InParams*/, UE::Landscape::FRDGBuilderRecorder& /*RDGBuilderRecorder*/)> OnBatchTargetGroupDone)
 {
 	using namespace UE::Landscape;
-	using namespace UE::Landscape::EditLayers;
 
 	TRACE_CPUPROFILER_EVENT_SCOPE(FMergeRenderContext::Render);
 
@@ -940,13 +965,13 @@ void FMergeRenderContext::Render(TFunction<void(const FOnRenderBatchTargetGroupD
 				{
 					// We start a new render layer group, so let's start tracking the number of successful RenderLayer steps within this group :
 					NumSuccessfulRenderLayerStepsUntilBlendLayerStep = 0;
-					ILandscapeEditLayerRenderer::FRenderParams RenderParams(this, TargetLayerGroupLayerNames, TargetLayerGroupLayerInfos, RenderStep.RendererState, SortedComponentMergeRenderInfos, 
+					FRenderParams RenderParams(this, TargetLayerGroupLayerNames, TargetLayerGroupLayerInfos, RenderStep.RendererState, SortedComponentMergeRenderInfos,
 						RenderAreaWorldTransform, RenderAreaSectionRect, NumSuccessfulRenderLayerStepsUntilBlendLayerStep);
 					Renderer->BeginRenderLayerGroup(RenderParams, RDGBuilderRecorder);
 				}
                 else
                 {
-					ILandscapeEditLayerRenderer::FRenderParams RenderParams(this, TargetLayerGroupLayerNames, TargetLayerGroupLayerInfos, RenderStep.RendererState, SortedComponentMergeRenderInfos,
+					FRenderParams RenderParams(this, TargetLayerGroupLayerNames, TargetLayerGroupLayerInfos, RenderStep.RendererState, SortedComponentMergeRenderInfos,
 						RenderAreaWorldTransform, RenderAreaSectionRect, NumSuccessfulRenderLayerStepsUntilBlendLayerStep);
 					Renderer->EndRenderLayerGroup(RenderParams, RDGBuilderRecorder);
                 }
@@ -982,7 +1007,7 @@ void FMergeRenderContext::Render(TFunction<void(const FOnRenderBatchTargetGroupD
 						}
 					});
 		
-				ILandscapeEditLayerRenderer::FRenderParams RenderParams(this, TargetLayerGroupLayerNames, TargetLayerGroupLayerInfos, RenderStep.RendererState, SortedComponentMergeRenderInfos, 
+				FRenderParams RenderParams(this, TargetLayerGroupLayerNames, TargetLayerGroupLayerInfos, RenderStep.RendererState, SortedComponentMergeRenderInfos, 
 					RenderAreaWorldTransform, RenderAreaSectionRect, NumSuccessfulRenderLayerStepsUntilBlendLayerStep);
 				auto RenderOrBlend = [&RDGBuilderRecorder, &RenderStepProfilingEventName, &NumSuccessfulRenderLayerStepsUntilBlendLayerStep, &RenderParams, bIsSeparateBlendStep, Renderer]()
 					{
diff --git a/Engine/Source/Runtime/Landscape/Private/LandscapeEditLayerRenderer.cpp b/Engine/Source/Runtime/Landscape/Private/LandscapeEditLayerRenderer.cpp
deleted file mode 100644
diff --git a/Engine/Source/Runtime/Landscape/Private/LandscapeEditLayerRendererPrivate.cpp b/Engine/Source/Runtime/Landscape/Private/LandscapeEditLayerRendererPrivate.cpp
--- a/Engine/Source/Runtime/Landscape/Private/LandscapeEditLayerRendererPrivate.cpp
+++ b/Engine/Source/Runtime/Landscape/Private/LandscapeEditLayerRendererPrivate.cpp
@@ -44,7 +44,7 @@ FString ULandscapeDefaultEditLayerRenderer::GetEditLayerRendererDebugName() cons
 	return TEXT("Default");
 }
 
-bool ULandscapeDefaultEditLayerRenderer::RenderLayer(FRenderParams& RenderParams, UE::Landscape::FRDGBuilderRecorder& RDGBuilderRecorder)
+bool ULandscapeDefaultEditLayerRenderer::RenderLayer(UE::Landscape::EditLayers::FRenderParams& RenderParams, UE::Landscape::FRDGBuilderRecorder& RDGBuilderRecorder)
 {
 	using namespace UE::Landscape::EditLayers;
 
diff --git a/Engine/Source/Runtime/Landscape/Private/LandscapeEditLayerRendererPrivate.h b/Engine/Source/Runtime/Landscape/Private/LandscapeEditLayerRendererPrivate.h
--- a/Engine/Source/Runtime/Landscape/Private/LandscapeEditLayerRendererPrivate.h
+++ b/Engine/Source/Runtime/Landscape/Private/LandscapeEditLayerRendererPrivate.h
@@ -32,7 +32,7 @@ public:
 	virtual TArray<UE::Landscape::EditLayers::FEditLayerRenderItem> GetRenderItems(const UE::Landscape::EditLayers::FMergeContext* InMergeContext) const override;
 	virtual FString GetEditLayerRendererDebugName() const override;
 	virtual UE::Landscape::EditLayers::ERenderFlags GetRenderFlags(const UE::Landscape::EditLayers::FMergeContext* InMergeContext) const override { return UE::Landscape::EditLayers::ERenderFlags::RenderMode_Recorded; }
-	virtual bool RenderLayer(FRenderParams& RenderParams, UE::Landscape::FRDGBuilderRecorder& RDGBuilderRecorder) override;
+	virtual bool RenderLayer(UE::Landscape::EditLayers::FRenderParams& RenderParams, UE::Landscape::FRDGBuilderRecorder& RDGBuilderRecorder) override;
 	//~ End ILandscapeEditLayerRenderer implementation
 #endif // WITH_EDITOR
 };
@@ -58,7 +58,7 @@ public:
 	virtual TArray<UE::Landscape::EditLayers::FEditLayerRenderItem> GetRenderItems(const UE::Landscape::EditLayers::FMergeContext* InMergeContext) const override;
 	virtual FString GetEditLayerRendererDebugName() const override;
 	virtual UE::Landscape::EditLayers::ERenderFlags GetRenderFlags(const UE::Landscape::EditLayers::FMergeContext* InMergeContext) const override { return UE::Landscape::EditLayers::ERenderFlags::RenderMode_Recorded; }
-	virtual bool RenderLayer(FRenderParams& RenderParams, UE::Landscape::FRDGBuilderRecorder& RDGBuilderRecorder) override;
+	virtual bool RenderLayer(UE::Landscape::EditLayers::FRenderParams& RenderParams, UE::Landscape::FRDGBuilderRecorder& RDGBuilderRecorder) override;
 	//~ End ILandscapeEditLayerRenderer implementation
 #endif // WITH_EDITOR
 };
@@ -82,7 +82,7 @@ public:
 	virtual TArray<UE::Landscape::EditLayers::FEditLayerRenderItem> GetRenderItems(const UE::Landscape::EditLayers::FMergeContext* InMergeContext) const override;
 	virtual FString GetEditLayerRendererDebugName() const override;
 	virtual UE::Landscape::EditLayers::ERenderFlags GetRenderFlags(const UE::Landscape::EditLayers::FMergeContext* InMergeContext) const override { return UE::Landscape::EditLayers::ERenderFlags::RenderMode_Recorded; }
-	virtual bool RenderLayer(FRenderParams& RenderParams, UE::Landscape::FRDGBuilderRecorder& RDGBuilderRecorder) override;
+	virtual bool RenderLayer(UE::Landscape::EditLayers::FRenderParams& RenderParams, UE::Landscape::FRDGBuilderRecorder& RDGBuilderRecorder) override;
 	//~ End ILandscapeEditLayerRenderer implementation
 
 	TBitArray<> GatherWeightBlendedWeightmapLayerBitIndices(const UE::Landscape::EditLayers::FMergeContext* InMergeContext) const;
diff --git a/Engine/Source/Runtime/Landscape/Private/LandscapeEditLayers.cpp b/Engine/Source/Runtime/Landscape/Private/LandscapeEditLayers.cpp
--- a/Engine/Source/Runtime/Landscape/Private/LandscapeEditLayers.cpp
+++ b/Engine/Source/Runtime/Landscape/Private/LandscapeEditLayers.cpp
@@ -5497,7 +5497,7 @@ TArray<UE::Landscape::EditLayers::FEditLayerRenderItem> ULandscapeEditLayerPersi
 	return RenderItems;
 }
 
-bool ULandscapeEditLayerPersistent::RenderLayer(FRenderParams& RenderParams, UE::Landscape::FRDGBuilderRecorder& RDGBuilderRecorder)
+bool ULandscapeEditLayerPersistent::RenderLayer(UE::Landscape::EditLayers::FRenderParams& RenderParams, UE::Landscape::FRDGBuilderRecorder& RDGBuilderRecorder)
 {
 	using namespace UE::Landscape;
 	using namespace UE::Landscape::Private;
@@ -5795,7 +5795,7 @@ bool ULandscapeEditLayerPersistent::RenderLayer(FRenderParams& RenderParams, UE:
 	return true;
 }
 
-void ULandscapeEditLayerPersistent::BlendLayer(FRenderParams& RenderParams, UE::Landscape::FRDGBuilderRecorder& RDGBuilderRecorder)
+void ULandscapeEditLayerPersistent::BlendLayer(UE::Landscape::EditLayers::FRenderParams& RenderParams, UE::Landscape::FRDGBuilderRecorder& RDGBuilderRecorder)
 {
 	using namespace UE::Landscape::EditLayers;
 	using namespace UE::Landscape::Private;
@@ -5845,7 +5845,7 @@ void ULandscapeEditLayerPersistent::BlendLayer(FRenderParams& RenderParams, UE::
 	}
 
 	// Then perform the generic blend : 
-	RenderParams.MergeRenderContext->GenericBlendLayer(BlendParams, RDGBuilderRecorder);
+	RenderParams.MergeRenderContext->GenericBlendLayer(BlendParams, RenderParams, RDGBuilderRecorder);
 }
 
 FString ULandscapeEditLayerPersistent::GetEditLayerRendererDebugName() const
@@ -6352,7 +6352,7 @@ namespace UE::Landscape::EditLayers::Private
 #endif // ENABLE_VISUAL_LOG
 } // UE::Landscape::EditLayers::Private
 
-bool ULandscapeHeightmapNormalsEditLayerRenderer::RenderLayer(FRenderParams& RenderParams, UE::Landscape::FRDGBuilderRecorder& RDGBuilderRecorder)
+bool ULandscapeHeightmapNormalsEditLayerRenderer::RenderLayer(UE::Landscape::EditLayers::FRenderParams& RenderParams, UE::Landscape::FRDGBuilderRecorder& RDGBuilderRecorder)
 {
 	using namespace UE::Landscape;
 	using namespace UE::Landscape::EditLayers;
@@ -6419,7 +6419,7 @@ bool ULandscapeHeightmapNormalsEditLayerRenderer::RenderLayer(FRenderParams& Ren
 	return true;
 }
 
-bool ULandscapeWeightmapWeightBlendedLayersRenderer::RenderLayer(FRenderParams& RenderParams, UE::Landscape::FRDGBuilderRecorder& RDGBuilderRecorder)
+bool ULandscapeWeightmapWeightBlendedLayersRenderer::RenderLayer(UE::Landscape::EditLayers::FRenderParams& RenderParams, UE::Landscape::FRDGBuilderRecorder& RDGBuilderRecorder)
 {
 	using namespace UE::Landscape;
 	using namespace UE::Landscape::EditLayers;
diff --git a/Engine/Source/Runtime/Landscape/Public/LandscapeBlueprintBrushBase.h b/Engine/Source/Runtime/Landscape/Public/LandscapeBlueprintBrushBase.h
--- a/Engine/Source/Runtime/Landscape/Public/LandscapeBlueprintBrushBase.h
+++ b/Engine/Source/Runtime/Landscape/Public/LandscapeBlueprintBrushBase.h
@@ -108,7 +108,7 @@ public:
 	LANDSCAPE_EDIT_LAYERS_BATCHED_MERGE_EXPERIMENTAL
 	LANDSCAPE_API virtual TArray<UE::Landscape::EditLayers::FEditLayerRenderItem> GetRenderItems(const UE::Landscape::EditLayers::FMergeContext* InMergeContext) const override;
 	LANDSCAPE_EDIT_LAYERS_BATCHED_MERGE_EXPERIMENTAL
-	LANDSCAPE_API virtual bool RenderLayer(ILandscapeEditLayerRenderer::FRenderParams& RenderParams, UE::Landscape::FRDGBuilderRecorder& RDGBuilderRecorder) override;
+	LANDSCAPE_API virtual bool RenderLayer(UE::Landscape::EditLayers::FRenderParams& RenderParams, UE::Landscape::FRDGBuilderRecorder& RDGBuilderRecorder) override;
 	LANDSCAPE_EDIT_LAYERS_BATCHED_MERGE_EXPERIMENTAL
 	LANDSCAPE_API virtual FString GetEditLayerRendererDebugName() const override;
 	//~ End ILandscapeEditLayerRenderer implementation
diff --git a/Engine/Source/Runtime/Landscape/Public/LandscapeEditLayerMergeRenderContext.h b/Engine/Source/Runtime/Landscape/Public/LandscapeEditLayerMergeRenderContext.h
--- a/Engine/Source/Runtime/Landscape/Public/LandscapeEditLayerMergeRenderContext.h
+++ b/Engine/Source/Runtime/Landscape/Public/LandscapeEditLayerMergeRenderContext.h
@@ -30,6 +30,12 @@ namespace UE::Landscape
 } // namespace UE::Landscape
 
 
+namespace UE::Landscape::EditLayers
+{
+	class FMergeRenderContext;
+} //namespace UE::Landscape::EditLayers
+
+
 // ----------------------------------------------------------------------------------
 
 namespace UE::Landscape::EditLayers
@@ -221,6 +227,54 @@ struct FComponentMergeRenderInfo
 	bool operator < (const FComponentMergeRenderInfo& InOther) const;
 };
 
+/** 
+ * Struct passed to ILandscapeEditLayerRenderer's render functions. This would normally be an inner struct of ILandscapeEditLayerRenderer but 
+ *  this prevents forward-declaration so we leave in the UE::Landscape::EditLayers namespace instead. 
+ */
+struct FRenderParams
+{
+	FRenderParams(FMergeRenderContext* InMergeRenderContext,
+		const TArrayView<FName>& InTargetLayerGroupLayerNames,
+		const TArrayView<ULandscapeLayerInfoObject*>& InTargetLayerGroupLayerInfos,
+		const FEditLayerRendererState& InRendererState,
+		const TArrayView<FComponentMergeRenderInfo>& InSortedComponentMergeRenderInfos,
+		const FTransform& InRenderAreaWorldTransform,
+		const FIntRect& InRenderAreaSectionRect,
+		int32 InNumSuccessfulRenderLayerStepsUntilBlendLayerStep);
+
+	/** Merge context */
+	FMergeRenderContext* MergeRenderContext = nullptr;
+
+	/** List of target layers being involved in this step */
+	TArray<FName> TargetLayerGroupLayerNames;
+
+	/** List of target layer info objects being involved in this step (same size as TargetLayerGroupLayerNames) */
+	TArray<ULandscapeLayerInfoObject*> TargetLayerGroupLayerInfos;
+
+	/** Full state for the renderer involved in this step. This allows to retrieve the exact state of this renderer (e.g. enabled weightmaps, which can be different than the target layer group, in that
+	 target layers A, B and C might belong to the same group but this renderer actually only has A enabled). This is therefore the renderer's responsibility to check that a given target layer from the
+	 target layer group is effectively enabled. */
+	FEditLayerRendererState RendererState;
+
+	/** List of components (with additional info) to render */
+	TArray<FComponentMergeRenderInfo> SortedComponentMergeRenderInfos;
+
+	// TODO [jonathan.bard] Verify that scale is correct
+	/** World transform that corresponds to the origin (bottom left corner) of the render area. The scale corresponds to the size of each quad in the landscape. */
+	FTransform RenderAreaWorldTransform;
+
+	/** SectionRect (i.e. landscape vertex coordinates, in landscape space) that corresponds to this render area */
+	FIntRect RenderAreaSectionRect;
+
+	/** When separate blend is enabled, tracks how many RenderLayer calls have succeeded yet (valid until the BlendLayer step occurs) */
+	int32 NumSuccessfulRenderLayerStepsUntilBlendLayerStep = 0;
+
+	UE_DEPRECATED(5.6, "Renamed : use TargetLayerGroupLayerNames")
+	TArray<FName> RenderGroupTargetLayerNames;
+	UE_DEPRECATED(5.6, "Renamed : use TargetLayerGroupLayerInfos")
+	TArray<ULandscapeLayerInfoObject*> RenderGroupTargetLayerInfos;
+};
+
 
 // ----------------------------------------------------------------------------------
 
@@ -339,7 +393,7 @@ public:
 	 * @param RDGBuilderRecorder recorder to append operations to a single FRDGBuilder, if the the recorder in "recording" mode. In "immediate" mode, a render command will be enqueued immediately
 	 */
 	LANDSCAPE_EDIT_LAYERS_BATCHED_MERGE_EXPERIMENTAL
-	LANDSCAPE_API void GenericBlendLayer(const FBlendParams& InBlendParams, UE::Landscape::FRDGBuilderRecorder& RDGBuilderRecorder);
+	LANDSCAPE_API void GenericBlendLayer(const FBlendParams& InBlendParams, FRenderParams& RenderParams, UE::Landscape::FRDGBuilderRecorder& RDGBuilderRecorder);
 
 	const TBitArray<>& GetFinalTargetLayerBitIndices() const 
 	{ 
diff --git a/Engine/Source/Runtime/Landscape/Public/LandscapeEditLayerRenderer.h b/Engine/Source/Runtime/Landscape/Public/LandscapeEditLayerRenderer.h
--- a/Engine/Source/Runtime/Landscape/Public/LandscapeEditLayerRenderer.h
+++ b/Engine/Source/Runtime/Landscape/Public/LandscapeEditLayerRenderer.h
@@ -29,6 +29,7 @@ namespace UE::Landscape
 namespace UE::Landscape::EditLayers
 {
 	struct FComponentMergeRenderInfo;
+	struct FRenderParams;
 	class FMergeRenderContext;
 } // namespace UE::Landscape::EditLayers
 
@@ -261,49 +262,7 @@ class ILandscapeEditLayerRenderer
 
 #if WITH_EDITOR
 public:
-	struct FRenderParams
-	{
-		FRenderParams(UE::Landscape::EditLayers::FMergeRenderContext* InMergeRenderContext, 
-			const TArrayView<FName>& InTargetLayerGroupLayerNames,
-			const TArrayView<ULandscapeLayerInfoObject*>& InTargetLayerGroupLayerInfos,
-			const UE::Landscape::EditLayers::FEditLayerRendererState& InRendererState, 
-			const TArrayView<UE::Landscape::EditLayers::FComponentMergeRenderInfo>& InSortedComponentMergeRenderInfos, 
-			const FTransform& InRenderAreaWorldTransform, 
-			const FIntRect& InRenderAreaSectionRect,
-			int32 InNumSuccessfulRenderLayerStepsUntilBlendLayerStep);
-
-		/** Merge context */
-		UE::Landscape::EditLayers::FMergeRenderContext* MergeRenderContext = nullptr;
-
-		/** List of target layers being involved in this step */
-		TArray<FName> TargetLayerGroupLayerNames;
-
-		/** List of target layer info objects being involved in this step (same size as TargetLayerGroupLayerNames) */
-		TArray<ULandscapeLayerInfoObject*> TargetLayerGroupLayerInfos;
-
-		/** Full state for the renderer involved in this step. This allows to retrieve the exact state of this renderer (e.g. enabled weightmaps, which can be different than the target layer group, in that 
-		 target layers A, B and C might belong to the same group but this renderer actually only has A enabled). This is therefore the renderer's responsibility to check that a given target layer from the 
-		 target layer group is effectively enabled. */
-		UE::Landscape::EditLayers::FEditLayerRendererState RendererState;
-
-		/** List of components (with additional info) to render */
-		TArray<UE::Landscape::EditLayers::FComponentMergeRenderInfo> SortedComponentMergeRenderInfos;
-
-		// TODO [jonathan.bard] Verify that scale is correct
-		/** World transform that corresponds to the origin (bottom left corner) of the render area. The scale corresponds to the size of each quad in the landscape. */
-		FTransform RenderAreaWorldTransform;
-
-		/** SectionRect (i.e. landscape vertex coordinates, in landscape space) that corresponds to this render area */
-		FIntRect RenderAreaSectionRect;
-
-		/** When separate blend is enabled, tracks how many RenderLayer calls have succeeded yet (valid until the BlendLayer step occurs) */
-		int32 NumSuccessfulRenderLayerStepsUntilBlendLayerStep = 0;
-
-		UE_DEPRECATED(5.6, "Renamed : use TargetLayerGroupLayerNames")
-		TArray<FName> RenderGroupTargetLayerNames;
-		UE_DEPRECATED(5.6, "Renamed : use TargetLayerGroupLayerInfos")
-		TArray<ULandscapeLayerInfoObject*> RenderGroupTargetLayerInfos;
-	};
+	using FRenderParams UE_DEPRECATED(5.6, "Use UE::Landscape::EditLayers::FRenderParams") = UE::Landscape::EditLayers::FRenderParams;
 
 	/**
 	 * GetRendererStateInfo retrieves the current state of this renderer (what it can and does render, as well as how to group target layers together), and part of this will then be mutable for the duration of the merge. 
@@ -377,7 +336,7 @@ public:
 	 * @return true if anything was rendered
 	*/
 	LANDSCAPE_EDIT_LAYERS_BATCHED_MERGE_EXPERIMENTAL
-	LANDSCAPE_API virtual bool RenderLayer(ILandscapeEditLayerRenderer::FRenderParams& RenderParams, UE::Landscape::FRDGBuilderRecorder& RDGBuilderRecorder) 
+	LANDSCAPE_API virtual bool RenderLayer(UE::Landscape::EditLayers::FRenderParams& RenderParams, UE::Landscape::FRDGBuilderRecorder& RDGBuilderRecorder)
 		PURE_VIRTUAL(ILandscapeEditLayerRenderer::RenderLayer, return false; );
 
 	/**
@@ -390,7 +349,7 @@ public:
 	 * @param RDGBuilderRecorder (see RenderLayer)
 	*/
 	LANDSCAPE_EDIT_LAYERS_BATCHED_MERGE_EXPERIMENTAL
-	LANDSCAPE_API virtual void BlendLayer(ILandscapeEditLayerRenderer::FRenderParams& RenderParams, UE::Landscape::FRDGBuilderRecorder& RDGBuilderRecorder) {}
+	LANDSCAPE_API virtual void BlendLayer(UE::Landscape::EditLayers::FRenderParams& RenderParams, UE::Landscape::FRDGBuilderRecorder& RDGBuilderRecorder) {}
 
 	/**
 	 * Return whether this renderer's RenderLayer operation can be grouped with one of the previous renderers in the render layer group. 
@@ -410,7 +369,7 @@ public:
 	 * @param RenderParams contains parameters necessary for this render (affected components, merge context, etc.)
 	 */
 	LANDSCAPE_EDIT_LAYERS_BATCHED_MERGE_EXPERIMENTAL
-	LANDSCAPE_API virtual void BeginRenderLayerGroup(ILandscapeEditLayerRenderer::FRenderParams& RenderParams, UE::Landscape::FRDGBuilderRecorder& RDGBuilderRecorder) {}
+	LANDSCAPE_API virtual void BeginRenderLayerGroup(UE::Landscape::EditLayers::FRenderParams& RenderParams, UE::Landscape::FRDGBuilderRecorder& RDGBuilderRecorder) {}
 
 	/**
 	 * Called after the last call to RenderLayer, on the last renderer of a render layer group. 
@@ -419,7 +378,7 @@ public:
 	 * @param RenderParams contains parameters necessary for this render (affected components, merge context, etc.)
 	 */
 	LANDSCAPE_EDIT_LAYERS_BATCHED_MERGE_EXPERIMENTAL
-	LANDSCAPE_API virtual void EndRenderLayerGroup(ILandscapeEditLayerRenderer::FRenderParams& RenderParams, UE::Landscape::FRDGBuilderRecorder& RDGBuilderRecorder) {}
+	LANDSCAPE_API virtual void EndRenderLayerGroup(UE::Landscape::EditLayers::FRenderParams& RenderParams, UE::Landscape::FRDGBuilderRecorder& RDGBuilderRecorder) {}
 
 #endif // WITH_EDITOR
 };
