Modify cvar "landscape.RuntimeVirtualTextureRenderWithQuad" behaviour so that a value of 2 will override the "Virtual Texture Render With Quad" setting on landscape actors to force off.
#rnx

diff --git a/Engine/Source/Runtime/Landscape/Private/LandscapeRender.cpp b/Engine/Source/Runtime/Landscape/Private/LandscapeRender.cpp
--- a/Engine/Source/Runtime/Landscape/Private/LandscapeRender.cpp
+++ b/Engine/Source/Runtime/Landscape/Private/LandscapeRender.cpp
@@ -227,11 +227,12 @@ static FAutoConsoleVariableRef CLandscapeAllowNanitePerClusterDisplacementDisabl
 	ECVF_RenderThreadSafe
 );
 
-bool GLandscapeRuntimeVirtualTextureRenderWithQuad = false;
+int32 GLandscapeRuntimeVirtualTextureRenderWithQuad = 0;
 FAutoConsoleVariableRef CVarLandscapeRuntimeVirtualTextureRenderWithQuad(
 	TEXT("landscape.RuntimeVirtualTextureRenderWithQuad"),
 	GLandscapeRuntimeVirtualTextureRenderWithQuad,
-	TEXT("Force use of a single quad when rendering landscape to runtime virtual texture pages"),
+	TEXT("Force whether to use of a single quad when rendering landscape to runtime virtual texture pages\n")
+	TEXT("0 takes setting from landscape, 1 forces on, 2 forces off."),
 	FConsoleVariableDelegate::CreateStatic(&OnCVarNeedingRenderStateInvalidationChanged),
 	ECVF_RenderThreadSafe | ECVF_Scalability
 );
@@ -1553,7 +1554,9 @@ FLandscapeComponentSceneProxy::FLandscapeComponentSceneProxy(ULandscapeComponent
 	LODSettings.LastLODScreenSizeSquared = LODScreenRatioSquared[LastLOD];
 	LODSettings.ForcedLOD = ForcedLOD;
 
-	const bool bVirtualTextureRenderWithQuad = InComponent->GetLandscapeProxy()->bVirtualTextureRenderWithQuad || GLandscapeRuntimeVirtualTextureRenderWithQuad;
+	bool bVirtualTextureRenderWithQuad = InComponent->GetLandscapeProxy()->bVirtualTextureRenderWithQuad;
+	bVirtualTextureRenderWithQuad |= GLandscapeRuntimeVirtualTextureRenderWithQuad == 1;
+	bVirtualTextureRenderWithQuad &= GLandscapeRuntimeVirtualTextureRenderWithQuad != 2;
 	const bool bVirtualTextureRenderWithQuadHQ = InComponent->GetLandscapeProxy()->bVirtualTextureRenderWithQuadHQ;
 	VirtualTexturePerPixelHeight = bVirtualTextureRenderWithQuad ? bVirtualTextureRenderWithQuadHQ ? 2 : 1 : 0;
 
