[WorldBuilding] [Landscape] Add a RuntimeLevel property to UActorContainer prior to duplication of Unsaved Actors, this fixes broken LandscapeGuid mappings when LandscapeProxy is part of a unsaved runtime cell failing to find its WorldPartition
- Landscape code that relies on WorldPartition being found is called from ALandscapeProxy::PostLoad (where duplicated actor is part of a UActorContainer) in ChangeLandscapeGuidIfObjectIsInstanced
- This was causing Landscape and its proxies to have different Guids in PIE and creating multiple LandscapeInfos (found in a PCG map)

#rnx

diff --git a/Engine/Source/Runtime/Engine/Classes/Engine/Level.h b/Engine/Source/Runtime/Engine/Classes/Engine/Level.h
--- a/Engine/Source/Runtime/Engine/Classes/Engine/Level.h
+++ b/Engine/Source/Runtime/Engine/Classes/Engine/Level.h
@@ -78,6 +78,12 @@ class ENGINE_API UActorContainer : public UObject
 public:
 	UPROPERTY()
 	TMap<FName, TObjectPtr<AActor>> Actors;
+
+#if WITH_EDITORONLY_DATA
+	// Can be used to resolve the Actor Level while they are outered to the UActorContainer in a PIE duplication
+	UPROPERTY()
+	TObjectPtr<ULevel> RuntimeLevel = nullptr;
+#endif
 };
 
 USTRUCT()
diff --git a/Engine/Source/Runtime/Engine/Private/WorldPartition/WorldPartitionLevelStreamingDynamic.cpp b/Engine/Source/Runtime/Engine/Private/WorldPartition/WorldPartitionLevelStreamingDynamic.cpp
--- a/Engine/Source/Runtime/Engine/Private/WorldPartition/WorldPartitionLevelStreamingDynamic.cpp
+++ b/Engine/Source/Runtime/Engine/Private/WorldPartition/WorldPartitionLevelStreamingDynamic.cpp
@@ -422,6 +422,8 @@ bool UWorldPartitionLevelStreamingDynamic::IssueLoadRequests()
 	// Duplicate unsaved actors
 	if (UnsavedActorsContainer)
 	{
+		// Set RuntimeLevel on container temporarily (for duplication) so that external code can find the Level in PostLoad code 
+		TGuardValue<TObjectPtr<ULevel>> ScopeRuntimeLevel(UnsavedActorsContainer->RuntimeLevel, RuntimeLevel);
 		FObjectDuplicationParameters Parameters(UnsavedActorsContainer, RuntimeLevel->GetPackage());
 		Parameters.DestClass = UnsavedActorsContainer->GetClass();
 		Parameters.FlagMask = RF_AllFlags & ~(RF_MarkAsRootSet | RF_MarkAsNative | RF_HasExternalPackage);
diff --git a/Engine/Source/Runtime/Landscape/Private/Landscape.cpp b/Engine/Source/Runtime/Landscape/Private/Landscape.cpp
--- a/Engine/Source/Runtime/Landscape/Private/Landscape.cpp
+++ b/Engine/Source/Runtime/Landscape/Private/Landscape.cpp
@@ -4156,6 +4156,19 @@ void ChangeLandscapeGuidIfObjectIsInstanced(FGuid& InOutGuid, UObject* InObject)
 	}
 
 	UWorldPartition* WorldPartition = FWorldPartitionHelpers::GetWorldPartition(InObject);
+
+#if WITH_EDITOR
+	// In PIE, Actors that are part of a Unsaved cluster of actors can end up being duplicated through a UActorContainer. 
+	// In this case we need to resolve the WorldPartition differently. This could be fixed in a more generic way but would require a lot more testing (Jira: tbd)
+	if (!WorldPartition)
+	{
+		if (UActorContainer* Container = InObject->GetTypedOuter<UActorContainer>())
+		{
+			WorldPartition = FWorldPartitionHelpers::GetWorldPartition(Container->RuntimeLevel.Get());
+		}
+	}
+#endif
+
 	UWorld* OuterWorld = WorldPartition ? WorldPartition->GetTypedOuter<UWorld>() : InObject->GetTypedOuter<UWorld>();
 
 	// TODO [chris.tchou] : Note this is not 100% correct, IsInstanced() returns TRUE when using PIE on non-instanced landscapes.
