- Add support for RayTracing section on landscape (subsections) to have smaller BLAS for faster tracing and build overlap
- Add support for RayTracing LOD bias to select lower LODs for landscape Raytracing BLAS

diff --git a/Engine/Plugins/Experimental/GPULightmass/Source/GPULightmass/Private/Scene/Scene.cpp b/Engine/Plugins/Experimental/GPULightmass/Source/GPULightmass/Private/Scene/Scene.cpp
--- a/Engine/Plugins/Experimental/GPULightmass/Source/GPULightmass/Private/Scene/Scene.cpp
+++ b/Engine/Plugins/Experimental/GPULightmass/Source/GPULightmass/Private/Scene/Scene.cpp
@@ -1272,9 +1272,10 @@ void FScene::AddGeometryInstanceFromComponent(ULandscapeComponent* InComponent)
 		InstanceRenderState.SharedBuffers = FLandscapeComponentSceneProxy::SharedBuffersMap.FindRef(InstanceRenderState.SharedBuffersKey);
 		if (InstanceRenderState.SharedBuffers == nullptr)
 		{
+			int NumRTSections = 1;
 			InstanceRenderState.SharedBuffers = new FLandscapeSharedBuffers(
 				RHICmdList,
-				InstanceRenderState.SharedBuffersKey, Initializer.SubsectionSizeQuads, Initializer.NumSubsections,
+				InstanceRenderState.SharedBuffersKey, Initializer.SubsectionSizeQuads, Initializer.NumSubsections, NumRTSections,
 				LocalFeatureLevel);
 
 			FLandscapeComponentSceneProxy::SharedBuffersMap.Add(InstanceRenderState.SharedBuffersKey, InstanceRenderState.SharedBuffers);
diff --git a/Engine/Shaders/Private/LandscapeVertexFactory.ush b/Engine/Shaders/Private/LandscapeVertexFactory.ush
--- a/Engine/Shaders/Private/LandscapeVertexFactory.ush
+++ b/Engine/Shaders/Private/LandscapeVertexFactory.ush
@@ -47,7 +47,12 @@ float GetSectionLod(uint ComponentIndex)
 	// Fixed grid has fixed Lod
 	return LandscapeFixedGrid.LodValues.x;
 #else
-	return min(GetViewLODData(ComponentIndex) + View.FarShadowStaticMeshLODBias, LandscapeParameters.LastLOD);
+#if RAYTRACING_DYNAMIC_GEOMETRY_CONVERTER
+	int LODBias = LandscapeParameters.RayTracingLODBias;
+#else
+	int LODBias = View.FarShadowStaticMeshLODBias;
+#endif
+	return min(GetViewLODData(ComponentIndex) + LODBias, LandscapeParameters.LastLOD);
 #endif
 }
 
@@ -70,6 +75,12 @@ float4 GetLodValues(uint ComponentIndex)
 #endif
 }
 
+uint GetRayTracingSectionVertCountForLOD(uint ComponentIndex)
+{
+	uint Lod = (uint)floor(GetSectionLod(ComponentIndex));
+	return ((LandscapeParameters.RayTracingSectionSizeVerts >> Lod) + 1);
+}
+
 /**
  * x = Heightmap texture LOD difference from current LOD to highest LOD
  * y = todo: XYOffset texture LOD difference from current LOD to highest LOD
@@ -99,17 +110,11 @@ float GetNeighborSectionLodFromOffset(int2 NeighborOffset, float CenterLOD)
 
 	// lookup neighbor LOD
 	uint NeighborComponentIndex = (NeighborPos.y - LandscapeContinuousLODParameters.Min.y) * LandscapeContinuousLODParameters.Size.x + (NeighborPos.x - LandscapeContinuousLODParameters.Min.x);
-	float NeighborLOD = GetViewLODData(NeighborComponentIndex);
+	float NeighborLOD = GetSectionLod(NeighborComponentIndex);
 
 	// the Edge LOD is the max of Our LOD and the Neighbor LOD (both sides will agree on EdgeLOD this way)
 	float EdgeLOD = max(CenterLOD, NeighborLOD);
-
-	// apply override and shadow bias
-	EdgeLOD = EdgeLOD + View.FarShadowStaticMeshLODBias;
-
-	// clamp to last lod
-	EdgeLOD = min(EdgeLOD, LandscapeParameters.LastLOD);
-
+	
 	return EdgeLOD;
 #endif // !FIXED_GRID
 }
@@ -210,15 +215,17 @@ FVertexFactoryInput LoadVertexFactoryInputForHGS(uint TriangleIndex, int VertexI
 FVertexFactoryInput LoadVertexFactoryInputForDynamicUpdate(uint TriangleIndex, int VertexIndex, uint PrimitiveId, uint DrawInstanceId)
 {
 	FVertexFactoryInput Input = (FVertexFactoryInput)0;
+	
+	uint RTSectionSizeVertsCurrentLOD = GetRayTracingSectionVertCountForLOD(GetComponentLinearIndex());
+	uint RTSectionSizeQuadsCurrentLOD = RTSectionSizeVertsCurrentLOD - 1;
+	uint NumRayTracingSections = LandscapeParameters.NumRayTracingSections;
 
-	float4 LodValues = GetLodValues(GetComponentLinearIndex());
-	uint SubsectionSizeQuadsCurrentLOD = LodValues.z;
-	uint SubsectionSizeVertsCurrentLOD = LodValues.z + 1;
+	float RTSectionSizeScaler = (RTSectionSizeQuadsCurrentLOD - 1.0f / NumRayTracingSections) / RTSectionSizeQuadsCurrentLOD ;
 
 	uint VertexId = TriangleIndex * 3 + VertexIndex;
 
-	Input.Position.x = VertexId % SubsectionSizeVertsCurrentLOD;
-	Input.Position.y = VertexId / SubsectionSizeVertsCurrentLOD;
+	Input.Position.x = (VertexId % RTSectionSizeVertsCurrentLOD + LandscapeMVF.RayTracingSectionXY.x * RTSectionSizeQuadsCurrentLOD) * RTSectionSizeScaler;
+	Input.Position.y = (VertexId / RTSectionSizeVertsCurrentLOD + LandscapeMVF.RayTracingSectionXY.y * RTSectionSizeQuadsCurrentLOD) * RTSectionSizeScaler;	
 	Input.Position.z = LandscapeMVF.SubXY.x;
 	Input.Position.w = LandscapeMVF.SubXY.y;
 
diff --git a/Engine/Source/Runtime/Landscape/Private/LandscapeRender.cpp b/Engine/Source/Runtime/Landscape/Private/LandscapeRender.cpp
--- a/Engine/Source/Runtime/Landscape/Private/LandscapeRender.cpp
+++ b/Engine/Source/Runtime/Landscape/Private/LandscapeRender.cpp
@@ -306,6 +306,23 @@ static FAutoConsoleVariableRef CVarLandscapeRayTracingGeometryFractionalLODUpdat
 
 #endif
 
+int32 GLandscapeRayTracingSections = 1;
+static FAutoConsoleVariableRef CVarLandscapeNumRayTracingSectionsPerSubsection(
+	TEXT("r.RayTracing.Geometry.Landscape.SectionsPerSubsection"),
+	GLandscapeRayTracingSections,
+	TEXT("Number of raytracing BLAS sections per landscape sub section - ray tracing sections will be evenly split (needs to be power of 2)"),
+	FConsoleVariableDelegate::CreateStatic(&OnCVarNeedingRenderStateInvalidationChanged)
+);
+
+int32 GLandscapeRayTracingLODBias = 0;
+static FAutoConsoleVariableRef CVarLandscapeRayTracingVertexLODBias(
+	TEXT("r.RayTracing.Geometry.Landscape.LODBias"),
+	GLandscapeRayTracingLODBias,
+	TEXT("LOD Bias to use when building ray tracing dynamic geometry elements"),
+	FConsoleVariableDelegate::CreateStatic(&OnCVarNeedingRenderStateInvalidationChanged),
+	ECVF_RenderThreadSafe
+);
+
 
 // ----------------------------------------------------------------------------------
 
@@ -500,6 +517,11 @@ static int32 GetDrawCollisionLodOverride(FSceneView const& View, FCollisionRespo
 }
 
 
+static uint32 GetRayTracingSectionVertCountForLOD(uint32 SubsectionSizeVerts, int32 RayTracingNumSections, uint32 LODIndex)
+{
+	return ((SubsectionSizeVerts / RayTracingNumSections) >> LODIndex) + 1;
+}
+
 //
 // FLandscapeComponentSceneProxy
 //
@@ -535,20 +557,9 @@ TBitArray<> FLandscapeRenderSystem::LandscapeIndexAllocator;
 
 struct FLandscapeSectionRayTracingState
 {
-	int8 CurrentLOD;
-	float FractionalLOD;
-	float HeightmapLODBias;
-	uint32 ReferencedTextureRHIHash;
-
 	FRayTracingGeometry Geometry;
 	FRWBuffer RayTracingDynamicVertexBuffer;
 	FLandscapeVertexFactoryMVFUniformBufferRef UniformBuffer;
-
-	FLandscapeSectionRayTracingState()
-		: CurrentLOD(-1)
-		, FractionalLOD(-1000.0f)
-		, HeightmapLODBias(-1000.0f)
-		, ReferencedTextureRHIHash(0) {}
 };
 
 // Where we are rendering multiple views, we need to branch the landscape ray tracing state (BLAS data) per view, for performance reasons.
@@ -556,23 +567,25 @@ struct FLandscapeSectionRayTracingState
 // for a scene with two views.  This structure represents the state for a single view.
 struct FLandscapeRayTracingState
 {
-	FLandscapeRayTracingState()
-		: Pimpl(nullptr), ViewStateNext(nullptr), ViewStatePrev(nullptr), ViewKey(-1), NumSubsections(0) {}
+	FLandscapeRayTracingState() = default;
 	~FLandscapeRayTracingState();
 
 	// Parent structure that holds this ray tracing state, needed for deletion of this item if its view is deleted
-	FLandscapeRayTracingImpl* Pimpl;
+	FLandscapeRayTracingImpl* Pimpl = nullptr;
 
 	// Linked list pointers for FLandscapeRayTracingStateList, referenced from FSceneViewState, iterated over if view gets deleted
-	FLandscapeRayTracingState* ViewStateNext;
-	FLandscapeRayTracingState** ViewStatePrev;
+	FLandscapeRayTracingState* ViewStateNext = nullptr;
+	FLandscapeRayTracingState** ViewStatePrev = nullptr;
 
 	// View state key from FSceneViewState.  Zero is used if FSceneViewState is null (view state keys start at 1, so 0 is invalid for an actual view).
-	uint32 ViewKey;
+	uint32 ViewKey = -1;
 
 	// Rendering data
-	int32 NumSubsections;
-	TStaticArray<FLandscapeSectionRayTracingState, FLandscapeComponentSceneProxy::MAX_SUBSECTION_COUNT> Sections;
+	int8 CurrentLOD = -1;
+	float FractionalLOD = -1000.0f;
+	float HeightmapLODBias = -1000.0f;
+	uint32 ReferencedTextureRHIHash = 0;	
+	TArray<FLandscapeSectionRayTracingState> Sections;
 };
 
 // This wrapper holds the ray tracing state for a single scene proxy for all views
@@ -582,7 +595,12 @@ struct FLandscapeRayTracingImpl
 	TIndirectArray<FLandscapeRayTracingState> PerViewRayTracingState;
 
 	// ViewStateInterface pointer can be NULL
-	FLandscapeRayTracingState* FindOrCreateRayTracingState(FRHICommandListBase& RHICmdList, FSceneViewStateInterface* ViewStateInterface, int32 NumSubsections, int32 SubsectionSizeVerts);
+	FLandscapeRayTracingState* FindOrCreateRayTracingState(
+		FRHICommandListBase& RHICmdList, 
+		FSceneViewStateInterface* ViewStateInterface, 
+		int32 NumSubsections,
+		int32 NumRayTracingSections,
+	    int32 RayTracingSectionVertCount);
 };
 
 // When views get deleted, we need to clean up the per view ray tracing data, which is handled by this class.  This is just a doubly linked list
@@ -1363,6 +1381,7 @@ FLandscapeComponentSceneProxy::FLandscapeComponentSceneProxy(ULandscapeComponent
 	, SubsectionSizeVerts(InComponent->SubsectionSizeQuads + 1)
 	, ComponentSizeQuads(InComponent->ComponentSizeQuads)
 	, ComponentSizeVerts(InComponent->ComponentSizeQuads + 1)
+	, NumRayTracingSections(GLandscapeRayTracingSections)
 	, SectionBase(InComponent->GetSectionBase())
 	, bUsesLandscapeCulling(false)
 	, WeightmapScaleBias(InComponent->WeightmapScaleBias)
@@ -1666,7 +1685,7 @@ FLandscapeComponentSceneProxy::FLandscapeComponentSceneProxy(ULandscapeComponent
 	Algo::Transform(AvailableMaterialInterfaces, MaterialRelevances, [FeatureLevel](UMaterialInterface* InMaterialInterface) { check(InMaterialInterface != nullptr); return InMaterialInterface->GetRelevance_Concurrent(FeatureLevel); });
 
 	const int8 SubsectionSizeLog2 = static_cast<int8>(FMath::CeilLogTwo(InComponent->SubsectionSizeQuads + 1));
-	SharedBuffersKey = (SubsectionSizeLog2 & 0xf) | ((NumSubsections & 0xf) << 4) |	(XYOffsetmapTexture == nullptr ? 0 : 1 << 31);
+	SharedBuffersKey = (SubsectionSizeLog2 & 0xf) | ((NumSubsections & 0xf) << 4) | ((NumRayTracingSections & 0xf) << 8) | (XYOffsetmapTexture == nullptr ? 0 : 1 << 31);
 
 	bSupportsHeightfieldRepresentation = true;
 
@@ -1731,7 +1750,7 @@ void FLandscapeComponentSceneProxy::CreateRenderThreadResources(FRHICommandListB
 #endif // !UE_BUILD_SHIPPING
 
 		SharedBuffers = new FLandscapeSharedBuffers(
-			RHICmdList, SharedBuffersKey, SubsectionSizeQuads, NumSubsections, FeatureLevel, BufferOwnerName);
+			RHICmdList, SharedBuffersKey, SubsectionSizeQuads, NumSubsections, NumRayTracingSections, FeatureLevel, BufferOwnerName);
 
 		FLandscapeComponentSceneProxy::SharedBuffersMap.Add(SharedBuffersKey, SharedBuffers);
 
@@ -1835,7 +1854,13 @@ void FLandscapeComponentSceneProxy::CreateRenderThreadResources(FRHICommandListB
 }
 
 #if RHI_RAYTRACING
-FLandscapeRayTracingState* FLandscapeRayTracingImpl::FindOrCreateRayTracingState(FRHICommandListBase& RHICmdList, FSceneViewStateInterface* ViewStateInterface, int32 NumSubsections, int32 SubsectionSizeVerts)
+
+FLandscapeRayTracingState* FLandscapeRayTracingImpl::FindOrCreateRayTracingState(
+	FRHICommandListBase& RHICmdList, 
+	FSceneViewStateInterface* ViewStateInterface, 
+	int32 NumSubsections,
+	int32 NumRayTracingSections,
+	int32 RayTracingSectionVertCount)
 {
 	// Default view key of zero if there's no ViewStateInterface provided.  View keys start at 1, so 0 wouldn't be a valid key on an actual view.
 	uint32 ViewKey = 0;
@@ -1881,38 +1906,47 @@ FLandscapeRayTracingState* FLandscapeRayTracingImpl::FindOrCreateRayTracingState
 	}
 
 	// Initialize rendering data
-	RayTracingState->NumSubsections = NumSubsections;
+	uint32 TotalSections = FMath::Square(NumRayTracingSections * NumSubsections);
+	RayTracingState->Sections.SetNum(TotalSections);
 
 	for (int32 SubY = 0; SubY < NumSubsections; SubY++)
 	{
 		for (int32 SubX = 0; SubX < NumSubsections; SubX++)
 		{
-			const int8 SubSectionIdx = static_cast<int8>(SubX + SubY * NumSubsections);
-
-			FLandscapeSectionRayTracingState& SectionRayTracingState = RayTracingState->Sections[SubSectionIdx];
-
-			FRayTracingGeometryInitializer Initializer;
-			static const FName DebugName("FLandscapeComponentSceneProxy");
-			static int32 DebugNumber = 0;
-			Initializer.DebugName = FDebugName(DebugName, DebugNumber++);
-			Initializer.IndexBuffer = nullptr;
-			Initializer.GeometryType = RTGT_Triangles;
-			Initializer.bFastBuild = true;
-			Initializer.bAllowUpdate = true;
-
-			FRayTracingGeometrySegment Segment;
-			Segment.VertexBuffer = nullptr;
-			Segment.VertexBufferStride = sizeof(FVector3f);
-			Segment.VertexBufferElementType = VET_Float3;
-			Segment.MaxVertices = FMath::Square(SubsectionSizeVerts);
-			Initializer.Segments.Add(Segment);
-
-			SectionRayTracingState.Geometry.SetInitializer(Initializer);
-			SectionRayTracingState.Geometry.InitResource(RHICmdList);
-
-			FLandscapeVertexFactoryMVFParameters UniformBufferParams;
-			UniformBufferParams.SubXY = FIntPoint(SubX, SubY);
-			SectionRayTracingState.UniformBuffer = FLandscapeVertexFactoryMVFUniformBufferRef::CreateUniformBufferImmediate(UniformBufferParams, UniformBuffer_MultiFrame);
+			for (int32 RTSubY = 0; RTSubY < NumRayTracingSections; RTSubY++)
+			{
+				for (int32 RTSubX = 0; RTSubX < NumRayTracingSections; RTSubX++)
+				{
+					const int32 SubSectionIdx = SubX + SubY * NumSubsections;
+					const int32 InstanceIdx = SubSectionIdx * FMath::Square(NumRayTracingSections) + RTSubX + RTSubY * NumRayTracingSections;
+
+					FLandscapeSectionRayTracingState& SectionRayTracingState = RayTracingState->Sections[InstanceIdx];
+
+					FRayTracingGeometryInitializer Initializer;
+					static const FName DebugName("FLandscapeComponentSceneProxy");
+					static int32 DebugNumber = 0;
+					Initializer.DebugName = FDebugName(DebugName, DebugNumber++);
+					Initializer.IndexBuffer = nullptr;
+					Initializer.GeometryType = RTGT_Triangles;
+					Initializer.bFastBuild = true;
+					Initializer.bAllowUpdate = true;
+
+					FRayTracingGeometrySegment Segment;
+					Segment.VertexBuffer = nullptr;
+					Segment.VertexBufferStride = sizeof(FVector3f);
+					Segment.VertexBufferElementType = VET_Float3;
+					Segment.MaxVertices = FMath::Square(RayTracingSectionVertCount);
+					Initializer.Segments.Add(Segment);
+
+					SectionRayTracingState.Geometry.SetInitializer(Initializer);
+					SectionRayTracingState.Geometry.InitResource(RHICmdList);
+
+					FLandscapeVertexFactoryMVFParameters UniformBufferParams;
+					UniformBufferParams.SubXY = FIntPoint(SubX, SubY);
+					UniformBufferParams.RayTracingSectionXY = FIntPoint(RTSubX, RTSubY);
+					SectionRayTracingState.UniformBuffer = FLandscapeVertexFactoryMVFUniformBufferRef::CreateUniformBufferImmediate(UniformBufferParams, UniformBuffer_MultiFrame);
+				}
+			}
 		}
 	}
 
@@ -1932,14 +1966,10 @@ FLandscapeRayTracingState::~FLandscapeRayTracingState()
 	}
 
 	// And clean up the contents
-	for (int32 SubY = 0; SubY < NumSubsections; SubY++)
+	for (FLandscapeSectionRayTracingState& SectionRayTracingState : Sections)
 	{
-		for (int32 SubX = 0; SubX < NumSubsections; SubX++)
-		{
-			const int8 SubSectionIdx = static_cast<int8>(SubX + SubY * NumSubsections);
-			Sections[SubSectionIdx].Geometry.ReleaseResource();
-			Sections[SubSectionIdx].RayTracingDynamicVertexBuffer.Release();
-		}
+		SectionRayTracingState.Geometry.ReleaseResource();
+		SectionRayTracingState.RayTracingDynamicVertexBuffer.Release();
 	}
 }
 #endif	// RHI_RAYTRACING
@@ -2290,6 +2320,9 @@ void FLandscapeComponentSceneProxy::OnTransformChanged(FRHICommandListBase& RHIC
 	LandscapeParams.ComponentBaseY = RenderCoord.Y;
 	LandscapeParams.SubsectionSizeVerts = SubsectionSizeVerts;
 	LandscapeParams.NumSubsections = NumSubsections;
+	LandscapeParams.RayTracingSectionSizeVerts = SubsectionSizeVerts / NumRayTracingSections;
+	LandscapeParams.NumRayTracingSections = NumRayTracingSections;
+	LandscapeParams.RayTracingLODBias = GLandscapeRayTracingLODBias;
 	LandscapeParams.LastLOD = LastLOD;
 	LandscapeParams.VirtualTexturePerPixelHeight = VirtualTexturePerPixelHeight;
 	LandscapeParams.HeightmapUVScaleBias = HeightmapScaleBias;
@@ -3039,25 +3072,18 @@ void FLandscapeComponentSceneProxy::GetDynamicRayTracingInstances(FRayTracingIns
 		return;
 	}
 
-	const FSceneView& SceneView = *Collector.GetReferenceView();
-	const FLandscapeRenderSystem& RenderSystem = *LandscapeRenderSystems.FindChecked(LandscapeKey);
-
-	if (!RayTracingImpl.IsValid())
-	{
-		RayTracingImpl = MakePimpl<FLandscapeRayTracingImpl>();
-	}
-	FLandscapeRayTracingState* RayTracingState = RayTracingImpl.Get()->FindOrCreateRayTracingState(Collector.GetRHICommandList(), SceneView.State, NumSubsections, SubsectionSizeVerts);
-
-	int32 LODToRender = static_cast<int32>(RenderSystem.GetSectionLODValue(SceneView, RenderCoord));
-
-	FLandscapeElementParamArray& ParameterArray = Collector.AllocateOneFrameResource<FLandscapeElementParamArray>();
-	ParameterArray.ElementParams.AddDefaulted(NumSubsections * NumSubsections);
-
 	if (AvailableMaterials.Num() == 0)
 	{
 		return;
 	}
 
+	const FSceneView& SceneView = *Collector.GetReferenceView();
+	const FLandscapeRenderSystem& RenderSystem = *LandscapeRenderSystems.FindChecked(LandscapeKey);
+
+	int32 RayTracingLODBias = GLandscapeRayTracingLODBias;
+	int32 LODToRender = static_cast<int32>(RenderSystem.GetSectionLODValue(SceneView, RenderCoord)) + RayTracingLODBias;
+	LODToRender = FMath::Min(LODToRender, LastLOD);
+
 	const int8 CurrentLODIndex = static_cast<int8>(LODToRender);
 	int8 MaterialIndex = LODIndexToMaterialIndex.IsValidIndex(CurrentLODIndex) ? LODIndexToMaterialIndex[CurrentLODIndex] : INDEX_NONE;
 	FMaterialRenderProxy* SelectedMaterial = MaterialIndex != INDEX_NONE ? AvailableMaterials[MaterialIndex] : nullptr;
@@ -3068,6 +3094,26 @@ void FLandscapeComponentSceneProxy::GetDynamicRayTracingInstances(FRayTracingIns
 		return;
 	}
 
+	int32 RayTracingSectionLODVertCount = GetRayTracingSectionVertCountForLOD(SubsectionSizeVerts, NumRayTracingSections, CurrentLODIndex);
+	if (RayTracingSectionLODVertCount <= 0)
+	{
+		return;
+	}
+
+	if (!RayTracingImpl.IsValid())
+	{
+		RayTracingImpl = MakePimpl<FLandscapeRayTracingImpl>();
+	}
+
+	
+	int32 RayTracingSectionVertCount = GetRayTracingSectionVertCountForLOD(SubsectionSizeVerts, NumRayTracingSections, 0);
+	FLandscapeRayTracingState* RayTracingState = RayTracingImpl.Get()->FindOrCreateRayTracingState(Collector.GetRHICommandList(), SceneView.State, NumSubsections, NumRayTracingSections, RayTracingSectionVertCount);
+
+	FLandscapeElementParamArray& ParameterArray = Collector.AllocateOneFrameResource<FLandscapeElementParamArray>();
+
+	uint32 TotalInstances = FMath::Square(NumRayTracingSections * NumSubsections);
+	ParameterArray.ElementParams.AddDefaulted(TotalInstances);
+
 	FMeshBatch BaseMeshBatch;
 	BaseMeshBatch.VertexFactory = VertexFactory;
 	BaseMeshBatch.MaterialRenderProxy = SelectedMaterial;
@@ -3078,128 +3124,131 @@ void FLandscapeComponentSceneProxy::GetDynamicRayTracingInstances(FRayTracingIns
 	BaseMeshBatch.SegmentIndex = 0;
 
 	BaseMeshBatch.Elements.Empty();
+		
+	uint32 NumPrimitives = FMath::Square(RayTracingSectionLODVertCount - 1) * 2;
 
-	for (int32 SubY = 0; SubY < NumSubsections; SubY++)
+	// Detect force update CVar
+	bool bNeedsRayTracingGeometryUpdate = (GLandscapeRayTracingGeometryLODsThatUpdateEveryFrame >= 0 && CurrentLODIndex <= (GLandscapeRayTracingGeometryLODsThatUpdateEveryFrame + RayTracingLODBias)) ? true : false;
+	bool bLODChanged = false;
+
+	// Detect continuous LOD parameter changes. This is for far-away high LODs - they change rarely yet the BLAS refit time is not ideal, even if they contains tiny amount of triangles
 	{
-		for (int32 SubX = 0; SubX < NumSubsections; SubX++)
+		if (RayTracingState->CurrentLOD != CurrentLODIndex)
 		{
-			const int8 SubSectionIdx = static_cast<int8>(SubX + SubY * NumSubsections);
-			const int8 CurrentLOD = static_cast<int8>(LODToRender);
-
-			FLandscapeSectionRayTracingState& SectionRayTracingState = RayTracingState->Sections[SubSectionIdx];
-
-			FMeshBatch MeshBatch = BaseMeshBatch;
-
-			FMeshBatchElement BatchElement;
-			FLandscapeBatchElementParams& BatchElementParams = ParameterArray.ElementParams[SubSectionIdx];
-
-			BatchElementParams.LandscapeUniformShaderParametersResource = &LandscapeUniformShaderParameters;
-			BatchElementParams.FixedGridUniformShaderParameters = &SharedBuffers->LandscapeFixedGridUniformShaderParameters;
-			BatchElementParams.LandscapeSectionLODUniformParameters = RenderSystem.SectionLODUniformBuffer;
-			BatchElementParams.SceneProxy = this;
-			BatchElementParams.CurrentLOD = CurrentLOD;
-			BatchElement.UserData = &BatchElementParams;
-			BatchElement.PrimitiveUniformBuffer = GetUniformBuffer();
-
-			int32 LodSubsectionSizeVerts = SubsectionSizeVerts >> CurrentLOD;
-
-			if (LodSubsectionSizeVerts <= 0)
-			{
-				continue;
-			}
-
-			uint32 NumPrimitives = FMath::Square(LodSubsectionSizeVerts - 1) * 2;
-
-			BatchElement.IndexBuffer = SharedBuffers->ZeroOffsetIndexBuffers[CurrentLOD];
-			BatchElement.FirstIndex = 0;
-			BatchElement.NumPrimitives = NumPrimitives;
-			BatchElement.MinVertexIndex = 0;
-			BatchElement.MaxVertexIndex = 0;
+			bNeedsRayTracingGeometryUpdate = true;
+			RayTracingState->CurrentLOD = CurrentLODIndex;
+			bLODChanged = true;
+		}
+		if (RayTracingState->HeightmapLODBias != RenderSystem.GetSectionLODBias(RenderCoord))
+		{
+			bNeedsRayTracingGeometryUpdate = true;
+			RayTracingState->HeightmapLODBias = RenderSystem.GetSectionLODBias(RenderCoord);
+		}
 
-			MeshBatch.Elements.Add(BatchElement);
+		const float PendingFractionalLOD = RenderSystem.GetSectionLODValue(SceneView, RenderCoord);
+		const float FractionLODAbsoluteDifference = FMath::Abs(RayTracingState->FractionalLOD - PendingFractionalLOD);
+		if (FractionLODAbsoluteDifference > GLandscapeRayTracingGeometryFractionalLODUpdateThreshold)
+		{
+			bNeedsRayTracingGeometryUpdate = true;
+			RayTracingState->FractionalLOD = PendingFractionalLOD;
+		}
+	}
 
-			SectionRayTracingState.Geometry.Initializer.IndexBuffer = BatchElement.IndexBuffer->IndexBufferRHI;
+	if (GLandscapeRayTracingGeometryDetectTextureStreaming > 0)
+	{
+		const FMaterialRenderProxy* FallbackMaterialRenderProxyPtr = nullptr;
+		const FMaterial& Material = SelectedMaterial->GetMaterialWithFallback(GetScene().GetFeatureLevel(), FallbackMaterialRenderProxyPtr);
 
-			BatchElementParams.LandscapeVertexFactoryMVFUniformBuffer = SectionRayTracingState.UniformBuffer;
+		if (Material.GetRenderingThreadShaderMap()->UsesWorldPositionOffset())
+		{
+			const FMaterialRenderProxy* MaterialRenderProxy = FallbackMaterialRenderProxyPtr ? FallbackMaterialRenderProxyPtr : SelectedMaterial;
 
-			bool bNeedsRayTracingGeometryUpdate = SectionRayTracingState.Geometry.IsEvicted();
+			FMaterialRenderContext MaterialRenderContext(MaterialRenderProxy, Material, Collector.GetReferenceView());
 
-			// Detect force update CVar
-			bNeedsRayTracingGeometryUpdate |= (CurrentLOD <= GLandscapeRayTracingGeometryLODsThatUpdateEveryFrame) ? true : false;
+			const FUniformExpressionSet& UniformExpressionSet = Material.GetRenderingThreadShaderMap()->GetUniformExpressionSet();
+			const uint32 Hash = UniformExpressionSet.GetReferencedTexture2DRHIHash(MaterialRenderContext);
 
-			// Detect continuous LOD parameter changes. This is for far-away high LODs - they change rarely yet the BLAS refit time is not ideal, even if they contains tiny amount of triangles
+			if (RayTracingState->ReferencedTextureRHIHash != Hash)
 			{
-				if (SectionRayTracingState.CurrentLOD != CurrentLOD)
-				{
-					bNeedsRayTracingGeometryUpdate = true;
-					SectionRayTracingState.CurrentLOD = CurrentLOD;
-					SectionRayTracingState.RayTracingDynamicVertexBuffer.Release();
-				}
-				if (SectionRayTracingState.HeightmapLODBias != RenderSystem.GetSectionLODBias(RenderCoord))
-				{
-					bNeedsRayTracingGeometryUpdate = true;
-					SectionRayTracingState.HeightmapLODBias = RenderSystem.GetSectionLODBias(RenderCoord);
-				}
-
-				const float PendingFractionalLOD = RenderSystem.GetSectionLODValue(SceneView, RenderCoord);
-				const float FractionLODAbsoluteDifference = FMath::Abs(SectionRayTracingState.FractionalLOD - PendingFractionalLOD);
-				if (FractionLODAbsoluteDifference > GLandscapeRayTracingGeometryFractionalLODUpdateThreshold)
-				{
-					bNeedsRayTracingGeometryUpdate = true;
-					SectionRayTracingState.FractionalLOD = PendingFractionalLOD;
-				}
+				bNeedsRayTracingGeometryUpdate = true;
+				RayTracingState->ReferencedTextureRHIHash = Hash;
 			}
+		}
+	}
 
-			if (GLandscapeRayTracingGeometryDetectTextureStreaming > 0)
+	for (int32 SubY = 0; SubY < NumSubsections; SubY++)
+	{
+		for (int32 SubX = 0; SubX < NumSubsections; SubX++)
+		{
+			for (int32 RTSubY = 0; RTSubY < NumRayTracingSections; RTSubY++)
 			{
-				const FMaterialRenderProxy* FallbackMaterialRenderProxyPtr = nullptr;
-				const FMaterial& Material = MeshBatch.MaterialRenderProxy->GetMaterialWithFallback(GetScene().GetFeatureLevel(), FallbackMaterialRenderProxyPtr);
-
-				if (Material.GetRenderingThreadShaderMap()->UsesWorldPositionOffset())
+				for (int32 RTSubX = 0; RTSubX < NumRayTracingSections; RTSubX++)
 				{
-					const FMaterialRenderProxy* MaterialRenderProxy = FallbackMaterialRenderProxyPtr ? FallbackMaterialRenderProxyPtr : MeshBatch.MaterialRenderProxy;
-
-					FMaterialRenderContext MaterialRenderContext(MaterialRenderProxy, Material, Collector.GetReferenceView());
-
-					const FUniformExpressionSet& UniformExpressionSet = Material.GetRenderingThreadShaderMap()->GetUniformExpressionSet();
-					const uint32 Hash = UniformExpressionSet.GetReferencedTexture2DRHIHash(MaterialRenderContext);
-
-					if (SectionRayTracingState.ReferencedTextureRHIHash != Hash)
+					const int32 SubSectionIdx = SubX + SubY * NumSubsections;
+					const int32 InstanceIdx = SubSectionIdx * FMath::Square(NumRayTracingSections) + RTSubX + RTSubY * NumRayTracingSections;
+
+					FLandscapeSectionRayTracingState& SectionRayTracingState = RayTracingState->Sections[InstanceIdx];
+
+					FLandscapeBatchElementParams& BatchElementParams = ParameterArray.ElementParams[InstanceIdx];
+					BatchElementParams.LandscapeUniformShaderParametersResource = &LandscapeUniformShaderParameters;
+					BatchElementParams.FixedGridUniformShaderParameters = &SharedBuffers->LandscapeFixedGridUniformShaderParameters;
+					BatchElementParams.LandscapeSectionLODUniformParameters = RenderSystem.SectionLODUniformBuffer;
+					BatchElementParams.LandscapeVertexFactoryMVFUniformBuffer = SectionRayTracingState.UniformBuffer;
+					BatchElementParams.SceneProxy = this;
+					BatchElementParams.CurrentLOD = CurrentLODIndex;					
+					
+					FMeshBatchElement BatchElement;
+					BatchElement.UserData = &BatchElementParams;
+					BatchElement.PrimitiveUniformBuffer = GetUniformBuffer();
+					BatchElement.IndexBuffer = SharedBuffers->ZeroOffsetIndexBuffers[CurrentLODIndex];
+					BatchElement.FirstIndex = 0;
+					BatchElement.NumPrimitives = NumPrimitives;
+					BatchElement.MinVertexIndex = 0;
+					BatchElement.MaxVertexIndex = 0;
+					
+					FMeshBatch MeshBatch = BaseMeshBatch;
+					MeshBatch.Elements.Add(BatchElement);		
+										
+					check(SectionRayTracingState.Geometry.IsValid() || bNeedsRayTracingGeometryUpdate);
+
+					FRayTracingInstance RayTracingInstance;
+					RayTracingInstance.Geometry = &SectionRayTracingState.Geometry;
+					RayTracingInstance.InstanceTransforms.Add(GetLocalToWorld());
+					RayTracingInstance.Materials.Add(MeshBatch);
+
+					bool bIsEvicted = SectionRayTracingState.Geometry.IsEvicted();
+					if ((bNeedsRayTracingGeometryUpdate || bIsEvicted) && VertexFactory->GetType()->SupportsRayTracingDynamicGeometry())
 					{
-						bNeedsRayTracingGeometryUpdate = true;
-						SectionRayTracingState.ReferencedTextureRHIHash = Hash;
-					}
-				}
-			}
-
-			check(SectionRayTracingState.Geometry.IsValid() || bNeedsRayTracingGeometryUpdate);
-
-			FRayTracingInstance RayTracingInstance;
-			RayTracingInstance.Geometry = &SectionRayTracingState.Geometry;
-			RayTracingInstance.InstanceTransforms.Add(GetLocalToWorld());
-			RayTracingInstance.Materials.Add(MeshBatch);
-
-			if (bNeedsRayTracingGeometryUpdate && VertexFactory->GetType()->SupportsRayTracingDynamicGeometry())
-			{
-				// Use the internal managed vertex buffer because landscape dynamic RT geometries are not updated every frame
-				// which is a requirement for the shared vertex buffer usage
+						// Use the internal managed vertex buffer because landscape dynamic RT geometries are not updated every frame
+						// which is a requirement for the shared vertex buffer usage
+						
+						// Index buffer could have been changed due to LOD change
+						SectionRayTracingState.Geometry.Initializer.IndexBuffer = BatchElement.IndexBuffer->IndexBufferRHI;
+
+						// Recreate the dynamic vertex buffer if the LOD changed - different primitive count
+						if (bLODChanged)
+						{						
+							SectionRayTracingState.RayTracingDynamicVertexBuffer.Release();
+						}
 
-				Collector.AddRayTracingGeometryUpdate(
-					FRayTracingDynamicGeometryUpdateParams
-					{
-						RayTracingInstance.Materials,
-						false,
-						(uint32)FMath::Square(LodSubsectionSizeVerts),
-						FMath::Square(LodSubsectionSizeVerts) * (uint32)sizeof(FVector3f),
-						(uint32)FMath::Square(LodSubsectionSizeVerts - 1) * 2,
-						&SectionRayTracingState.Geometry,
-						&SectionRayTracingState.RayTracingDynamicVertexBuffer,
-						true
+						Collector.AddRayTracingGeometryUpdate(
+							FRayTracingDynamicGeometryUpdateParams
+							{
+								RayTracingInstance.Materials,
+								false,
+								(uint32)FMath::Square(RayTracingSectionLODVertCount),
+								FMath::Square(RayTracingSectionLODVertCount) * (uint32)sizeof(FVector3f),
+								(uint32)FMath::Square(RayTracingSectionLODVertCount - 1) * 2,
+								&SectionRayTracingState.Geometry,
+								&SectionRayTracingState.RayTracingDynamicVertexBuffer,
+								true
+							}
+						);
 					}
-				);
-			}
 
-			Collector.AddRayTracingInstance(MoveTemp(RayTracingInstance));
+					Collector.AddRayTracingInstance(MoveTemp(RayTracingInstance));
+				}
+			}
 		}
 	}
 }
@@ -3350,16 +3399,19 @@ void FLandscapeSharedBuffers::CreateIndexBuffers(FRHICommandListBase& RHICmdList
 #if RHI_RAYTRACING
 		if (IsRayTracingEnabled())
 		{
+			int32 LodRayTracingSectionSizeVerts = GetRayTracingSectionVertCountForLOD(SubsectionSizeVerts, NumRayTracingSections, Mip);
+			int32 LodRayTracingSectionSizeQuads = LodRayTracingSectionSizeVerts - 1;
+			
 			TArray<INDEX_TYPE> ZeroOffsetIndices;
 
-			for (int32 y = 0; y < LodSubsectionSizeQuads; y++)
+			for (int32 y = 0; y < LodRayTracingSectionSizeQuads; y++)
 			{
-				for (int32 x = 0; x < LodSubsectionSizeQuads; x++)
+				for (int32 x = 0; x < LodRayTracingSectionSizeQuads; x++)
 				{
-					INDEX_TYPE i00 = static_cast<INDEX_TYPE>((x + 0) + (y + 0) * (SubsectionSizeVerts >> Mip));
-					INDEX_TYPE i10 = static_cast<INDEX_TYPE>((x + 1) + (y + 0) * (SubsectionSizeVerts >> Mip));
-					INDEX_TYPE i11 = static_cast<INDEX_TYPE>((x + 1) + (y + 1) * (SubsectionSizeVerts >> Mip));
-					INDEX_TYPE i01 = static_cast<INDEX_TYPE>((x + 0) + (y + 1) * (SubsectionSizeVerts >> Mip));
+					INDEX_TYPE i00 = static_cast<INDEX_TYPE>((x + 0) + (y + 0) * LodRayTracingSectionSizeVerts);
+					INDEX_TYPE i10 = static_cast<INDEX_TYPE>((x + 1) + (y + 0) * LodRayTracingSectionSizeVerts);
+					INDEX_TYPE i11 = static_cast<INDEX_TYPE>((x + 1) + (y + 1) * LodRayTracingSectionSizeVerts);
+					INDEX_TYPE i01 = static_cast<INDEX_TYPE>((x + 0) + (y + 1) * LodRayTracingSectionSizeVerts);
 
 					ZeroOffsetIndices.Add(i00);
 					ZeroOffsetIndices.Add(i11);
@@ -3426,11 +3478,18 @@ void FLandscapeSharedBuffers::CreateGrassIndexBuffer(FRHICommandListBase& RHICmd
 	GrassIndexBuffer = IndexBuffer;
 }
 
+
 FLandscapeSharedBuffers::FLandscapeSharedBuffers(FRHICommandListBase& RHICmdList, const int32 InSharedBuffersKey, const int32 InSubsectionSizeQuads, const int32 InNumSubsections, const ERHIFeatureLevel::Type InFeatureLevel, const FName& InOwnerName)
+	: FLandscapeSharedBuffers(RHICmdList, InSharedBuffersKey, InSubsectionSizeQuads, InNumSubsections, 1, InFeatureLevel, InOwnerName)
+{
+}
+
+FLandscapeSharedBuffers::FLandscapeSharedBuffers(FRHICommandListBase& RHICmdList, const int32 InSharedBuffersKey, const int32 InSubsectionSizeQuads, const int32 InNumSubsections, const int32 InNumRayTracingSections, const ERHIFeatureLevel::Type InFeatureLevel, const FName& InOwnerName)
 	: SharedBuffersKey(InSharedBuffersKey)
 	, NumIndexBuffers(FMath::CeilLogTwo(InSubsectionSizeQuads + 1))
 	, SubsectionSizeVerts(InSubsectionSizeQuads + 1)
 	, NumSubsections(InNumSubsections)
+	, NumRayTracingSections(InNumRayTracingSections)
 	, VertexFactory(nullptr)
 	, FixedGridVertexFactory(nullptr)
 	, VertexBuffer(nullptr)
diff --git a/Engine/Source/Runtime/Landscape/Public/LandscapeRender.h b/Engine/Source/Runtime/Landscape/Public/LandscapeRender.h
--- a/Engine/Source/Runtime/Landscape/Public/LandscapeRender.h
+++ b/Engine/Source/Runtime/Landscape/Public/LandscapeRender.h
@@ -118,6 +118,9 @@ BEGIN_GLOBAL_SHADER_PARAMETER_STRUCT(FLandscapeUniformShaderParameters, LANDSCAP
 	SHADER_PARAMETER(int32, ComponentBaseY)
 	SHADER_PARAMETER(int32, SubsectionSizeVerts)
 	SHADER_PARAMETER(int32, NumSubsections)
+	SHADER_PARAMETER(int32, RayTracingSectionSizeVerts)
+	SHADER_PARAMETER(int32, NumRayTracingSections)
+	SHADER_PARAMETER(int32, RayTracingLODBias)
 	SHADER_PARAMETER(int32, LastLOD)
 	SHADER_PARAMETER(uint32, VirtualTexturePerPixelHeight)
 	SHADER_PARAMETER(float, InvLODBlendRange)
@@ -140,6 +143,7 @@ END_GLOBAL_SHADER_PARAMETER_STRUCT()
 
 BEGIN_GLOBAL_SHADER_PARAMETER_STRUCT(FLandscapeVertexFactoryMVFParameters, LANDSCAPE_API)
 	SHADER_PARAMETER(FIntPoint, SubXY)
+	SHADER_PARAMETER(FIntPoint, RayTracingSectionXY)
 END_GLOBAL_SHADER_PARAMETER_STRUCT()
 
 typedef TUniformBufferRef<FLandscapeVertexFactoryMVFParameters> FLandscapeVertexFactoryMVFUniformBufferRef;
@@ -373,6 +377,7 @@ public:
 	int32 NumIndexBuffers;
 	int32 SubsectionSizeVerts;
 	int32 NumSubsections;
+	int32 NumRayTracingSections;
 
 	FLandscapeVertexFactory* VertexFactory;
 	FLandscapeVertexFactory* FixedGridVertexFactory;
@@ -397,9 +402,12 @@ public:
 #if RHI_RAYTRACING
 	TArray<FIndexBuffer*> ZeroOffsetIndexBuffers;
 #endif
-
+	
+	UE_DEPRECATED(5.6, "This constructor has been deprecated. Please use constructor which also provides the number of raytracing sections.")
 	LANDSCAPE_API FLandscapeSharedBuffers(FRHICommandListBase& RHICmdList, int32 SharedBuffersKey, int32 SubsectionSizeQuads, int32 NumSubsections, ERHIFeatureLevel::Type FeatureLevel, const FName& OwnerName = NAME_None);
 
+	LANDSCAPE_API FLandscapeSharedBuffers(FRHICommandListBase& RHICmdList, int32 SharedBuffersKey, int32 SubsectionSizeQuads, int32 NumSubsections, int32 NumRayTracingSections, ERHIFeatureLevel::Type FeatureLevel, const FName& OwnerName = NAME_None);
+
 	template <typename INDEX_TYPE>
 	void CreateIndexBuffers(FRHICommandListBase& RHICmdList, const FName& OwnerName);
 	
@@ -786,6 +794,9 @@ protected:
 	 * Note: in the case of multiple subsections, this is not very useful, as there will be an internal duplicate row of heights in addition to the row at the end.
 	 */
 	int32 ComponentSizeVerts;
+	/** Number of ray tracing sections in a landscape section for better culling and optimized BLAS builds */
+	int32 NumRayTracingSections;
+
 	float StaticLightingResolution;
 	/** Address of the component within the parent Landscape in unique height texels. */
 	FIntPoint SectionBase;
