Fixed landscape thumbnail material warning

#rnx

diff --git a/Engine/Source/Runtime/Landscape/Private/LandscapeRender.cpp b/Engine/Source/Runtime/Landscape/Private/LandscapeRender.cpp
--- a/Engine/Source/Runtime/Landscape/Private/LandscapeRender.cpp
+++ b/Engine/Source/Runtime/Landscape/Private/LandscapeRender.cpp
@@ -4219,9 +4219,11 @@ public:
 			FName(TEXT("FLumenCardPS")),
 
 			// No material cache on thumbnails
-			FName(TEXT("FMaterialCacheUnwrapVS")),
+			FName(TEXT("FMaterialCacheUnwrapVS0")),
+			FName(TEXT("FMaterialCacheUnwrapVS1")),
 			FName(TEXT("FMaterialCacheUnwrapPS")),
 			FName(TEXT("FMaterialCacheNaniteShadeCS")),
+			FName(TEXT("FMaterialCacheShadeCS")),
 		};
 		return ExcludedShaderTypes;
 	}
