Refactored landscape events for heightmap streaming and landscape component data changes : they are now located on the subsystem. That simplifies the registration/unregistration process so that users don't have to monitor add/removal of actors, etc. The downside is that they will receive the event for any landscape actor change instead of just one.
With the current use cases (PCG, water), this is not a problem since most systems want to watch changes coming from any landscape, not a particular one (and they can still check the originating landscape in the event params for this).
The water zone actor is now responsible for reacting to landscape changes instead of the water subsystem. This has the advantage that water zones that we don't needlessly invalidate water zones that don't use bAutoIncludeLandscapesAsTerrain whenever landscape updates

diff --git a/Engine/Plugins/Experimental/Water/Source/Runtime/Private/WaterSubsystem.cpp b/Engine/Plugins/Experimental/Water/Source/Runtime/Private/WaterSubsystem.cpp
--- a/Engine/Plugins/Experimental/Water/Source/Runtime/Private/WaterSubsystem.cpp
+++ b/Engine/Plugins/Experimental/Water/Source/Runtime/Private/WaterSubsystem.cpp
@@ -185,7 +185,7 @@ FWaterBodyManager* UWaterSubsystem::GetWaterBodyManager(const UWorld* InWorld)
 {
 	if (UWaterSubsystem* Subsystem = GetWaterSubsystem(InWorld))
 	{
-		return &Subsystem->GetWaterBodyManagerInternal();
+		return &Subsystem->WaterBodyManager;
 	}
 
 	return nullptr;
@@ -261,7 +261,7 @@ void UWaterSubsystem::Initialize(FSubsystemCollectionBase& Collection)
 	UWorld* World = GetWorld();
 	check(World != nullptr);
 
-	GetWaterBodyManagerInternal().Initialize(World);
+	WaterBodyManager.Initialize(World);
 
 	bUsingSmoothedTime = false;
 	FConsoleVariableDelegate NotifyWaterScalabilityChanged = FConsoleVariableDelegate::CreateUObject(this, &UWaterSubsystem::NotifyWaterScalabilityChangedInternal);
@@ -303,13 +303,6 @@ void UWaterSubsystem::PostInitialize()
 	check(World);
 
 #if WITH_EDITOR
-	// Register for heightmap streaming notifications
-	if (ULandscapeSubsystem* LandscapeSubsystem = World->GetSubsystem<ULandscapeSubsystem>())
-	{
-		check(!OnHeightmapStreamedHandle.IsValid());
-		OnHeightmapStreamedHandle = LandscapeSubsystem->GetOnHeightmapStreamedDelegate().AddUObject(this, &UWaterSubsystem::OnHeightmapStreamed);
-	}
-
 	if (GEngine)
 	{
 		GEngine->OnActorMoved().AddUObject(this, &UWaterSubsystem::OnActorMoved);
@@ -325,15 +318,6 @@ void UWaterSubsystem::Deinitialize()
 	check(World != nullptr);
 
 #if WITH_EDITOR
-	if (OnHeightmapStreamedHandle.IsValid())
-	{
-		if (ULandscapeSubsystem* LandscapeSubsystem = World->GetSubsystem<ULandscapeSubsystem>())
-		{
-			LandscapeSubsystem->GetOnHeightmapStreamedDelegate().Remove(OnHeightmapStreamedHandle);
-		}
-		OnHeightmapStreamedHandle.Reset();
-	}
-
 	if (GEngine)
 	{
 		GEngine->OnActorMoved().RemoveAll(this);
@@ -352,7 +336,7 @@ void UWaterSubsystem::Deinitialize()
 	World->OnBeginPostProcessSettings.RemoveAll(this);
 	World->RemovePostProcessVolume(&UnderwaterPostProcessVolume);
 
-	GetWaterBodyManagerInternal().Deinitialize();
+	WaterBodyManager.Deinitialize();
 
 #if WITH_EDITOR
 	GetDefault<UWaterRuntimeSettings>()->OnSettingsChange.RemoveAll(this);
@@ -371,7 +355,7 @@ void UWaterSubsystem::ApplyRuntimeSettings(const UWaterRuntimeSettings* Settings
 #if WITH_EDITOR
 	// Update sprites since we may have changed the sprite Z offset setting.
 
-	GetWaterBodyManagerInternal().ForEachWaterBodyComponent([](UWaterBodyComponent* Component)
+	WaterBodyManager.ForEachWaterBodyComponent([](UWaterBodyComponent* Component)
 	{
 		Component->UpdateWaterSpriteComponent();
 		return true;
@@ -441,22 +425,8 @@ void UWaterSubsystem::OnActorMoved(AActor* MovedActor)
 		OnWaterTerrainActorChanged(MovedActor);
 	}
 }
-
-void UWaterSubsystem::OnHeightmapStreamed(const FOnHeightmapStreamedContext& InContext)
-{
-	UE_LOG(LogWater, Verbose, TEXT("UWaterSubsystem::OnHeightmapStreamed() -- Rebuilding Water Info Texture..."));
-	MarkWaterZonesInRegionForRebuild(InContext.GetUpdateRegion(), EWaterZoneRebuildFlags::UpdateWaterInfoTexture, /* DebugRequestingObject = */ this);
-}
 #endif // WITH_EDITOR
 
-FWaterBodyManager& UWaterSubsystem::GetWaterBodyManagerInternal()
-{
-	// #todo_water [roey]: Remove these pragmas when we can move WaterBodyManager to private.
-	PRAGMA_DISABLE_DEPRECATION_WARNINGS
-	return WaterBodyManager;
-	PRAGMA_ENABLE_DEPRECATION_WARNINGS
-}
-
 bool UWaterSubsystem::IsShallowWaterSimulationEnabled() const
 {
 	return ShallowWaterSim != 0;
@@ -563,7 +533,7 @@ void UWaterSubsystem::SetOceanFloodHeight(float InFloodHeight)
 				OceanBodyComponent->SetHeightOffset(InFloodHeight);
 			}
 
-			GetWaterBodyManagerInternal().ForEachWaterBodyComponent([this](UWaterBodyComponent* WaterBodyComponent)
+			WaterBodyManager.ForEachWaterBodyComponent([this](UWaterBodyComponent* WaterBodyComponent)
 			{
 				WaterBodyComponent->UpdateMaterialInstances();
 				return true;
@@ -728,7 +698,7 @@ void UWaterSubsystem::NotifyWaterScalabilityChangedInternal(IConsoleVariable* CV
 void UWaterSubsystem::NotifyWaterVisibilityChangedInternal(IConsoleVariable* CVar)
 {
 	// Water body visibility depends on various CVars. All need to update the visibility in water body components : 
-	GetWaterBodyManagerInternal().ForEachWaterBodyComponent([](UWaterBodyComponent* WaterBodyComponent)
+	WaterBodyManager.ForEachWaterBodyComponent([](UWaterBodyComponent* WaterBodyComponent)
 	{
 		WaterBodyComponent->UpdateVisibility();
 		return true;
@@ -800,7 +770,7 @@ void UWaterSubsystem::ComputeUnderwaterPostProcess(FVector ViewLocation, FSceneV
 
 	TArray<FHitResult> Hits;
 	TArray<FWaterBodyPostProcessQuery, TInlineAllocator<4>> WaterBodyQueriesToProcess;
-	const bool bWorldHasWater  = GetWaterBodyManagerInternal().HasAnyWaterBodies();
+	const bool bWorldHasWater  = WaterBodyManager.HasAnyWaterBodies();
 	if (bWorldHasWater && World->SweepMultiByChannel(Hits, ViewLocation, ViewLocation + FVector(0, 0, TraceDistance), FQuat::Identity, UnderwaterTraceChannel, FCollisionShape::MakeSphere(TraceDistance), TraceSimple))
 	{
 		if (Hits.Num() > 1)
diff --git a/Engine/Plugins/Experimental/Water/Source/Runtime/Private/WaterZoneActor.cpp b/Engine/Plugins/Experimental/Water/Source/Runtime/Private/WaterZoneActor.cpp
--- a/Engine/Plugins/Experimental/Water/Source/Runtime/Private/WaterZoneActor.cpp
+++ b/Engine/Plugins/Experimental/Water/Source/Runtime/Private/WaterZoneActor.cpp
@@ -9,6 +9,7 @@
 #include "EngineUtils.h"
 #include "LandscapeComponent.h"
 #include "LandscapeProxy.h"
+#include "LandscapeSubsystem.h"
 #include "WaterUtils.h"
 #include "Materials/MaterialInstanceDynamic.h"
 #include "WaterViewExtension.h"
@@ -251,12 +252,40 @@ void AWaterZone::PostRegisterAllComponents()
 	{
 		WaterZoneIndex = Manager->AddWaterZone(this);
 	}
+
+#if WITH_EDITOR
+	// Register for landscape notifications
+	if (ULandscapeSubsystem* LandscapeSubsystem = GetWorld()->GetSubsystem<ULandscapeSubsystem>())
+	{
+		check(!OnHeightmapStreamedHandle.IsValid());
+		OnHeightmapStreamedHandle = LandscapeSubsystem->OnHeightmapStreamed().AddUObject(this, &AWaterZone::OnLandscapeHeightmapStreamed);
+
+		check(!OnLandscapeComponentDataChangedHandle.IsValid());
+		OnLandscapeComponentDataChangedHandle = LandscapeSubsystem->OnLandscapeProxyComponentDataChanged().AddUObject(this, &AWaterZone::OnLandscapeComponentDataChanged);
+	}
+#endif // WITH_EDITOR
 }
 
 void AWaterZone::PostUnregisterAllComponents()
 {
 	Super::PostUnregisterAllComponents();
 
+#if WITH_EDITOR
+	if (ULandscapeSubsystem* LandscapeSubsystem = GetWorld()->GetSubsystem<ULandscapeSubsystem>())
+	{
+		if (OnHeightmapStreamedHandle.IsValid())
+		{
+			LandscapeSubsystem->OnHeightmapStreamed().Remove(OnHeightmapStreamedHandle);
+		}
+		if (OnLandscapeComponentDataChangedHandle.IsValid())
+		{
+			LandscapeSubsystem->OnLandscapeProxyComponentDataChanged().Remove(OnLandscapeComponentDataChangedHandle);
+		}
+	}
+	OnHeightmapStreamedHandle.Reset();
+	OnLandscapeComponentDataChangedHandle.Reset();
+#endif // WITH_EDITOR
+
 	// We must check for the index because UnregisterAllComponents can be called multiple times in a row by PostEditChangeProperty, etc.
 	FWaterBodyManager* Manager = UWaterSubsystem::GetWaterBodyManager(GetWorld());
 	if (Manager && !IsTemplate() && (WaterZoneIndex != INDEX_NONE))
@@ -481,6 +510,36 @@ void AWaterZone::GetStreamingBounds(FBox& OutRuntimeBounds, FBox& OutEditorBound
 {
 	OutRuntimeBounds = OutEditorBounds = GetZoneBounds();
 }
+
+void AWaterZone::OnLandscapeHeightmapStreamed(const FOnHeightmapStreamedContext& InContext)
+{
+	if (bAutoIncludeLandscapesAsTerrain)
+	{
+		UE_LOG(LogWater, Verbose, TEXT("AWaterZone::OnLandscapeHeightmapStreamed() -- Rebuilding Water Info Texture..."));
+		// Invalidate the region corresponding the heightmap being streamed in so that the water info is always up-to-date with the current landscape (i.e. ground) height :
+		const FBox2D& UpdateRegion = InContext.GetUpdateRegion();
+		MarkForRebuild(EWaterZoneRebuildFlags::UpdateWaterInfoTexture, FBox2D(FVector2D(UpdateRegion.Min), FVector2D(UpdateRegion.Max)), reinterpret_cast<const UObject*>(InContext.GetLandscape()));
+	}
+}
+
+void AWaterZone::OnLandscapeComponentDataChanged(ALandscapeProxy* InLandscape, const FLandscapeProxyComponentDataChangedParams& InChangeParams)
+{
+	// TODO [jonathan.bard] : if that is a problem, we could add more info to FLandscapeProxyComponentDataChangedParams and tell exactly what has changed (heightmap, weightmap, etc.)
+	//  because so far, the water info only needs to be notified about heightmap changes
+	if (bAutoIncludeLandscapesAsTerrain)
+	{
+		UE_LOG(LogWater, Verbose, TEXT("AWaterZone::OnLandscapeComponentDataChanged() -- Rebuilding Water Info Texture..."));
+		// Invalidate the region corresponding to each of the changed components. Use a combined region, because 99% of times, the component are contiguous and laid out in a square :
+		FBox CombinedBounds;
+		InChangeParams.ForEachComponent([&CombinedBounds](const ULandscapeComponent* InLandscapeComponent)
+			{
+				FBox LandscapeComponentBounds = InLandscapeComponent->GetLocalBounds().TransformBy(InLandscapeComponent->GetComponentTransform()).GetBox();
+				CombinedBounds += LandscapeComponentBounds;
+			});
+		MarkForRebuild(EWaterZoneRebuildFlags::UpdateWaterInfoTexture, FBox2D(FVector2D(CombinedBounds.Min), FVector2D(CombinedBounds.Max)), InLandscape);
+	}
+}
+
 #endif // WITH_EDITOR
 
 void AWaterZone::SetFarMeshMaterial(UMaterialInterface* InFarDistanceMaterial)
diff --git a/Engine/Plugins/Experimental/Water/Source/Runtime/Public/WaterSubsystem.h b/Engine/Plugins/Experimental/Water/Source/Runtime/Public/WaterSubsystem.h
--- a/Engine/Plugins/Experimental/Water/Source/Runtime/Public/WaterSubsystem.h
+++ b/Engine/Plugins/Experimental/Water/Source/Runtime/Public/WaterSubsystem.h
@@ -212,8 +212,6 @@ private:
 
 	void OnWaterTerrainActorChanged(const AActor* TerrainActor);
 
-	FWaterBodyManager& GetWaterBodyManagerInternal();
-
 #if !(UE_BUILD_SHIPPING || UE_BUILD_TEST)
 	void ShowOnScreenDebugInfo(const FVector& InViewLocation, const FUnderwaterPostProcessDebugInfo& InDebugInfo);
 #endif // !(UE_BUILD_SHIPPING || UE_BUILD_TEST)
@@ -269,9 +267,6 @@ private:
 	TMultiMap<const AActor*, TWeakObjectPtr<UWaterTerrainComponent>> WaterTerrainActors;
 
 #if WITH_EDITOR
-	FDelegateHandle OnHeightmapStreamedHandle;
-	void OnHeightmapStreamed(const struct FOnHeightmapStreamedContext& InContext);
-
 	/** By default, there is no water subsystem allowed on preview worlds except when explicitly requested : */
 	static bool bAllowWaterSubsystemOnPreviewWorld;
 #endif // WITH_EDITOR
diff --git a/Engine/Plugins/Experimental/Water/Source/Runtime/Public/WaterZoneActor.h b/Engine/Plugins/Experimental/Water/Source/Runtime/Public/WaterZoneActor.h
--- a/Engine/Plugins/Experimental/Water/Source/Runtime/Public/WaterZoneActor.h
+++ b/Engine/Plugins/Experimental/Water/Source/Runtime/Public/WaterZoneActor.h
@@ -5,16 +5,17 @@
 #include "GameFramework/Actor.h"
 #include "WaterZoneActor.generated.h"
 
+class ALandscapeProxy;
+class AWaterBody;
+class FLandscapeProxyComponentDataChangedParams;
 class UBillboardComponent;
+class UBoxComponent;
 class UTexture2D;
 class UTextureRenderTarget2D;
 class UTextureRenderTarget2DArray;
-namespace EEndPlayReason { enum Type : int; }
-
-class UWaterMeshComponent;
-class UBoxComponent;
-class AWaterBody;
 class UWaterBodyComponent;
+class UWaterMeshComponent;
+namespace EEndPlayReason { enum Type : int; }
 
 enum class EWaterZoneRebuildFlags
 {
@@ -175,6 +176,12 @@ private:
 
 	/** Called when the Bounds component is modified. Updates the value of ZoneExtent to match the new bounds */
 	void OnBoundsComponentModified();
+
+	/** Called whenever a landscape component has its heightmap streamed */
+	void OnLandscapeHeightmapStreamed(const struct FOnHeightmapStreamedContext& InContext);
+
+	/** Called whenever a landscape component has its heightmap updated */
+	void OnLandscapeComponentDataChanged(ALandscapeProxy* InLandscape, const FLandscapeProxyComponentDataChangedParams& InChangeParams);
 #endif // WITH_EDITOR
 
 private:
@@ -262,6 +269,8 @@ private:
 	UPROPERTY()
 	bool bEnableNonTesselatedLODMesh_DEPRECATED;
 
+	FDelegateHandle OnHeightmapStreamedHandle;
+	FDelegateHandle OnLandscapeComponentDataChangedHandle;
 #endif // WITH_EDITORONLY_DATA
 };
 DEFINE_ACTORDESC_TYPE(AWaterZone, FWaterZoneActorDesc);
diff --git a/Engine/Plugins/PCG/Source/PCG/Private/Grid/PCGLandscapeCache.cpp b/Engine/Plugins/PCG/Source/PCG/Private/Grid/PCGLandscapeCache.cpp
--- a/Engine/Plugins/PCG/Source/PCG/Private/Grid/PCGLandscapeCache.cpp
+++ b/Engine/Plugins/PCG/Source/PCG/Private/Grid/PCGLandscapeCache.cpp
@@ -13,6 +13,7 @@
 #include "LandscapeDataAccess.h"
 #include "LandscapeInfoMap.h"
 #include "LandscapeProxy.h"
+#include "LandscapeSubsystem.h"
 
 #include "Async/ParallelFor.h"
 #include "Engine/Engine.h"
@@ -1022,7 +1023,11 @@ void UPCGLandscapeCache::SetupLandscapeCallbacks()
 		ALandscapeProxy* Landscape = CastChecked<ALandscapeProxy>(FoundLandscape);
 
 		Landscapes.Add(Landscape);
-		Landscape->OnComponentDataChanged.AddUObject(this, &UPCGLandscapeCache::OnLandscapeChanged);
+	}
+
+	if (ULandscapeSubsystem* LandscapeSubsystem = World->GetSubsystem<ULandscapeSubsystem>())
+	{
+		LandscapeSubsystem->OnLandscapeProxyComponentDataChanged().AddUObject(this, &UPCGLandscapeCache::OnLandscapeChanged);
 	}
 
 	// Also track when the landscape is moved, added, or deleted
@@ -1042,14 +1047,6 @@ void UPCGLandscapeCache::SetupLandscapeCallbacks()
 
 void UPCGLandscapeCache::TeardownLandscapeCallbacks()
 {
-	for (TWeakObjectPtr<ALandscapeProxy> LandscapeWeakPtr : Landscapes)
-	{
-		if (ALandscapeProxy* Landscape = LandscapeWeakPtr.Get())
-		{
-			Landscape->OnComponentDataChanged.RemoveAll(this);
-		}
-	}
-
 	if (GEngine)
 	{
 		GEngine->OnActorMoved().RemoveAll(this);
@@ -1059,6 +1056,11 @@ void UPCGLandscapeCache::TeardownLandscapeCallbacks()
 
 	if (UWorld* World = GetWorld())
 	{
+	    if (ULandscapeSubsystem* LandscapeSubsystem = World->GetSubsystem<ULandscapeSubsystem>())
+	    {
+		    LandscapeSubsystem->OnLandscapeProxyComponentDataChanged().RemoveAll(this);
+	    }
+
 		if (!World->IsPlayInEditor() && World->IsPartitionedWorld() && World->PersistentLevel)
 		{
 			World->PersistentLevel->OnLoadedActorAddedToLevelEvent.RemoveAll(this);
@@ -1127,7 +1129,6 @@ void UPCGLandscapeCache::OnLandscapeDeleted(AActor* Actor)
 	RemoveComponentFromCache(LandscapeProxy);
 
 	Landscapes.Remove(LandscapeProxy);
-	LandscapeProxy->OnComponentDataChanged.RemoveAll(this);
 
 	CacheLock.WriteUnlock();
 }
@@ -1142,10 +1143,6 @@ void UPCGLandscapeCache::OnLandscapeAdded(AActor* Actor)
 	if (ALandscapeProxy* LandscapeProxy = Cast<ALandscapeProxy>(Actor))
 	{
 		Landscapes.Add(LandscapeProxy);
-		if (!LandscapeProxy->OnComponentDataChanged.IsBoundToObject(this))
-		{
-			LandscapeProxy->OnComponentDataChanged.AddUObject(this, &UPCGLandscapeCache::OnLandscapeChanged);
-		}
 
 		// Note: Landscape Proxies have no components at this stage, so they will need to be added on demand
 	}
diff --git a/Engine/Plugins/PCG/Source/PCG/Private/PCGActorAndComponentMapping.cpp b/Engine/Plugins/PCG/Source/PCG/Private/PCGActorAndComponentMapping.cpp
--- a/Engine/Plugins/PCG/Source/PCG/Private/PCGActorAndComponentMapping.cpp
+++ b/Engine/Plugins/PCG/Source/PCG/Private/PCGActorAndComponentMapping.cpp
@@ -15,7 +15,7 @@
 #include "Utils/PCGGraphExecutionLogging.h"
 
 #include "Landscape.h"
-#include "LandscapeProxy.h"
+#include "LandscapeSubsystem.h"
 #include "StaticMeshCompiler.h"
 #include "TextureCompiler.h"
 #include "Algo/AnyOf.h"
@@ -1265,18 +1265,10 @@ void FPCGActorAndComponentMapping::RegisterDelegates()
 		}
 	}
 
-	// Considering that some landscape might already be loaded, we should attach to their change callbacks now
-	UPCGActorHelpers::ForEachActorInWorld<ALandscapeProxy>(World, [this](AActor* InActor)
+	if (ULandscapeSubsystem* LandscapeSubsystem = World->GetSubsystem<ULandscapeSubsystem>())
 	{
-		if (ALandscapeProxy* LandscapeProxy = Cast<ALandscapeProxy>(InActor))
-		{
-			if (!LandscapeProxy->OnComponentDataChanged.IsBoundToObject(this))
-			{
-				LandscapeProxy->OnComponentDataChanged.AddRaw(this, &FPCGActorAndComponentMapping::OnLandscapeChanged);
-			}
-		}
-		return true;
-	});
+		LandscapeSubsystem->OnLandscapeProxyComponentDataChanged().AddRaw(this, &FPCGActorAndComponentMapping::OnLandscapeChanged);
+	}
 }
 
 void FPCGActorAndComponentMapping::UnregisterDelegates()
@@ -1308,6 +1300,11 @@ void FPCGActorAndComponentMapping::UnregisterDelegates()
 			LevelInstanceSubsystem->OnLevelInstanceEditCancelled().RemoveAll(this);
 		}
 	}
+
+	if (ULandscapeSubsystem* LandscapeSubsystem = World->GetSubsystem<ULandscapeSubsystem>())
+	{
+		LandscapeSubsystem->OnLandscapeProxyComponentDataChanged().RemoveAll(this);
+	}
 }
 
 bool FPCGActorAndComponentMapping::ShouldDelayActor(AActor* InActor) const
@@ -1553,13 +1550,6 @@ void FPCGActorAndComponentMapping::OnActorAdded_Internal(AActor* InActor, bool b
 	TRACE_CPUPROFILER_EVENT_SCOPE(FPCGActorAndComponentMapping::OnActorAdded);
 	check(!PCGActorAndComponentMapping::ShouldIgnoreActor(InActor) && InActor->GetWorld() == GetWorld());
 
-	if (ALandscapeProxy* LandscapeProxy = Cast<ALandscapeProxy>(InActor))
-	{
-		if (!LandscapeProxy->OnComponentDataChanged.IsBoundToObject(this))
-		{
-			LandscapeProxy->OnComponentDataChanged.AddRaw(this, &FPCGActorAndComponentMapping::OnLandscapeChanged);
-		}
-	}
 
 	if (!bShouldDirty || DelayedChangedActors.Contains(InActor))
 	{
@@ -1613,11 +1603,6 @@ void FPCGActorAndComponentMapping::OnActorDeleted_Internal(AActor* InActor, bool
 	TRACE_CPUPROFILER_EVENT_SCOPE(UPCGActorAndComponentMapping::OnActorDeleted);
 	check(InActor && InActor->GetWorld() == GetWorld());
 
-	if (ALandscapeProxy* LandscapeProxy = Cast<ALandscapeProxy>(InActor))
-	{
-		LandscapeProxy->OnComponentDataChanged.RemoveAll(this);
-	}
-
 	if (!bShouldDirty)
 	{
 		// Nothing to do
diff --git a/Engine/Source/Runtime/Landscape/Classes/LandscapeProxy.h b/Engine/Source/Runtime/Landscape/Classes/LandscapeProxy.h
--- a/Engine/Source/Runtime/Landscape/Classes/LandscapeProxy.h
+++ b/Engine/Source/Runtime/Landscape/Classes/LandscapeProxy.h
@@ -973,6 +973,7 @@ PRAGMA_ENABLE_DEPRECATION_WARNINGS
 #if WITH_EDITOR
 	LANDSCAPE_API static ULandscapeLayerInfoObject* VisibilityLayer;
 
+	UE_DEPRECATED(5.6, "Use ULandscapeSubsystem::OnLandscapeProxyComponentDataChanged")
 	FOnLandscapeProxyComponentDataChanged OnComponentDataChanged;
 #endif // WITH_EDITOR
 
diff --git a/Engine/Source/Runtime/Landscape/Private/Landscape.cpp b/Engine/Source/Runtime/Landscape/Private/Landscape.cpp
--- a/Engine/Source/Runtime/Landscape/Private/Landscape.cpp
+++ b/Engine/Source/Runtime/Landscape/Private/Landscape.cpp
@@ -7030,7 +7030,11 @@ void ALandscapeProxy::InvalidateGeneratedComponentData(const TArray<ULandscapeCo
 		}
 
 		FLandscapeProxyComponentDataChangedParams ChangeParams(Iter.Value());
+		Subsystem->GetDelegateAccess().OnLandscapeProxyComponentDataChangedDelegate.Broadcast(Iter.Key(), ChangeParams);
+
+PRAGMA_DISABLE_DEPRECATION_WARNINGS
 		Proxy->OnComponentDataChanged.Broadcast(Iter.Key(), ChangeParams);
+PRAGMA_ENABLE_DEPRECATION_WARNINGS
 #endif // WITH_EDITOR
 
 		Proxy->UpdateRenderingMethod();
diff --git a/Engine/Source/Runtime/Landscape/Private/LandscapeEditLayers.cpp b/Engine/Source/Runtime/Landscape/Private/LandscapeEditLayers.cpp
--- a/Engine/Source/Runtime/Landscape/Private/LandscapeEditLayers.cpp
+++ b/Engine/Source/Runtime/Landscape/Private/LandscapeEditLayers.cpp
@@ -4094,8 +4094,8 @@ bool ALandscape::PrepareTextureResources(bool bInWaitForStreaming)
 		// Notify that heightmaps have been streamed.
 		if (ULandscapeSubsystem* LandscapeSubsystem = GetWorld()->GetSubsystem<ULandscapeSubsystem>())
 		{
-			FOnHeightmapStreamedContext Context(HeightmapUpdateRegion, StreamedInHeightmapComponents);
-			LandscapeSubsystem->GetOnHeightmapStreamedDelegate().Broadcast(Context);
+			FOnHeightmapStreamedContext Context(this, HeightmapUpdateRegion, StreamedInHeightmapComponents);
+			LandscapeSubsystem->GetDelegateAccess().OnHeightmapStreamedDelegate.Broadcast(Context);
 		}
 	}
 
diff --git a/Engine/Source/Runtime/Landscape/Private/LandscapeSubsystem.cpp b/Engine/Source/Runtime/Landscape/Private/LandscapeSubsystem.cpp
--- a/Engine/Source/Runtime/Landscape/Private/LandscapeSubsystem.cpp
+++ b/Engine/Source/Runtime/Landscape/Private/LandscapeSubsystem.cpp
@@ -72,6 +72,24 @@ static FAutoConsoleVariable CVarMaxAsyncNaniteProxiesPerSecond(
 	6.0f,
 	TEXT("Number of Async nanite proxies to dispatch per second"));
 
+int32 LiveRebuildNaniteOnModification = 0;
+static FAutoConsoleVariableRef CVarLiveRebuildNaniteOnModification(
+	TEXT("landscape.Nanite.LiveRebuildOnModification"),
+	LiveRebuildNaniteOnModification,
+	TEXT("Trigger a rebuild of Nanite representation immediately when a modification is performed (World Partition Maps Only)"));
+
+int32 LandscapeMultithreadNaniteBuild = 1;
+static FAutoConsoleVariableRef CVarLandscapeMultithreadNaniteBuild(
+	TEXT("landscape.Nanite.MultithreadBuild"),
+	LandscapeMultithreadNaniteBuild,
+	TEXT("Multithread nanite landscape build in (World Partition Maps Only)"));
+
+int32 LandscapeMaxSimultaneousMultithreadNaniteBuilds = -1;
+static FAutoConsoleVariableRef CVarLandscapeMaxSimultaneousMultithreadNaniteBuilds(
+	TEXT("landscape.Nanite.MaxSimultaneousMultithreadBuilds"),
+	LandscapeMaxSimultaneousMultithreadNaniteBuilds,
+	TEXT("Max number of simultaneous Nanite static mesh tasks (-1 = unlimited )"));
+
 extern int32 GGrassMapUseRuntimeGeneration;
 
 DECLARE_CYCLE_STAT(TEXT("LandscapeSubsystem Tick"), STAT_LandscapeSubsystemTick, STATGROUP_Landscape);
@@ -207,10 +225,10 @@ namespace UE::Landscape
 	{
 		if (GEditor)
 		{
-			if (UWorld* World = GEditor->GetEditorWorldContext().World())
+		if (UWorld* World = GEditor->GetEditorWorldContext().World())
+		{
+			if (ULandscapeSubsystem* LandscapeSubsystem = World->GetSubsystem<ULandscapeSubsystem>())
 			{
-				if (ULandscapeSubsystem* LandscapeSubsystem = World->GetSubsystem<ULandscapeSubsystem>())
-				{
 					LandscapeSubsystem->MarkModifiedLandscapesAsDirty(InBuildFlags);
 				}
 			}
@@ -227,10 +245,10 @@ namespace UE::Landscape
 	{
 		if (GEditor)
 		{
-			if (UWorld* World = GEditor->GetEditorWorldContext().World())
+		if (UWorld* World = GEditor->GetEditorWorldContext().World())
+		{
+			if (ULandscapeSubsystem* LandscapeSubsystem = World->GetSubsystem<ULandscapeSubsystem>())
 			{
-				if (ULandscapeSubsystem* LandscapeSubsystem = World->GetSubsystem<ULandscapeSubsystem>())
-				{
 					LandscapeSubsystem->BuildGrassMaps(InBuildFlags);
 				}
 			}
@@ -247,10 +265,10 @@ namespace UE::Landscape
 	{
 		if (GEditor)
 		{
-			if (UWorld* World = GEditor->GetEditorWorldContext().World())
+		if (UWorld* World = GEditor->GetEditorWorldContext().World())
+		{
+			if (ULandscapeSubsystem* LandscapeSubsystem = World->GetSubsystem<ULandscapeSubsystem>())
 			{
-				if (ULandscapeSubsystem* LandscapeSubsystem = World->GetSubsystem<ULandscapeSubsystem>())
-				{
 					LandscapeSubsystem->BuildPhysicalMaterial(InBuildFlags);
 				}
 			}
@@ -267,10 +285,10 @@ namespace UE::Landscape
 	{
 		if (GEditor)
 		{
-			if (UWorld* World = GEditor->GetEditorWorldContext().World())
+		if (UWorld* World = GEditor->GetEditorWorldContext().World())
+		{
+			if (ULandscapeSubsystem* LandscapeSubsystem = World->GetSubsystem<ULandscapeSubsystem>())
 			{
-				if (ULandscapeSubsystem* LandscapeSubsystem = World->GetSubsystem<ULandscapeSubsystem>())
-				{
 					LandscapeSubsystem->BuildNanite(InBuildFlags);
 				}
 			}
@@ -287,15 +305,15 @@ namespace UE::Landscape
 	{
 		if (GEditor)
 		{
-			if (UWorld* World = GEditor->GetEditorWorldContext().World())
+		if (UWorld* World = GEditor->GetEditorWorldContext().World())
+		{
+			if (ULandscapeSubsystem* LandscapeSubsystem = World->GetSubsystem<ULandscapeSubsystem>())
 			{
-				if (ULandscapeSubsystem* LandscapeSubsystem = World->GetSubsystem<ULandscapeSubsystem>())
-				{
 					LandscapeSubsystem->BuildAll(InBuildFlags);
-				}
 			}
 		}
 	}
+	}
 
 	static FAutoConsoleCommand CmdLandscapeRebuildPhysicalMaterial(
 		TEXT("landscape.RebuildPhysicalMaterial"),
@@ -756,13 +774,13 @@ void ULandscapeSubsystem::Tick(float DeltaTime)
 
 	static TArray<ALandscapeProxy*> ActiveProxies;
 	{
-		ActiveProxies.Reset(Proxies.Num());
+	ActiveProxies.Reset(Proxies.Num());
 
 		for (TObjectPtr<ALandscapeProxy> ProxyPtr : Proxies)
-		{
+	{
 			ALandscapeProxy* Proxy = ProxyPtr.Get();
 			ActiveProxies.Add(Proxy);
-
+			
 			// Update the proxies proxy
 			{
 				if (!Proxy->GetDisableRuntimeGrassMapGeneration())
@@ -834,8 +852,8 @@ void ULandscapeSubsystem::Tick(float DeltaTime)
 				Landscape->TickLayers(DeltaTime);
 			}
 
-			Proxy->UpdatePhysicalMaterialTasks();
-		}
+				Proxy->UpdatePhysicalMaterialTasks();
+			}
 
 		Proxy->GetAsyncWorkMonitor().Tick(DeltaTime);
 
@@ -1187,7 +1205,7 @@ void ULandscapeSubsystem::BuildNanite(UE::Landscape::EBuildFlags InBuildFlags, T
 		TStatId(),
 		&AsyncEvents,
 		ENamedThreads::GameThread);
-
+	
 	FinishAllNaniteBuildsInFlightNow(EFinishAllNaniteBuildsInFlightFlags::DisplaySlowTaskDialog);
 	check(NaniteBuildsInFlight == 0);
 
@@ -1570,24 +1588,6 @@ bool ULandscapeSubsystem::FinishAllNaniteBuildsInFlightNow(EFinishAllNaniteBuild
 	return !bCancelled && !AreNaniteBuildsInProgress();
 }
 
-int32 LiveRebuildNaniteOnModification = 0;
-static FAutoConsoleVariableRef CVarLiveRebuildNaniteOnModification(
-	TEXT("landscape.Nanite.LiveRebuildOnModification"),
-	LiveRebuildNaniteOnModification,
-	TEXT("Trigger a rebuild of Nanite representation immediately when a modification is performed (World Partition Maps Only)"));
-
-int32 LandscapeMultithreadNaniteBuild = 1;
-static FAutoConsoleVariableRef CVarLandscapeMultithreadNaniteBuild(
-	TEXT("landscape.Nanite.MultithreadBuild"),
-	LandscapeMultithreadNaniteBuild,
-	TEXT("Multithread nanite landscape build in (World Partition Maps Only)"));
-
-int32 LandscapeMaxSimultaneousMultithreadNaniteBuilds = -1;
-static FAutoConsoleVariableRef CVarLandscapeMaxSimultaneousMultithreadNaniteBuilds(
-	TEXT("landscape.Nanite.MaxSimultaneousMultithreadBuilds"),
-	LandscapeMaxSimultaneousMultithreadNaniteBuilds,
-	TEXT("Max number of simultaneous Nanite static mesh tasks (-1 = unlimited )"));
-
 bool ULandscapeSubsystem::IsMultithreadedNaniteBuildEnabled()
 {
 	return LandscapeMultithreadNaniteBuild > 0;
@@ -1641,6 +1641,16 @@ void ULandscapeSubsystem::WaitLaunchNaniteBuild()
 	check((LandscapeMultithreadNaniteBuild != 0) || (NaniteStaticMeshesInFlight <= 1));
 }
 
+ULandscapeSubsystem::FDelegateAccess::FDelegateAccess(FOnHeightmapStreamedDelegate& InOnHeightmapStreamed, FOnLandscapeProxyComponentDataChanged& InOnLandscapeProxyComponentDataChanged) 
+	: OnHeightmapStreamedDelegate(InOnHeightmapStreamed)
+	, OnLandscapeProxyComponentDataChangedDelegate(InOnLandscapeProxyComponentDataChanged)
+{}
+
+ULandscapeSubsystem::FDelegateAccess ULandscapeSubsystem::GetDelegateAccess() const
+{
+	return FDelegateAccess(OnHeightmapStreamedDelegate, OnLandscapeProxyComponentDataChangedDelegate);
+}
+
 #endif // WITH_EDITOR
 
 #undef LOCTEXT_NAMESPACE
diff --git a/Engine/Source/Runtime/Landscape/Public/LandscapeSubsystem.h b/Engine/Source/Runtime/Landscape/Public/LandscapeSubsystem.h
--- a/Engine/Source/Runtime/Landscape/Public/LandscapeSubsystem.h
+++ b/Engine/Source/Runtime/Landscape/Public/LandscapeSubsystem.h
@@ -5,6 +5,7 @@
 #include "CoreMinimal.h"
 #include "UObject/ObjectMacros.h"
 #include "Subsystems/WorldSubsystem.h"
+#include "LandscapeProxy.h"
 #include "LandscapeSubsystem.generated.h"
 
 class ALandscapeProxy;
@@ -67,15 +68,30 @@ namespace UE::Landscape
 struct FOnHeightmapStreamedContext
 {
 private:
+	const ALandscape* Landscape = nullptr;
 	const FBox2D& UpdateRegion;
 	const TSet<class ULandscapeComponent*>& LandscapeComponentsInvolved;
 
 public:
-	const FBox2D& GetUpdateRegion() const { return UpdateRegion; }
-	const TSet<class ULandscapeComponent*>& GetLandscapeComponentsInvolved() const { return LandscapeComponentsInvolved; }
+	const ALandscape* GetLandscape() const 
+	{ 
+		return Landscape; 
+	}
+
+	const FBox2D& GetUpdateRegion() const 
+	{ 
+		return UpdateRegion; 
+	}
 
-	FOnHeightmapStreamedContext(const FBox2D& updateRegion, const TSet<class ULandscapeComponent*>& landscapeComponentsInvolved)
-		: UpdateRegion(updateRegion), LandscapeComponentsInvolved(landscapeComponentsInvolved)
+	const TSet<class ULandscapeComponent*>& GetLandscapeComponentsInvolved() const 
+	{ 
+		return LandscapeComponentsInvolved; 
+	}
+
+	FOnHeightmapStreamedContext(const ALandscape* InLandscape, const FBox2D& InUpdateRegion, const TSet<class ULandscapeComponent*>& InLandscapeComponentsInvolved)
+		: Landscape(InLandscape)
+		, UpdateRegion(InUpdateRegion)
+		, LandscapeComponentsInvolved(InLandscapeComponentsInvolved)
 	{}
 };
 
@@ -170,8 +186,28 @@ public:
 	/** Returns true if there are some landscapes in the editor world that have been automatically modified and are in need of being saved (see LandscapeDirtyingMode) */
 	LANDSCAPE_API bool HasModifiedLandscapes() const;
 	LANDSCAPE_API bool GetDirtyOnlyInMode() const;
-	FLandscapeNotificationManager* GetNotificationManager() { return NotificationManager; }
-	FOnHeightmapStreamedDelegate& GetOnHeightmapStreamedDelegate() { return OnHeightmapStreamed; }
+	
+	FLandscapeNotificationManager* GetNotificationManager() 
+	{ 
+		return NotificationManager; 
+	}
+
+	UE_DEPRECATED(5.6, "Use OnHeightmapStreamed()")
+	FOnHeightmapStreamedDelegate& GetOnHeightmapStreamedDelegate() 
+	{ 
+		return OnHeightmapStreamedDelegate;
+	}
+
+	FOnHeightmapStreamedDelegate::RegistrationType& OnHeightmapStreamed()
+	{ 
+		return OnHeightmapStreamedDelegate;
+	}
+
+	FOnLandscapeProxyComponentDataChanged::RegistrationType& OnLandscapeProxyComponentDataChanged() const
+	{
+		return OnLandscapeProxyComponentDataChangedDelegate;
+	}
+
 	bool AnyViewShowCollisions() const { return bAnyViewShowCollisions; }  //! Returns true if any view has view collisions enabled.
 	FDateTime GetAppCurrentDateTime();
 
@@ -202,6 +238,20 @@ public:
 
 	// Wait unit we're able to continue a landscape export task (Max concurrent nanite mesh builds is defined by  landscape.Nanite.MaxSimultaneousMultithreadBuilds and landscape.Nanite.MultithreadBuild CVars)
 	void WaitLaunchNaniteBuild(); 
+
+	// Helper class reserved for friends that are allowed to fire the subsystem's inner callbacks
+	class FDelegateAccess
+	{
+		friend ULandscapeSubsystem;
+		friend ALandscape;
+		friend ALandscapeProxy;
+
+		FDelegateAccess(FOnHeightmapStreamedDelegate& InOnHeightmapStreamed, FOnLandscapeProxyComponentDataChanged& InOnLandscapeProxyComponentDataChanged);
+
+		FOnHeightmapStreamedDelegate& OnHeightmapStreamedDelegate;
+		FOnLandscapeProxyComponentDataChanged& OnLandscapeProxyComponentDataChangedDelegate;
+	};
+	FDelegateAccess GetDelegateAccess() const;
 #endif // WITH_EDITOR
 
 	// runs per-tick edge fixup on ALL landscape groups in the subsystem
@@ -252,10 +302,11 @@ private:
 	FLandscapeGrassMapsBuilder* GrassMapsBuilder = nullptr;
 
 #if WITH_EDITORONLY_DATA
+	mutable FOnHeightmapStreamedDelegate OnHeightmapStreamedDelegate;
+	mutable FOnLandscapeProxyComponentDataChanged OnLandscapeProxyComponentDataChangedDelegate;
+
 	class FLandscapePhysicalMaterialBuilder* PhysicalMaterialBuilder = nullptr;
-	
 	FLandscapeNotificationManager* NotificationManager = nullptr;
-	FOnHeightmapStreamedDelegate OnHeightmapStreamed;
 	bool bAnyViewShowCollisions = false;
 	FDateTime AppCurrentDateTime; // Represents FDateTime::Now(), at the beginning of the frame (useful to get a human-readable date/time that is fixed during the frame)
 	int32 LastTickFrameNumber = -1;
