Fixed assert when running landscape batched merge and creating new weightmap allocation while partially updating landscape. The code was asking for new weightmap allocations to be restricted to landscape components that are being resolved which are different than those being rendered . That could mean that a new allocation would be added to a component that is resolved but not rendered, which would require to re-prepare the entire merge operation (to take into account the newly added weightmap). Instead, we now restrict the allocation to landscape components already meant for rendering.

#rnx

diff --git a/Engine/Source/Runtime/Landscape/Classes/Landscape.h b/Engine/Source/Runtime/Landscape/Classes/Landscape.h
--- a/Engine/Source/Runtime/Landscape/Classes/Landscape.h
+++ b/Engine/Source/Runtime/Landscape/Classes/Landscape.h
@@ -580,7 +580,7 @@ private:
 	void PrepareComponentDataToPackMaterialLayersCS(int32 InCurrentWeightmapToProcessIndex, const FIntPoint& InLandscapeBase, const TArray<ULandscapeComponent*>& InAllLandscapeComponents, TArray<UTexture2D*>& InOutProcessedWeightmaps,
 													TArray<FLandscapeEditLayerReadback*>& OutProcessedCPUReadBacks, TArray<struct FLandscapeLayerWeightmapPackMaterialLayersComponentData>& OutComponentData);
 	void ReallocateLayersWeightmaps(FUpdateLayersContentContext& InUpdateLayersContentContext, const TArray<ULandscapeLayerInfoObject*>& InBrushRequiredAllocations, 
-		const TMap<ULandscapeComponent*, TArray<ULandscapeLayerInfoObject*>>* InPerComponentAllocations, bool bInRestrictTextureSharingToTheseComponents);
+		const TMap<ULandscapeComponent*, TArray<ULandscapeLayerInfoObject*>>* InPerComponentAllocations, TSet<ULandscapeComponent*>* InRestrictTextureSharingToComponents);
 	void InitializeLayersWeightmapResources();
 	bool GenerateZeroAllocationPerComponents(const TArray<ALandscapeProxy*>& InAllLandscape, const TMap<ULandscapeLayerInfoObject*, bool>& InWeightmapLayersBlendSubstractive);
 
diff --git a/Engine/Source/Runtime/Landscape/Private/LandscapeEditLayers.cpp b/Engine/Source/Runtime/Landscape/Private/LandscapeEditLayers.cpp
--- a/Engine/Source/Runtime/Landscape/Private/LandscapeEditLayers.cpp
+++ b/Engine/Source/Runtime/Landscape/Private/LandscapeEditLayers.cpp
@@ -8794,7 +8794,7 @@ void ALandscape::PrepareComponentDataToPackMaterialLayersCS(int32 InCurrentWeigh
 }
 
 void ALandscape::ReallocateLayersWeightmaps(FUpdateLayersContentContext& InUpdateLayersContentContext, const TArray<ULandscapeLayerInfoObject*>& InBrushRequiredAllocations, 
-	const TMap<ULandscapeComponent*, TArray<ULandscapeLayerInfoObject*>>* InPerComponentAllocations, bool bInRestrictTextureSharingToTheseComponents)
+	const TMap<ULandscapeComponent*, TArray<ULandscapeLayerInfoObject*>>* InPerComponentAllocations, TSet<ULandscapeComponent*>* InRestrictTextureSharingToComponents)
 {
 	TRACE_CPUPROFILER_EVENT_SCOPE(LandscapeLayers_ReallocateLayersWeightmaps);
 
@@ -8833,43 +8833,43 @@ void ALandscape::ReallocateLayersWeightmaps(FUpdateLayersContentContext& InUpdat
 	}
 	else
 	{
-	for (FLandscapeLayer& Layer : LandscapeEditLayers)
-	{
-		for (ULandscapeComponent* Component : InUpdateLayersContentContext.LandscapeComponentsWeightmapsToResolve)
+		for (FLandscapeLayer& Layer : LandscapeEditLayers)
 		{
-			TArray<ULandscapeLayerInfoObject*>* ComponentLayerAlloc = LayerAllocsPerComponent.Find(Component);
-			if (ComponentLayerAlloc == nullptr)
+			for (ULandscapeComponent* Component : InUpdateLayersContentContext.LandscapeComponentsWeightmapsToResolve)
 			{
-				TArray<ULandscapeLayerInfoObject*> NewLayerAllocs;
-				ComponentLayerAlloc = &LayerAllocsPerComponent.Add(Component, NewLayerAllocs);
-			}
-
-			// No need for an allocation if the edit layer is invisible : 
-			if (Layer.bVisible)
-			{
-				if (FLandscapeLayerComponentData* LayerComponentData = Component->GetLayerData(Layer.Guid))
+				TArray<ULandscapeLayerInfoObject*>* ComponentLayerAlloc = LayerAllocsPerComponent.Find(Component);
+				if (ComponentLayerAlloc == nullptr)
 				{
-				for (const FWeightmapLayerAllocationInfo& LayerWeightmapAllocation : LayerComponentData->WeightmapData.LayerAllocations)
+					TArray<ULandscapeLayerInfoObject*> NewLayerAllocs;
+					ComponentLayerAlloc = &LayerAllocsPerComponent.Add(Component, NewLayerAllocs);
+				}
+
+				// No need for an allocation if the edit layer is invisible : 
+				if (Layer.bVisible)
 				{
-					if (LayerWeightmapAllocation.LayerInfo != nullptr)
+					if (FLandscapeLayerComponentData* LayerComponentData = Component->GetLayerData(Layer.Guid))
 					{
-						ComponentLayerAlloc->AddUnique(LayerWeightmapAllocation.LayerInfo);
+					for (const FWeightmapLayerAllocationInfo& LayerWeightmapAllocation : LayerComponentData->WeightmapData.LayerAllocations)
+					{
+						if (LayerWeightmapAllocation.LayerInfo != nullptr)
+						{
+							ComponentLayerAlloc->AddUnique(LayerWeightmapAllocation.LayerInfo);
+						}
 					}
 				}
-			}
-			}
+				}
 
-			// Add the brush alloc also (only if !InMergeParams.bSkipBrush, but InBrushRequiredAllocations should be empty already if InMergeParams.bSkipBrush is true) :
-			for (ULandscapeLayerInfoObject* BrushLayerInfo : InBrushRequiredAllocations)
-			{
-				if (BrushLayerInfo != nullptr)
+				// Add the brush alloc also (only if !InMergeParams.bSkipBrush, but InBrushRequiredAllocations should be empty already if InMergeParams.bSkipBrush is true) :
+				for (ULandscapeLayerInfoObject* BrushLayerInfo : InBrushRequiredAllocations)
 				{
-					ComponentLayerAlloc->AddUnique(BrushLayerInfo);
+					if (BrushLayerInfo != nullptr)
+					{
+						ComponentLayerAlloc->AddUnique(BrushLayerInfo);
+					}
 				}
 			}
 		}
 	}
-	}
 
 	// Trim the components that don't need weightmaps anymore (e.g. all edit layers are made invisible : there were some
 	// components in LandscapeComponentsWeightmapsToResolve but there aren't anymore now).
@@ -8923,13 +8923,8 @@ void ALandscape::ReallocateLayersWeightmaps(FUpdateLayersContentContext& InUpdat
 	for (ULandscapeComponent* Component : InUpdateLayersContentContext.LandscapeComponentsWeightmapsToResolve)
 	{
 		// If requested, don't allow the component to use textures that have a weightmap in another component than those provided when reallocating weightmaps :
-		TSet<ULandscapeComponent*> RestrictSharingToComponents;
-		if (bInRestrictTextureSharingToTheseComponents)
-		{
-			LayerAllocsPerComponent.GetKeys(RestrictSharingToComponents);
-		}
 		NewCreatedTextures.Append(Component->ReallocateWeightmaps(/*DataInterface = */nullptr, /*InEditLayerGuid = */FGuid(), /*bInSaveToTransactionBuffer = */false, /*bool bInForceReallocate = */false, 
-			/*InTargetProxy = */nullptr, /*InRestrictSharingToComponents = */bInRestrictTextureSharingToTheseComponents ? &RestrictSharingToComponents : nullptr));
+			/*InTargetProxy = */nullptr, InRestrictTextureSharingToComponents));
 	}
 
 	// TODO: correctly only recreate what is required instead of everything..
@@ -9877,7 +9872,7 @@ int32 ALandscape::PerformLayersWeightmapsLocalMerge(FUpdateLayersContentContext&
 	// TODO [jonathan.bard] : this is only possible to do so in local merge, since we don't yet support BP brushes, for which we need to call Render() in order to be able to know which layer allocations
 	//  they need. As such, the global merge path is broken and cannot be fixed unless the BP brush interface is changed in order to inform the system, by advance, about what paint layer it needs
 	// Make sure we have proper textures+allocations for all the final weightmaps we're about to resolve :
-	ReallocateLayersWeightmaps(InUpdateLayersContentContext, /*InBrushRequiredAllocations = */TArray<ULandscapeLayerInfoObject*>(), /*InPerComponentAllocations = */nullptr, /*bInRestrictTextureSharingToTheseComponents = */false);
+	ReallocateLayersWeightmaps(InUpdateLayersContentContext, /*InBrushRequiredAllocations = */TArray<ULandscapeLayerInfoObject*>(), /*InPerComponentAllocations = */nullptr, /*InRestrictTextureSharingToComponents = */nullptr);
 
 	FMergeInfo RenderThreadData;
 	PrepareLayersWeightmapsLocalMergeRenderThreadData(InUpdateLayersContentContext, InMergeParams, RenderThreadData);
@@ -10261,7 +10256,7 @@ int32 ALandscape::PerformLayersWeightmapsGlobalMerge(FUpdateLayersContentContext
 			bFirstLayer = false;
 		}
 
-		ReallocateLayersWeightmaps(InUpdateLayersContentContext, BrushRequiredAllocations, /*InPerComponentAllocations = */nullptr, /*bInRestrictTextureSharingToTheseComponents = */false);
+		ReallocateLayersWeightmaps(InUpdateLayersContentContext, BrushRequiredAllocations, /*InPerComponentAllocations = */nullptr, /*InRestrictTextureSharingToComponents = */nullptr);
 
 		// List of UTexture2D that we need to kick off readbacks for :
 		TArray<UTexture2D*> TexturesNeedingReadback;
@@ -10510,7 +10505,11 @@ int32 ALandscape::PerformLayersWeightmapsBatchedMerge(FUpdateLayersContentContex
 				}
 			}
 		}
-		ReallocateLayersWeightmaps(InUpdateLayersContentContext, /*InBrushRequiredAllocations = */{}, &PerComponentAllocations, /*bInRestrictTextureSharingToTheseComponents = */true);
+
+		// We don't want new components to be required for rendering because of ReallocateLayersWeightmaps, as that would so would require re-running the merge preparation step. 
+		//  So we prevent new texture allocations from using textures from components that are not already in the list of components to render : 
+		TSet<ULandscapeComponent*> RestrictTextureSharingToComponents(InUpdateLayersContentContext.LandscapeComponentsWeightmapsToRender);
+		ReallocateLayersWeightmaps(InUpdateLayersContentContext, /*InBrushRequiredAllocations = */{}, &PerComponentAllocations, &RestrictTextureSharingToComponents);
 
 		checkf(Algo::AllOf(InUpdateLayersContentContext.LandscapeComponentsWeightmapsToRender, [&LandscapeComponentsWeightmapsToRenderBefore](ULandscapeComponent* InComponent) { return LandscapeComponentsWeightmapsToRenderBefore.Contains(InComponent); }), 
 			TEXT("If this asserts, the list of components involved in the weightmaps being merged has changed, which would require re-running the merge preparation step. bInRestrictTextureSharingToTheseComponents should prevent this from happening"));
