Make AddTargetLayer optional after calling ALandscapeProxy::Import.  Missing layer infos will be automatically added as needed.  This restores the behavior that Import had in 5.4.

diff --git a/Engine/Source/Runtime/Landscape/Private/Landscape.cpp b/Engine/Source/Runtime/Landscape/Private/Landscape.cpp
--- a/Engine/Source/Runtime/Landscape/Private/Landscape.cpp
+++ b/Engine/Source/Runtime/Landscape/Private/Landscape.cpp
@@ -6862,10 +6862,7 @@ FLandscapeTargetLayerSettings& ALandscapeProxy::AddTargetLayer(const FName& Name
 {
 	Modify();
 
-	check(!HasTargetLayer(Name));
-	
 	FLandscapeTargetLayerSettings& Settings = TargetLayers.Add(Name, TargetLayerSettings);
-
 	if (FProperty* Property = StaticClass()->FindPropertyByName(GET_MEMBER_NAME_CHECKED(ALandscapeProxy, TargetLayers)); Property != nullptr && bPostEditChange)
 	{
 		FPropertyChangedEvent PropertyChangedEvent(Property);
diff --git a/Engine/Source/Runtime/Landscape/Private/LandscapeEdit.cpp b/Engine/Source/Runtime/Landscape/Private/LandscapeEdit.cpp
--- a/Engine/Source/Runtime/Landscape/Private/LandscapeEdit.cpp
+++ b/Engine/Source/Runtime/Landscape/Private/LandscapeEdit.cpp
@@ -3986,6 +3986,12 @@ void ALandscapeProxy::Import(const FGuid& InGuid, int32 InMinX, int32 InMinY, in
 					if (MaterialLayerInfo.LayerInfo != nullptr && MaterialLayerInfo.LayerData.Num() > 0)
 					{
 						LandscapeEdit.SetAlphaData(MaterialLayerInfo.LayerInfo, InMinX, InMinY, InMaxX, InMaxY, MaterialLayerInfo.LayerData.GetData(), 0, ELandscapeLayerPaintingRestriction::None, true, false);
+
+						// Add to TargetLayers if needed.
+						if (!HasTargetLayer(MaterialLayerInfo.LayerInfo->LayerName))
+						{
+							AddTargetLayer(MaterialLayerInfo.LayerInfo->LayerName, FLandscapeTargetLayerSettings(MaterialLayerInfo.LayerInfo));
+						}
 					}
 				}
 			}
