Fixed grass map generation being stalled forever when one of the landscape actor of the map has no loaded components. This was due to global merge resource initialization returning false when no component is loaded, which leads to an early out of UpdateLayersContent, which prevents the update content flags from being updated, which leads ALandscape::IsUpToDate to return false, which prevents *all* grass maps rendering from happening

diff --git a/Engine/Source/Runtime/Landscape/Private/LandscapeEditLayers.cpp b/Engine/Source/Runtime/Landscape/Private/LandscapeEditLayers.cpp
--- a/Engine/Source/Runtime/Landscape/Private/LandscapeEditLayers.cpp
+++ b/Engine/Source/Runtime/Landscape/Private/LandscapeEditLayers.cpp
@@ -2399,7 +2399,9 @@ bool ALandscape::CreateLayersRenderingResource()
 	// No components, can't update the render targets	
 	if (ComponentCounts.X <= 0 || ComponentCounts.Y <= 0)
 	{
-		return false;
+		// We return true, still, because we don't want to prevent the update from happening. It will update nothing, but at least it will clear the update flags LayerContentUpdateModes
+		//  which will in turn unblock things like grass map rendering, etc.
+		return true;
 	}
 
 	ALandscape* Landscape = GetLandscapeActor();
