Landscape ULandscapeEditLayer Base refactor:

User-facing changes:
* Moved editor settings from water settings to landscape settings to avoid updating landscape during interactive changes (e.g. when dragging the alpha slider of an edit layer)
* When double-clicking a target layer, it now opens the corresponding asset
* Added custom "Inspected" objects detail view docking tab for the landscape editor mode
* Fixed Edit Layer Alpha slider bugs (Incorrect minimum value, able to set alpha for splines)

Non-user facing changes:
* Moved all settings from FLandscapeLayer to ULandscapeEditLayerBase (Besides Brushes)
* Added generic delegate-based system for communicating data changes from ULandscapeEditLayerBase (to landscape)
* Deprecated/Removed LandscapeEdMode and Landscape.h functions to directly use ULandscapeEditLayerBase accessors
* Moved LandscapeGrass, LandscapeProxy, LandscapeUI Detail Customizations into Landscape Editor Module

diff --git a/Engine/Plugins/Experimental/LandscapePatch/Source/LandscapePatch/Private/LandscapePatchComponent.cpp b/Engine/Plugins/Experimental/LandscapePatch/Source/LandscapePatch/Private/LandscapePatchComponent.cpp
--- a/Engine/Plugins/Experimental/LandscapePatch/Source/LandscapePatch/Private/LandscapePatchComponent.cpp
+++ b/Engine/Plugins/Experimental/LandscapePatch/Source/LandscapePatch/Private/LandscapePatchComponent.cpp
@@ -861,9 +861,9 @@ bool ULandscapePatchComponent::BindToEditLayer(const FGuid& Guid)
 	EditLayer = Cast<ULandscapePatchEditLayer>(LayerStruct ? LayerStruct->EditLayer : nullptr);
 	EditLayerGuid = Guid;
 	DetailPanelLayerGuid = Guid.ToString();
-	DetailPanelLayerName = 
-		EditLayer.IsValid() && LayerStruct ? LayerStruct->Name.ToString()
-		: LayerStruct ? TEXT("-Incorrect Layer Type-")
+	DetailPanelLayerName =
+		EditLayer.IsValid()
+		? EditLayer->GetName().ToString()
 		: TEXT("-Layer Not Found-");
 
 	if (PreviousLayer && EditLayer != PreviousLayer)
@@ -944,7 +944,7 @@ bool ULandscapePatchComponent::BindToLandscape(ALandscape* LandscapeIn)
 	if (!Layer)
 	{
 		// See if we have any layer of the appropriate type
-		LayerStruct = LandscapeIn->FindLayerOfType(ULandscapePatchEditLayer::StaticClass());
+		LayerStruct = LandscapeIn->FindLayerOfTypeConst(ULandscapePatchEditLayer::StaticClass());
 		Layer = LayerStruct ? Cast<ULandscapePatchEditLayer>(LayerStruct->EditLayer) : nullptr;
 	}
 
@@ -972,7 +972,7 @@ bool ULandscapePatchComponent::BindToLandscape(ALandscape* LandscapeIn)
 		return false;
 	}
 
-	return BindToEditLayer(LayerStruct->Guid);
+	return BindToEditLayer(LayerStruct->EditLayer->GetGuid());
 }
 
 bool ULandscapePatchComponent::BindToAnyLandscape()
@@ -1013,7 +1013,7 @@ void ULandscapePatchComponent::UpdateEditLayerFromDetailPanelLayerName()
 	}
 
 	const FLandscapeLayer* Layer = Landscape->GetLayerConst(*DetailPanelLayerName);
-	SetEditLayerGuid(Layer ? Layer->Guid : FGuid());
+	SetEditLayerGuid(Layer ? Layer->EditLayer->GetGuid() : FGuid());
 }
 
 void ULandscapePatchComponent::FixBindings()
@@ -1151,11 +1151,11 @@ const TArray<FString> ULandscapePatchComponent::GetLayerOptions()
 	// Get the layer names from our landscape
 	if (Landscape.IsValid())
 	{
-		for (const FLandscapeLayer& Layer : Landscape->GetLayers())
+		for (const FLandscapeLayer& Layer : Landscape->GetLayersConst())
 		{
 			if (Layer.EditLayer->IsA<ULandscapePatchEditLayer>())
 			{
-				Options.Add(Layer.Name.ToString());
+				Options.Add(Layer.EditLayer->GetName().ToString());
 			}
 		}
 	}
diff --git a/Engine/Plugins/Experimental/LandscapePatch/Source/LandscapePatch/Private/LandscapePatchEditLayer.cpp b/Engine/Plugins/Experimental/LandscapePatch/Source/LandscapePatch/Private/LandscapePatchEditLayer.cpp
--- a/Engine/Plugins/Experimental/LandscapePatch/Source/LandscapePatch/Private/LandscapePatchEditLayer.cpp
+++ b/Engine/Plugins/Experimental/LandscapePatch/Source/LandscapePatch/Private/LandscapePatchEditLayer.cpp
@@ -296,16 +296,6 @@ void ULandscapePatchEditLayer::OnLayerRemoved()
 #endif
 }
 
-void ULandscapePatchEditLayer::OnLayerCreated(FLandscapeLayer& Layer)
-{
-	Super::OnLayerCreated(Layer);
-
-#if WITH_EDITOR
-	// TODO: Once the guid lives on the UObject, we won't need this.
-	EditLayerGuid = Layer.Guid;
-#endif
-}
-
 bool ULandscapePatchEditLayer::SupportsTargetType(ELandscapeToolTargetType InType) const
 {
 	return InType != ELandscapeToolTargetType::Invalid;
@@ -337,22 +327,10 @@ TArray<UE::Landscape::EditLayers::FEditLayerRendererState> ULandscapePatchEditLa
 	return RendererStates;
 }
 
-void ULandscapePatchEditLayer::PostLoad()
-{
-	Super::PostLoad();
-
-	// Our layer should always have a guid that it got from OnLayerCreated. But if somehow that got lost,
-	// try to find the layer in our owning landscape.
-	if (!EditLayerGuid.IsValid() && ensure(OwningLandscape.IsValid()))
-	{
-		TArrayView<const FLandscapeLayer> Layers = OwningLandscape->GetLayers();
-		const FLandscapeLayer* LayerStruct = Layers.FindByPredicate([this](const FLandscapeLayer& Struct) { return Struct.EditLayer == this; });
-		EditLayerGuid = LayerStruct ? LayerStruct->Guid : FGuid();
-	}
-}
-
 void ULandscapePatchEditLayer::PostEditUndo()
 {
+	Super::PostEditUndo();
+
 	UpdatePatchList();
 }
 
@@ -369,7 +347,7 @@ void ULandscapePatchEditLayer::RequestLandscapeUpdate(bool bInUserTriggered /* =
 
 bool ULandscapePatchEditLayer::ShouldPatchBeIncludedInList(const ULandscapePatchComponent* Patch) const
 {
-	return Patch && Patch->IsPatchInWorld() && Patch->GetEditLayerGuid() == EditLayerGuid;
+	return Patch && Patch->IsPatchInWorld() && Patch->GetEditLayerGuid() == GetGuid();
 }
 
 // Attempt binary search to find the insertion index for a patch. Returns end of array if an invalid patch is sampled.
diff --git a/Engine/Plugins/Experimental/LandscapePatch/Source/LandscapePatch/Private/LandscapePatchManager.cpp b/Engine/Plugins/Experimental/LandscapePatch/Source/LandscapePatch/Private/LandscapePatchManager.cpp
--- a/Engine/Plugins/Experimental/LandscapePatch/Source/LandscapePatch/Private/LandscapePatchManager.cpp
+++ b/Engine/Plugins/Experimental/LandscapePatch/Source/LandscapePatch/Private/LandscapePatchManager.cpp
@@ -532,7 +532,7 @@ void ALandscapePatchManager::MigrateToPrioritySystemAndDeleteInternal(bool bAllo
 		{
 			// Create the ULandscapePatchEditLayer in advance, if needed.  This prevents BindToLandscape from triggering the
 			// modal UI window for choosing the new layer index.
-			const FLandscapeLayer* LayerStruct = OwningLandscape->FindLayerOfType(ULandscapePatchEditLayer::StaticClass());
+			const FLandscapeLayer* LayerStruct = OwningLandscape->FindLayerOfTypeConst(ULandscapePatchEditLayer::StaticClass());
 			ULandscapePatchEditLayer* Layer = LayerStruct ? Cast<ULandscapePatchEditLayer>(LayerStruct->EditLayer) : nullptr;
 
 			if (!Layer)
diff --git a/Engine/Plugins/Experimental/LandscapePatch/Source/LandscapePatch/Public/LandscapePatchEditLayer.h b/Engine/Plugins/Experimental/LandscapePatch/Source/LandscapePatch/Public/LandscapePatchEditLayer.h
--- a/Engine/Plugins/Experimental/LandscapePatch/Source/LandscapePatch/Public/LandscapePatchEditLayer.h
+++ b/Engine/Plugins/Experimental/LandscapePatch/Source/LandscapePatch/Public/LandscapePatchEditLayer.h
@@ -57,8 +57,6 @@ public:
 	virtual FString GetDefaultName() const override { return TEXT("Patches"); }
 	virtual void GetRenderDependencies(TSet<UObject*>& OutDependencies) const override;
 	virtual void OnLayerRemoved() override;
-	// TODO: Remove once our guid gets initialized somewhere else
-	virtual void OnLayerCreated(FLandscapeLayer& Layer);
 
 #if WITH_EDITOR
 	using FEditLayerRendererState = UE::Landscape::EditLayers::FEditLayerRendererState;
@@ -67,7 +65,6 @@ public:
 	virtual TArray<FEditLayerRendererState> GetEditLayerRendererStates(const UE::Landscape::EditLayers::FMergeContext* InMergeContext) override;
 
 	// UObject
-	virtual void PostLoad() override;
 	virtual void PostEditUndo() override;
 
 	// ILandscapeBrushRenderCallAdapter_GlobalMergeLegacySupport
@@ -116,10 +113,4 @@ private:
 	// Used in legacy global merge path to pass transform information from Initialize to RenderLayer
 	FTransform HeightmapCoordsToWorld;
 #endif // WITH_EDITOR
-
-#if WITH_EDITORONLY_DATA
-	// TODO: Remove this once guid is stored in the base edit layer UObject
-	UPROPERTY()
-	FGuid EditLayerGuid;
-#endif // WITH_EDITORONLY_DATA
 };
diff --git a/Engine/Plugins/Experimental/Water/Source/Editor/Private/WaterLandscapeBrush.cpp b/Engine/Plugins/Experimental/Water/Source/Editor/Private/WaterLandscapeBrush.cpp
--- a/Engine/Plugins/Experimental/Water/Source/Editor/Private/WaterLandscapeBrush.cpp
+++ b/Engine/Plugins/Experimental/Water/Source/Editor/Private/WaterLandscapeBrush.cpp
@@ -18,7 +18,7 @@
 #include "Components/BillboardComponent.h"
 #include "Modules/ModuleManager.h"
 #include "WaterModule.h"
-#include "WaterEditorSettings.h"
+#include "LandscapeSettings.h"
 #include "LandscapeModule.h"
 #include "LandscapeEditorServices.h"
 
@@ -192,10 +192,10 @@ void AWaterLandscapeBrush::OnWaterBrushActorChanged(const IWaterBrushActorInterf
 
 	BlueprintWaterBodyChanged(Actor);
 
-	const UWaterEditorSettings* WaterEditorSettings = GetDefault<UWaterEditorSettings>();
-	check(WaterEditorSettings != nullptr);
+	const ULandscapeSettings* LandscapeSettings = GetDefault<ULandscapeSettings>();
+	check(LandscapeSettings != nullptr);
 
-	const bool bAllowLandscapeUpdate = (InParams.PropertyChangedEvent.ChangeType != EPropertyChangeType::Interactive) || WaterEditorSettings->GetShouldUpdateLandscapeDuringInteractiveChanges();
+	const bool bAllowLandscapeUpdate = (InParams.PropertyChangedEvent.ChangeType != EPropertyChangeType::Interactive) || LandscapeSettings->GetShouldUpdateEditLayersDuringInteractiveChanges();
 	if (bForceUpdateBrush || (bAffectsLandscape && bAllowLandscapeUpdate))
 	{
 		RequestLandscapeUpdate(/* bInUserTriggered = */ InParams.bUserTriggered);
diff --git a/Engine/Plugins/Experimental/Water/Source/Editor/Public/WaterEditorSettings.h b/Engine/Plugins/Experimental/Water/Source/Editor/Public/WaterEditorSettings.h
--- a/Engine/Plugins/Experimental/Water/Source/Editor/Public/WaterEditorSettings.h
+++ b/Engine/Plugins/Experimental/Water/Source/Editor/Public/WaterEditorSettings.h
@@ -184,7 +184,7 @@ struct WATEREDITOR_API FWaterBodyIslandDefaults
  * Implements the editor settings for the Water plugin.
  */
 UCLASS(config = Engine, defaultconfig, meta=(DisplayName="Water Editor"))
-class WATEREDITOR_API UWaterEditorSettings : public UDeveloperSettings
+class UWaterEditorSettings : public UDeveloperSettings
 {
 	GENERATED_BODY()
 
@@ -232,16 +232,15 @@ public:
 	UMaterialInterface* GetDefaultRenderRiverSplineDepthsMaterial() const;
 	FSoftObjectPath GetDefaultRenderRiverSplineDepthsMaterialPath() const { return DefaultRenderRiverSplineDepthsMaterial.ToSoftObjectPath(); }
 
-	bool GetShouldUpdateLandscapeDuringInteractiveChanges() const { return bUpdateLandscapeDuringInteractiveChanges; }
 	bool GetShouldUpdateWaterMeshDuringInteractiveChanges() const { return bUpdateWaterMeshDuringInteractiveChanges; }
 
 	UE_DEPRECATED(5.5, "Use GetShouldUpdateLandscapeDuringInteractiveChanges")
-	bool GetUpdateLandscapeDuringInteractiveChanges() const { return bUpdateLandscapeDuringInteractiveChanges; }
+	bool GetUpdateLandscapeDuringInteractiveChanges() const { return false; }
 	UE_DEPRECATED(5.5, "Use GetShouldUpdateWaterMeshDuringInteractiveChanges")
 	bool GetUpdateWaterMeshDuringInteractiveChanges() const { return bUpdateWaterMeshDuringInteractiveChanges; }
 	
 public:
-	/** The texture group to use for generated textures such as the combined veloctiy and height texture */
+	/** The texture group to use for generated textures such as the combined velocity and height texture */
 	UPROPERTY(EditAnywhere, config, Category = Rendering)
 	TEnumAsByte<TextureGroup> TextureGroupForGeneratedTextures;
 
@@ -281,10 +280,6 @@ public:
 	FWaterBodyIslandDefaults WaterBodyIslandDefaults;
 
 private:
-	/** For landscape-affecting water bodies, allows the landscape to be updated when the water body's shape is modified interactively (e.g. when dragging a spline point). Set to false if the performance when editing a water body gets too bad (the landscape will be properly updated when the dragging operation is done). */
-	UPROPERTY(EditAnywhere, config, Category = Brush)
-	bool bUpdateLandscapeDuringInteractiveChanges = false;
-
 	/** Allows the water mesh to be updated when the water body's shape is modified interactively (e.g. when dragging a spline point). Set to false if the performance when editing a water body gets too bad (the water mesh will be properly updated when the dragging operation is done). */
 	UPROPERTY(EditAnywhere, config, Category = Brush)
 	bool bUpdateWaterMeshDuringInteractiveChanges = false;
diff --git a/Engine/Source/Editor/DetailCustomizations/Private/DetailCustomizations.cpp b/Engine/Source/Editor/DetailCustomizations/Private/DetailCustomizations.cpp
--- a/Engine/Source/Editor/DetailCustomizations/Private/DetailCustomizations.cpp
+++ b/Engine/Source/Editor/DetailCustomizations/Private/DetailCustomizations.cpp
@@ -78,8 +78,6 @@
 #include "Internationalization/Internationalization.h"
 #include "IntervalStructCustomization.h"
 #include "KeyStructCustomization.h"
-#include "LandscapeProxyUIDetails.h"
-#include "LandscapeUIDetails.h"
 #include "LevelSequenceActorDetails.h"
 #include "LevelSequenceBurnInOptionsCustomization.h"
 #include "LightComponentDetails.h"
@@ -161,7 +159,6 @@
 #include "Vector4StructCustomization.h"
 #include "WindowsTargetSettingsDetails.h"
 #include "WorldSettingsDetails.h"
-#include "LandscapeGrassTypeDetails.h"
 
 struct FPerPlatformBool;
 struct FPerPlatformFloat;
@@ -495,10 +492,6 @@ void FDetailCustomizationsModule::RegisterObjectCustomizations()
 
 	RegisterCustomClassLayout("MotionControllerComponent", FOnGetDetailCustomizationInstance::CreateStatic(&FMotionControllerDetails::MakeInstance));
 
-	RegisterCustomClassLayout("Landscape", FOnGetDetailCustomizationInstance::CreateStatic(&FLandscapeUIDetails::MakeInstance));
-	RegisterCustomClassLayout("LandscapeProxy", FOnGetDetailCustomizationInstance::CreateStatic(&FLandscapeProxyUIDetails::MakeInstance));
-	RegisterCustomClassLayout("LandscapeGrassType", FOnGetDetailCustomizationInstance::CreateStatic(&FLandscapeGrassTypeDetails::MakeInstance));
-
 	RegisterCustomClassLayout("BoundsCopyComponent", FOnGetDetailCustomizationInstance::CreateStatic(&FBoundsCopyComponentDetailsCustomization::MakeInstance));
 
 	RegisterCustomClassLayout("SoundNodeDistanceCrossFade", FOnGetDetailCustomizationInstance::CreateStatic(&FCrossFadeCustomization::MakeInstance));
diff --git a/Engine/Source/Editor/DetailCustomizations/Private/LandscapeGrassTypeDetails.cpp b/Engine/Source/Editor/DetailCustomizations/Private/LandscapeGrassTypeDetails.cpp
deleted file mode 100644
diff --git a/Engine/Source/Editor/DetailCustomizations/Private/LandscapeGrassTypeDetails.h b/Engine/Source/Editor/DetailCustomizations/Private/LandscapeGrassTypeDetails.h
deleted file mode 100644
diff --git a/Engine/Source/Editor/DetailCustomizations/Private/LandscapeProxyUIDetails.cpp b/Engine/Source/Editor/DetailCustomizations/Private/LandscapeProxyUIDetails.cpp
deleted file mode 100644
diff --git a/Engine/Source/Editor/DetailCustomizations/Private/LandscapeProxyUIDetails.h b/Engine/Source/Editor/DetailCustomizations/Private/LandscapeProxyUIDetails.h
deleted file mode 100644
diff --git a/Engine/Source/Editor/DetailCustomizations/Private/LandscapeUIDetails.cpp b/Engine/Source/Editor/DetailCustomizations/Private/LandscapeUIDetails.cpp
deleted file mode 100644
diff --git a/Engine/Source/Editor/DetailCustomizations/Private/LandscapeUIDetails.h b/Engine/Source/Editor/DetailCustomizations/Private/LandscapeUIDetails.h
deleted file mode 100644
diff --git a/Engine/Source/Editor/EditorStyle/Private/StarshipStyle.cpp b/Engine/Source/Editor/EditorStyle/Private/StarshipStyle.cpp
--- a/Engine/Source/Editor/EditorStyle/Private/StarshipStyle.cpp
+++ b/Engine/Source/Editor/EditorStyle/Private/StarshipStyle.cpp
@@ -7177,6 +7177,8 @@ void FStarshipEditorStyle::FStyle::SetupLandscapeEditorStyle()
 		Set("LandscapeEditor.Brush.AffectsVisibilityLayer.Enabled", new IMAGE_BRUSH("Icons/icon_Landscape_Affects_VisibilityLayer_Enabled_16x", Icon16x16));
 		Set("LandscapeEditor.Brush.AffectsVisibilityLayer.Disabled", new IMAGE_BRUSH("Icons/icon_Landscape_Affects_VisibilityLayer_Disabled_16x", Icon16x16));
 
+		Set("LandscapeEditor.InspectedObjects.ShowDetails", new CORE_IMAGE_BRUSH_SVG("Starship/Common/settings", Icon16x16));
+
 		// Target List
 		Set("LandscapeEditor.TargetList.RowBackground",        new FSlateNoResource());
 		Set("LandscapeEditor.TargetList.RowBackgroundHovered", new BOX_BRUSH("Common/RoundedSelection_16x", 4.0f / 16.0f, FStyleColors::SelectHover));
diff --git a/Engine/Source/Editor/LandscapeEditor/LandscapeEditor.Build.cs b/Engine/Source/Editor/LandscapeEditor/LandscapeEditor.Build.cs
--- a/Engine/Source/Editor/LandscapeEditor/LandscapeEditor.Build.cs
+++ b/Engine/Source/Editor/LandscapeEditor/LandscapeEditor.Build.cs
@@ -34,7 +34,8 @@ public class LandscapeEditor : ModuleRules
 				"DeveloperSettings",
 				"PlacementMode",
 				"CQTest",
-				"NavigationSystem"
+				"NavigationSystem",
+				"VirtualTexturingEditor",
 			}
 			);
 
diff --git a/Engine/Source/Editor/LandscapeEditor/Private/LandscapeEdMode.cpp b/Engine/Source/Editor/LandscapeEditor/Private/LandscapeEdMode.cpp
--- a/Engine/Source/Editor/LandscapeEditor/Private/LandscapeEdMode.cpp
+++ b/Engine/Source/Editor/LandscapeEditor/Private/LandscapeEdMode.cpp
@@ -327,7 +327,7 @@ ELandscapeToolTargetType FEdModeLandscape::GetLandscapeToolTargetType() const
 
 const FLandscapeLayer* FEdModeLandscape::GetLandscapeSelectedLayer() const
 {
-	return GetCurrentLayer();
+	return GetCurrentLayerConst();
 }
 
 ULandscapeLayerInfoObject* FEdModeLandscape::GetSelectedLandscapeLayerInfo() const
@@ -407,6 +407,12 @@ int32 FEdModeLandscape::GetNewLandscapeResolutionY() const
 	return UISettings->NewLandscape_ComponentCount.Y * UISettings->NewLandscape_SectionsPerComponent * UISettings->NewLandscape_QuadsPerSection + 1;
 }
 
+void FEdModeLandscape::SetInspectedObjects(const TArray<TWeakObjectPtr<UObject>>& InObjects)
+{
+	InspectedObjects = InObjects;
+	RefreshInspectedObjectsDetailPanel();
+}
+
 /** FEdMode: Called when the mode is entered */
 void FEdModeLandscape::Enter()
 {
@@ -1650,7 +1656,7 @@ bool FEdModeLandscape::ProcessEditPaste()
 	if (CurrentTool == (FLandscapeTool*)SplinesTool)
 	{
 		ALandscape* Landscape = GetLandscape();
-		SplinesLayer = Landscape ? Landscape->FindLayerOfType(ULandscapeEditLayerSplines::StaticClass()) : nullptr;
+		SplinesLayer = Landscape ? Landscape->FindLayerOfTypeConst(ULandscapeEditLayerSplines::StaticClass()) : nullptr;
 	}
 	FText Reason;
 	if (!CanEditLayer(&Reason, SplinesLayer))
@@ -2244,6 +2250,17 @@ void FEdModeLandscape::RefreshDetailPanel()
 	}
 }
 
+void FEdModeLandscape::RefreshInspectedObjectsDetailPanel()
+{
+	// Inspected objects update separate from the main details panel 
+	// When inspecting an object, we don't want the main panel to refresh because this resets the scroll 
+	// navigating the user away from their current view
+	if (Toolkit.IsValid() && !bUpdatingLandscapeInfo)
+	{
+		StaticCastSharedPtr<FLandscapeToolKit>(Toolkit)->RefreshInspectedObjectsDetailPanel();
+	}
+}
+
 void FEdModeLandscape::RegenerateLayerThumbnails()
 {
 	// Regenerate the lists that might hold on to a ThumbnailMIC : 
@@ -2482,7 +2499,7 @@ bool FEdModeLandscape::IsGridBased() const
 
 bool FEdModeLandscape::HasValidLandscapeEditLayerSelection() const
 {
-	return !CanHaveLandscapeLayersContent() || GetCurrentLayer();
+	return !CanHaveLandscapeLayersContent() || GetCurrentLayerConst();
 }
 
 void FEdModeLandscape::SetTargetLandscape(const TWeakObjectPtr<ULandscapeInfo>& InLandscapeInfo)
@@ -2752,10 +2769,7 @@ void FEdModeLandscape::UpdateTargetLayerDisplayOrder(ELandscapeLayerDisplayMode
 
 	if (DetailPanelRefreshRequired)
 	{
-		if (Toolkit.IsValid())
-		{
-			StaticCastSharedPtr<FLandscapeToolKit>(Toolkit)->RefreshDetailPanel();
-		}
+		RefreshDetailPanel();
 	}
 }
 
@@ -2828,10 +2842,7 @@ void FEdModeLandscape::UpdateLayerUsageInformation(TWeakObjectPtr<ULandscapeLaye
 	
 	if (DetailPanelRefreshRequired)
 	{
-		if (Toolkit.IsValid())
-		{
-			StaticCastSharedPtr<FLandscapeToolKit>(Toolkit)->RefreshDetailPanel();
-		}
+		RefreshDetailPanel();
 	}
 }
 
@@ -3676,9 +3687,9 @@ ALandscape* FEdModeLandscape::ChangeComponentSetting(int32 NumComponentsX, int32
 				{
 					int32 HeightCount = 0;
 
-					for (const FLandscapeLayer& OldLayer : OldLandscape->GetLayers())
+					for (const FLandscapeLayer& OldLayer : OldLandscape->GetLayersConst())
 					{
-						FScopedSetLandscapeEditingLayer Scope(OldLandscape, OldLayer.Guid);
+						FScopedSetLandscapeEditingLayer Scope(OldLandscape, OldLayer.EditLayer->GetGuid());
 
 						TArray<uint16> HeightData;
 						TArray<FLandscapeImportLayerInfo> ImportMaterialLayerInfos;
@@ -3686,8 +3697,8 @@ ALandscape* FEdModeLandscape::ChangeComponentSetting(int32 NumComponentsX, int32
 						ExtractHeightmapWeightmapContent(HeightData, ImportMaterialLayerInfos);
 
 						HeightCount = FMath::Max(HeightCount, HeightData.Num());
-						HeightDataPerLayers.Add(OldLayer.Guid, MoveTemp(HeightData));
-						ImportMaterialLayerInfosPerLayers.Add(OldLayer.Guid, MoveTemp(ImportMaterialLayerInfos));
+						HeightDataPerLayers.Add(OldLayer.EditLayer->GetGuid(), MoveTemp(HeightData));
+						ImportMaterialLayerInfosPerLayers.Add(OldLayer.EditLayer->GetGuid(), MoveTemp(ImportMaterialLayerInfos));
 					}
 
 					TArray<uint16> DefaultHeightData;
@@ -3754,7 +3765,7 @@ ALandscape* FEdModeLandscape::ChangeComponentSetting(int32 NumComponentsX, int32
 			TArrayView<const FLandscapeLayer> LandscapeLayers;
 			if (CanHaveLandscapeLayersContent())
 			{
-				LandscapeLayers = OldLandscape->GetLayers();
+				LandscapeLayers = OldLandscape->GetLayersConst();
 			} 
 			NewLandscape->Import(FGuid::NewGuid(), NewMinX, NewMinY, NewMaxX, NewMaxY, NumSubsections, SubsectionSizeQuads, HeightDataPerLayers, *OldLandscape->ReimportHeightmapFilePath, ImportMaterialLayerInfosPerLayers, ELandscapeImportAlphamapType::Additive, LandscapeLayers);
 
@@ -3952,7 +3963,7 @@ bool FEdModeLandscape::HasLandscapeLayersContent() const
 int32 FEdModeLandscape::GetLayerCount() const
 {
 	ALandscape* Landscape = GetLandscape();
-	return Landscape ? Landscape->GetLayerCount() : 0;
+	return Landscape ? Landscape->GetLayersConst().Num() : 0;
 }
 
 void FEdModeLandscape::SetCurrentLayer(int32 InLayerIndex)
@@ -3986,130 +3997,36 @@ ALandscape* FEdModeLandscape::GetLandscape() const
 	return CurrentToolTarget.LandscapeInfo.IsValid() ? CurrentToolTarget.LandscapeInfo->LandscapeActor.Get() : nullptr;
 }
 
-const FLandscapeLayer* FEdModeLandscape::GetLayer(int32 InLayerIndex) const
+const FLandscapeLayer* FEdModeLandscape::GetLayerConst(int32 InLayerIndex) const
 {
 	const ALandscape* Landscape = GetLandscape();
 	return Landscape ? Landscape->GetLayerConst(InLayerIndex) : nullptr;
 }
 
-FName FEdModeLandscape::GetLayerName(int32 InLayerIndex) const
-{
-	const FLandscapeLayer* Layer = GetLayer(InLayerIndex);
-	return Layer ? Layer->Name : NAME_None;
-}
-
-bool FEdModeLandscape::CanRenameLayerTo(int32 InLayerIndex, const FName& InNewName)
+bool FEdModeLandscape::CanRenameLayerTo(int32 InLayerIndex, const FName& InNewName) const
 {
-	ALandscape* Landscape = GetLandscape();
-	if (Landscape)
-	{
-		int32 LayerCount = GetLayerCount();
-		for (int32 LayerIdx = 0; LayerIdx < LayerCount; ++LayerIdx)
-		{
-			if (LayerIdx != InLayerIndex && GetLayerName(LayerIdx) == InNewName)
+	if (ALandscape* Landscape = GetLandscape(); Landscape && !Landscape->IsLayerNameUnique(InNewName))
 			{
 				return false;
 			}
-		}
-	}
 	return true;
 }
 
-void FEdModeLandscape::SetLayerName(int32 InLayerIndex, const FName& InName)
-{
-	ALandscape* Landscape = GetLandscape();
-	if (Landscape)
-	{
-		Landscape->SetLayerName(InLayerIndex, InName);
-	}
-}
-
 bool FEdModeLandscape::IsLayerAlphaVisible(int32 InLayerIndex) const
 {
 	return (CurrentToolTarget.TargetType == ELandscapeToolTargetType::Heightmap || CurrentToolTarget.TargetType == ELandscapeToolTargetType::Weightmap);
 }
 
-float FEdModeLandscape::GetClampedLayerAlpha(float InLayerAlpha) const
-{
-	ALandscape* Landscape = GetLandscape();
-	if (Landscape)
-	{
-		if (CurrentToolTarget.TargetType == ELandscapeToolTargetType::Heightmap || CurrentToolTarget.TargetType == ELandscapeToolTargetType::Weightmap)
-		{
-			return Landscape->GetClampedLayerAlpha(InLayerAlpha, CurrentToolTarget.TargetType == ELandscapeToolTargetType::Heightmap);
-		}
-	}
-	return InLayerAlpha;
-}
-
-float FEdModeLandscape::GetLayerAlpha(int32 InLayerIndex) const
-{
-	ALandscape* Landscape = GetLandscape();
-	if (Landscape)
-	{
-		if (CurrentToolTarget.TargetType == ELandscapeToolTargetType::Heightmap || CurrentToolTarget.TargetType == ELandscapeToolTargetType::Weightmap)
-		{
-			return Landscape->GetLayerAlpha(InLayerIndex, CurrentToolTarget.TargetType == ELandscapeToolTargetType::Heightmap);
-		}
-	}
-	return 1.0f;
-}
-
-const ULandscapeEditLayerBase* FEdModeLandscape::GetEditLayer(int32 InLayerIndex) const
+const ULandscapeEditLayerBase* FEdModeLandscape::GetEditLayerConst(int32 InLayerIndex) const
 {
 	const ALandscape* Landscape = GetLandscape();
-	if (Landscape)
-	{
-		if (const FLandscapeLayer* EditLayer = Landscape->GetLayerConst(InLayerIndex))
-		{
-			check(EditLayer->EditLayer != nullptr);
-			return EditLayer->EditLayer;
-		}
-	}
-	return nullptr;
-}
-
-
-void FEdModeLandscape::SetLayerAlpha(int32 InLayerIndex, float InAlpha)
-{
-	ALandscape* Landscape = GetLandscape();
-	if (Landscape)
-	{
-		if (CurrentToolTarget.TargetType == ELandscapeToolTargetType::Heightmap || CurrentToolTarget.TargetType == ELandscapeToolTargetType::Weightmap)
-		{
-			Landscape->SetLayerAlpha(InLayerIndex, InAlpha, CurrentToolTarget.TargetType == ELandscapeToolTargetType::Heightmap);
-		}
-	}
-}
-
-bool FEdModeLandscape::IsLayerVisible(int32 InLayerIndex) const
-{
-	const FLandscapeLayer* Layer = GetLayer(InLayerIndex);
-	return Layer ? Layer->bVisible : false;
-}
-
-void FEdModeLandscape::SetLayerVisibility(bool bInVisible, int32 InLayerIndex)
-{
-	ALandscape* Landscape = GetLandscape();
-	if (Landscape)
-	{
-		Landscape->SetLayerVisibility(InLayerIndex, bInVisible);
-	}
+	return Landscape ? Landscape->GetEditLayerConst(InLayerIndex) : nullptr;
 }
 
-bool FEdModeLandscape::IsLayerLocked(int32 InLayerIndex) const
+ULandscapeEditLayerBase* FEdModeLandscape::GetEditLayer(int32 InLayerIndex) const
 {
-	const FLandscapeLayer* Layer = GetLayer(InLayerIndex);
-	return Layer && Layer->bLocked;
-}
-
-void FEdModeLandscape::SetLayerLocked(int32 InLayerIndex, bool bInLocked)
-{
-	ALandscape* Landscape = GetLandscape();
-	if (Landscape)
-	{
-		Landscape->SetLayerLocked(InLayerIndex, bInLocked);
-	}
+	const ALandscape* Landscape = GetLandscape();
+	return Landscape ? Landscape->GetEditLayer(InLayerIndex) : nullptr;
 }
 
 void FEdModeLandscape::RequestLayersContentUpdate(ELandscapeLayerUpdateMode InUpdateMode)
@@ -4193,31 +4110,9 @@ void FEdModeLandscape::DuplicateBrush(class ALandscapeBlueprintBrushBase* InBrus
 	GEditor->edactDuplicateSelected(InBrush->GetLevel(), false);
 }
 
-bool FEdModeLandscape::IsCurrentLayerBlendSubstractive(const TWeakObjectPtr<ULandscapeLayerInfoObject>& InLayerInfoObj) const
-{
-	ALandscape* Landscape = GetLandscape();
-
-	if (Landscape != nullptr)
-	{
-		return Landscape->IsLayerBlendSubstractive(GetCurrentLayerIndex(), InLayerInfoObj);
-	}
-
-	return false;
-}
-
-void FEdModeLandscape::SetCurrentLayerSubstractiveBlendStatus(bool InStatus, const TWeakObjectPtr<ULandscapeLayerInfoObject>& InLayerInfoObj)
-{
-	ALandscape* Landscape = GetLandscape();
-
-	if (Landscape != nullptr)
-	{
-		return Landscape->SetLayerSubstractiveBlendStatus(GetCurrentLayerIndex(), InStatus, InLayerInfoObj);
-	}
-}
-
-const FLandscapeLayer* FEdModeLandscape::GetCurrentLayer() const
+const FLandscapeLayer* FEdModeLandscape::GetCurrentLayerConst() const
 {
-	return GetLayer(GetCurrentLayerIndex());
+	return GetLayerConst(GetCurrentLayerIndex());
 }
 
 void FEdModeLandscape::AutoUpdateDirtyLandscapeSplines()
@@ -4239,7 +4134,7 @@ bool FEdModeLandscape::CanEditLayer(FText* Reason /*=nullptr*/, const FLandscape
 	if (CanHaveLandscapeLayersContent())
 	{
 		ALandscape* Landscape = GetLandscape();
-		const FLandscapeLayer* TargetLayer = InLayer ? InLayer : GetCurrentLayer();
+		const FLandscapeLayer* TargetLayer = InLayer ? InLayer : GetCurrentLayerConst();
 		if (!TargetLayer)
 		{
 			if (Reason)
@@ -4250,7 +4145,7 @@ bool FEdModeLandscape::CanEditLayer(FText* Reason /*=nullptr*/, const FLandscape
 		}
 
 		check(TargetLayer->EditLayer != nullptr);
-		if (!TargetLayer->bVisible)
+		if (!TargetLayer->EditLayer->IsVisible())
 		{
 			if (Reason)
 			{
@@ -4258,7 +4153,7 @@ bool FEdModeLandscape::CanEditLayer(FText* Reason /*=nullptr*/, const FLandscape
 			}
 			return false;
 		}
-		else if (TargetLayer->bLocked)
+		else if (TargetLayer->EditLayer->IsLocked())
 		{
 			if (Reason)
 			{
@@ -4271,7 +4166,7 @@ bool FEdModeLandscape::CanEditLayer(FText* Reason /*=nullptr*/, const FLandscape
 			// Special case for the splines tool, which is supported on the splines layer : 
 			if (CurrentTool != (FLandscapeTool*)SplinesTool)
 			{
-				if (Landscape && (TargetLayer == Landscape->FindLayerOfType(ULandscapeEditLayerSplines::StaticClass())))
+				if (Landscape && (TargetLayer == Landscape->FindLayerOfTypeConst(ULandscapeEditLayerSplines::StaticClass())))
 				{
 					if (Reason)
 					{
@@ -4337,8 +4232,8 @@ void FEdModeLandscape::UpdateLandscapeSplines(bool bUpdateOnlySelected /* = fals
 
 FGuid FEdModeLandscape::GetCurrentLayerGuid() const
 {
-	const FLandscapeLayer* CurrentLayer = GetCurrentLayer();
-	return CurrentLayer ? CurrentLayer->Guid : FGuid();
+	const FLandscapeLayer* CurrentLayer = GetCurrentLayerConst();
+	return CurrentLayer ? CurrentLayer->EditLayer->GetGuid() : FGuid();
 }
 
 void FEdModeLandscape::UpdateBrushList()
diff --git a/Engine/Source/Editor/LandscapeEditor/Private/LandscapeEdMode.h b/Engine/Source/Editor/LandscapeEditor/Private/LandscapeEdMode.h
--- a/Engine/Source/Editor/LandscapeEditor/Private/LandscapeEdMode.h
+++ b/Engine/Source/Editor/LandscapeEditor/Private/LandscapeEdMode.h
@@ -576,6 +576,7 @@ public:
 	void UpdateLayerUsageInformation(TWeakObjectPtr<ULandscapeLayerInfoObject>* LayerInfoObjectThatChanged = nullptr);
 	void OnLandscapeMaterialChangedDelegate();
 	void RefreshDetailPanel();
+	void RefreshInspectedObjectsDetailPanel();
 	void RegenerateLayerThumbnails();
 
 	bool IsGridBased() const;
@@ -588,23 +589,13 @@ public:
 	void SetCurrentLayer(int32 InLayerIndex);
 	int32 GetCurrentLayerIndex() const;
 	ALandscape* GetLandscape() const;
-	const struct FLandscapeLayer* GetLayer(int32 InLayerIndex) const;
-	FName GetLayerName(int32 InLayerIndex) const;
-	void SetLayerName(int32 InLayerIndex, const FName& InName);
-	bool CanRenameLayerTo(int32 InLayerIndex, const FName& InNewName);
-	void SetLayerAlpha(int32 InLayerIndex, float InAlpha);
-	float GetLayerAlpha(int32 InLayerIndex) const;
-	const ULandscapeEditLayerBase* GetEditLayer(int32 InLayerIndex) const;
-	float GetClampedLayerAlpha(float InLayerAlpha) const;
-	void SetLayerVisibility(bool InVisible, int32 InLayerIndex);
-	bool IsLayerVisible(int32 InLayerIndex) const;
-	bool IsLayerLocked(int32 InLayerIndex) const;
+	const struct FLandscapeLayer* GetLayerConst(int32 InLayerIndex) const;
+	bool CanRenameLayerTo(int32 InLayerIndex, const FName& InNewName) const;
+	ULandscapeEditLayerBase* GetEditLayer(int32 InLayerIndex) const;
+	const ULandscapeEditLayerBase* GetEditLayerConst(int32 InLayerIndex) const;
 	bool IsLayerAlphaVisible(int32 InLayerIndex) const;
-	void SetLayerLocked(int32 InLayerIndex, bool bInLocked);
-	const struct FLandscapeLayer* GetCurrentLayer() const;
+	const struct FLandscapeLayer* GetCurrentLayerConst() const;
 	FGuid GetCurrentLayerGuid() const;
-	bool IsCurrentLayerBlendSubstractive(const TWeakObjectPtr<ULandscapeLayerInfoObject>& InLayerInfoObj) const;
-	void SetCurrentLayerSubstractiveBlendStatus(bool InStatus, const TWeakObjectPtr<ULandscapeLayerInfoObject>& InLayerInfoObj);
 	void UpdateLandscapeSplines(bool bUpdateOnlySelected = false);
 	void AutoUpdateDirtyLandscapeSplines();
 	bool CanEditLayer(FText* Reason = nullptr, const FLandscapeLayer* InLayer = nullptr);
@@ -676,6 +667,12 @@ public:
 	int32 GetNewLandscapeResolutionX() const;
 	int32 GetNewLandscapeResolutionY() const;
 	
+	TArray<TWeakObjectPtr<UObject>> GetInspectedObjects() const 
+	{ 
+		return InspectedObjects;
+	}
+	void SetInspectedObjects(const TArray<TWeakObjectPtr<UObject>>& InObjects);
+
 private:
 	TArray<TSharedRef<FLandscapeTargetListInfo>> LandscapeTargetList;
 	TArray<FLandscapeListInfo> LandscapeList;
@@ -706,4 +703,8 @@ private:
 	/** Delayed refresh */
 	bool bNeedsUpdateLayerUsageInformation;
 	bool bUpdatingLandscapeInfo;
+
+	/** The landscape mode panel has the ability to display the properties of certain UObjects if needed (to bypass the fact that the standard details 
+	 panel is not able to display anything else than actors and components) */
+	TArray<TWeakObjectPtr<UObject>> InspectedObjects;
 };
\ No newline at end of file
diff --git a/Engine/Source/Editor/LandscapeEditor/Private/LandscapeEdModeComponentTools.cpp b/Engine/Source/Editor/LandscapeEditor/Private/LandscapeEdModeComponentTools.cpp
--- a/Engine/Source/Editor/LandscapeEditor/Private/LandscapeEdModeComponentTools.cpp
+++ b/Engine/Source/Editor/LandscapeEditor/Private/LandscapeEdModeComponentTools.cpp
@@ -602,11 +602,11 @@ namespace
 				TArray<ULandscapeComponent*> ComponentsUsingHeightmap;
 				ComponentsUsingHeightmap.Add(NewComponent);
 
-				for (const FLandscapeLayer& Layer : Landscape->GetLayers())
+				for (const FLandscapeLayer& Layer : Landscape->GetLayersConst())
 				{
 					// Since we do not share heightmap when adding new component, we will provided the required array, but they will only be used for 1 component
 					TMap<UTexture2D*, UTexture2D*> CreatedHeightmapTextures;
-					NewComponent->AddDefaultLayerData(Layer.Guid, ComponentsUsingHeightmap, CreatedHeightmapTextures);
+					NewComponent->AddDefaultLayerData(Layer.EditLayer->GetGuid(), ComponentsUsingHeightmap, CreatedHeightmapTextures);
 				}
 			}
 
@@ -766,11 +766,11 @@ public:
 					TArray<ULandscapeComponent*> ComponentsUsingHeightmap;
 					ComponentsUsingHeightmap.Add(NewComponent);
 
-					for (const FLandscapeLayer& Layer : Landscape->GetLayers())
+					for (const FLandscapeLayer& Layer : Landscape->GetLayersConst())
 					{
 						// Since we do not share heightmap when adding new component, we will provide the required array, but they will only be used for 1 component
 						TMap<UTexture2D*, UTexture2D*> CreatedHeightmapTextures;
-						NewComponent->AddDefaultLayerData(Layer.Guid, ComponentsUsingHeightmap, CreatedHeightmapTextures);
+						NewComponent->AddDefaultLayerData(Layer.EditLayer->GetGuid(), ComponentsUsingHeightmap, CreatedHeightmapTextures);
 					}
 				}
 
diff --git a/Engine/Source/Editor/LandscapeEditor/Private/LandscapeEdModeTools.h b/Engine/Source/Editor/LandscapeEditor/Private/LandscapeEdModeTools.h
--- a/Engine/Source/Editor/LandscapeEditor/Private/LandscapeEdModeTools.h
+++ b/Engine/Source/Editor/LandscapeEditor/Private/LandscapeEdModeTools.h
@@ -16,6 +16,7 @@
 #include "InstancedFoliageActor.h"
 #include "AI/NavigationSystemBase.h"
 #include "Landscape.h"
+#include "LandscapeEditLayer.h"
 #include "LandscapeEditorPrivate.h"
 #include "Logging/LogMacros.h"
 #include "VisualLogger/VisualLogger.h"
@@ -897,16 +898,16 @@ public:
 			if (bCombinedLayerOperation)
 			{
 				int32 I = 0;
-				for (const FLandscapeLayer& CurrentLayer : Landscape->GetLayers())
+				for (const FLandscapeLayer& CurrentLayer : Landscape->GetLayersConst())
 				{
-					BackupLayerVisibility.Add(CurrentLayer.bVisible);
-					if (CurrentLayer.Guid == EditingLayerGuid.GetValue())
+					BackupLayerVisibility.Add(CurrentLayer.EditLayer->IsVisible());
+					if (CurrentLayer.EditLayer->GetGuid() == EditingLayerGuid.GetValue())
 					{
 						EditingLayerIndex = I;
 					}
 					++I;
 				}
-				check(EditingLayerIndex < Landscape->GetLayerCount());
+				check(Landscape->GetLayersConst().IsValidIndex(EditingLayerIndex));
 			}
 			bIsInitialized = true;
 		}
@@ -920,9 +921,9 @@ public:
 		{
 			TArray<bool> NewLayerVisibility;
 			int32 I = 0;
-			for (const FLandscapeLayer& CurrentLayer : Landscape->GetLayers())
+			for (const FLandscapeLayer& CurrentLayer : Landscape->GetLayersConst())
 			{
-				NewLayerVisibility.Add((I > EditingLayerIndex) ? false : CurrentLayer.bVisible);
+				NewLayerVisibility.Add((I > EditingLayerIndex) ? false : CurrentLayer.EditLayer->IsVisible());
 				++I;
 			}
 
@@ -978,7 +979,12 @@ public:
 		check(bIsInitialized);
 		if (bCombinedLayerOperation)
 		{
-			const float Alpha = Landscape->GetLayerAlpha(EditingLayerIndex, bTargetIsHeightmap);
+			const FLandscapeLayer* Layer = Landscape->GetLayerConst(EditingLayerIndex);
+			check(Layer != nullptr);
+			const ULandscapeEditLayerBase* EditLayer = Layer->EditLayer;
+			check(EditLayer != nullptr);
+			check(bTargetIsHeightmap);
+			const float Alpha = EditLayer->GetAlphaForTargetType(ELandscapeToolTargetType::Heightmap);
 			const float InverseAlpha = (Alpha != 0.f) ? 1.f / Alpha : 1.f;
 			TArray<typename ToolTarget::CacheClass::DataType> DataContribution;
 			DataContribution.Empty(Data.Num());
@@ -992,8 +998,8 @@ public:
 
 			FGuid CacheAccessorLayerGuid = CacheUpToEditingLayer.DataAccess.GetEditLayer();
  			checkf(EditingLayerGuid.GetValue() == CacheAccessorLayerGuid, TEXT("Editing Layer has changed between Initialize and Write. Was: %s (%s). Is now: %s (%s)"),
- 				Landscape->GetLayerConst(*EditingLayerGuid) ? *(Landscape->GetLayerConst(*EditingLayerGuid)->Name.ToString()) : TEXT("<unknown>"), *EditingLayerGuid->ToString(),
- 				Landscape->GetLayerConst(CacheAccessorLayerGuid) ? *(Landscape->GetLayerConst(CacheAccessorLayerGuid)->Name.ToString()) : TEXT("<unknown>"), *CacheAccessorLayerGuid.ToString());
+ 				Landscape->GetLayerConst(*EditingLayerGuid) ? *(Landscape->GetLayerConst(*EditingLayerGuid)->EditLayer->GetName().ToString()) : TEXT("<unknown>"), *EditingLayerGuid->ToString(),
+ 				Landscape->GetLayerConst(CacheAccessorLayerGuid) ? *(Landscape->GetLayerConst(CacheAccessorLayerGuid)->EditLayer->GetName().ToString()) : TEXT("<unknown>"), *CacheAccessorLayerGuid.ToString());
 
 			// Restore layers visibility
 			SetLayersVisibility(BackupLayerVisibility);
@@ -1032,16 +1038,16 @@ private:
 
 	void SetLayersVisibility(const TArray<bool>& InLayerVisibility)
 	{
-		check(InLayerVisibility.Num() == Landscape->GetLayerCount());
+		check(InLayerVisibility.Num() == Landscape->GetLayersConst().Num());
 		for (int i = 0; i < InLayerVisibility.Num(); ++i)
 		{
 			if (const FLandscapeLayer* Layer = Landscape->GetLayerConst(i))
 			{
-				if (Layer->bVisible != InLayerVisibility[i])
+				if (Layer->EditLayer->IsVisible() != InLayerVisibility[i])
 				{
-					// Pass bInForIntermediateRender = true here, because we don't want to call Modify on the landscape, we don't actually want to truly change 
+					// Pass bInModify = false here, because we don't want to call Modify on the landscape, we don't actually want to truly change 
 					//  the layer visibility, only temporarily for the purpose of the intermediate render : 
-					Landscape->SetLayerVisibility(i, InLayerVisibility[i], /*bInForIntermediateRender = */true);
+					Layer->EditLayer->SetVisible(InLayerVisibility[i], /*bInModify = */ false);
 					bVisibilityChanged = true;
 				}
 			}
diff --git a/Engine/Source/Editor/LandscapeEditor/Private/LandscapeEditLayerObjectNameEditSink.h b/Engine/Source/Editor/LandscapeEditor/Private/LandscapeEditLayerObjectNameEditSink.h
new file mode 100644
--- /dev/null
+++ b/Engine/Source/Editor/LandscapeEditor/Private/LandscapeEditLayerObjectNameEditSink.h
@@ -0,0 +1,46 @@
+// Copyright Epic Games, Inc. All Rights Reserved
+
+#pragma once
+
+#include "LandscapeEditLayer.h"
+#include "IObjectNameEditSink.h"
+#include "ScopedTransaction.h"
+
+#define LOCTEXT_NAMESPACE "LandscapeEditLayer"
+
+class FLandscapeEditLayerObjectNameEditSink : public UE::EditorWidgets::IObjectNameEditSink
+{
+	virtual UClass* GetSupportedClass() const override
+	{
+		return ULandscapeEditLayerBase::StaticClass();
+	}
+
+	virtual FText GetObjectDisplayName(UObject* Object) const override
+	{
+		ULandscapeEditLayerBase* EditLayer = CastChecked<ULandscapeEditLayerBase>(Object);
+		return FText::FromString(EditLayer->GetName().ToString());
+	}
+
+	virtual bool IsObjectDisplayNameReadOnly(UObject* Object) const override
+	{
+		return false;
+	};
+
+	virtual bool SetObjectDisplayName(UObject* Object, FString DisplayName) override
+	{
+		if (ULandscapeEditLayerBase* EditLayer = CastChecked<ULandscapeEditLayerBase>(Object))
+		{
+			EditLayer->SetName(*DisplayName, /*bInModify =*/ true);
+			return true;
+		}
+
+		return false;
+	}
+
+	virtual FText GetObjectNameTooltip(UObject* Object) const override
+	{
+		return LOCTEXT("NonEditableLandscapeEditLayerLabel_TooltipFmt", "Edit Layer Name");
+	}
+};
+
+#undef LOCTEXT_NAMESPACE
\ No newline at end of file
diff --git a/Engine/Source/Editor/LandscapeEditor/Private/LandscapeEditorDetailCustomization_Layers.cpp b/Engine/Source/Editor/LandscapeEditor/Private/LandscapeEditorDetailCustomization_Layers.cpp
--- a/Engine/Source/Editor/LandscapeEditor/Private/LandscapeEditorDetailCustomization_Layers.cpp
+++ b/Engine/Source/Editor/LandscapeEditor/Private/LandscapeEditorDetailCustomization_Layers.cpp
@@ -169,8 +169,7 @@ void FLandscapeEditorCustomNodeBuilder_Layers::GenerateHeaderRowContent(FDetailW
 BEGIN_SLATE_FUNCTION_BUILD_OPTIMIZATION
 void FLandscapeEditorCustomNodeBuilder_Layers::GenerateChildContent(IDetailChildrenBuilder& ChildrenBuilder)
 {
-	FEdModeLandscape* LandscapeEdMode = GetEditorMode();
-	if (LandscapeEdMode != NULL)
+	if (const FEdModeLandscape* LandscapeEdMode = GetEditorMode())
 	{
 		TSharedPtr<SDragAndDropVerticalBox> LayerList = SNew(SDragAndDropVerticalBox)
 			.OnCanAcceptDropAdvanced(this, &FLandscapeEditorCustomNodeBuilder_Layers::HandleCanAcceptDrop)
@@ -187,9 +186,10 @@ void FLandscapeEditorCustomNodeBuilder_Layers::GenerateChildContent(IDetailChild
 			];
 
 		InlineTextBlocks.Empty();
-		InlineTextBlocks.AddDefaulted(LandscapeEdMode->GetLayerCount());
+		const int32 NumLayers = LandscapeEdMode->GetLayerCount();
+		InlineTextBlocks.AddDefaulted(NumLayers);
 		// Slots are displayed in the opposite order of LandscapeEditLayers
-		for (int32 i = LandscapeEdMode->GetLayerCount()-1; i >= 0 ; --i)
+		for (int32 i = NumLayers - 1; i >= 0 ; --i)
 		{
 			TSharedPtr<SWidget> GeneratedRowWidget = GenerateRow(i);
 
@@ -213,6 +213,11 @@ TSharedPtr<SWidget> FLandscapeEditorCustomNodeBuilder_Layers::GenerateRow(int32
 		FSimpleDelegate::CreateSP(this, &FLandscapeEditorCustomNodeBuilder_Layers::DeleteLayer, InLayerIndex),
 		MakeAttributeSPLambda(this, [this, InLayerIndex] { FText Reason; CanDeleteLayer(InLayerIndex, Reason); return Reason; }),
 		MakeAttributeSPLambda(this, [this, InLayerIndex] { FText Reason; return CanDeleteLayer(InLayerIndex, Reason); }));
+	
+	TSharedRef<SWidget> InspectObjectButton = PropertyCustomizationHelpers::MakeCustomButton(
+		FAppStyle::GetBrush(TEXT("LandscapeEditor.InspectedObjects.ShowDetails")),
+		FSimpleDelegate::CreateSP(this, &FLandscapeEditorCustomNodeBuilder_Layers::OnSetInspectedDetailsToEditLayer, InLayerIndex),
+		LOCTEXT("LandscapeEditLayerInspect", "Inspect the edit layer in the Landscape Details panel"));
 
 	FEdModeLandscape* LandscapeEdMode = GetEditorMode();
 	TSharedPtr<SWidget> RowWidget = SNew(SLandscapeEditorSelectableBorder)
@@ -337,9 +342,15 @@ TSharedPtr<SWidget> FLandscapeEditorCustomNodeBuilder_Layers::GenerateRow(int32
 						[
 							SNew(SNumericEntryBox<float>)
 								.AllowSpin(true)
-								.MinValue(-1.0f)
+								.MinValue_Lambda([this]
+								{
+									return GetLayerAlphaMinValue();
+								})
 								.MaxValue(1.0f)
-								.MinSliderValue(-1.0f)
+								.MinSliderValue_Lambda([this]
+								{
+									return GetLayerAlphaMinValue();
+								})
 								.MaxSliderValue(1.0f)
 								.Delta(0.01f)
 								.MinDesiredValueWidth(60.0f)
@@ -350,18 +361,27 @@ TSharedPtr<SWidget> FLandscapeEditorCustomNodeBuilder_Layers::GenerateRow(int32
 								.OnValueChanged_Lambda([this, InLayerIndex](float InValue) { SetLayerAlpha(InValue, InLayerIndex, false); })
 								.OnValueCommitted_Lambda([this, InLayerIndex](float InValue, ETextCommit::Type InCommitType) { SetLayerAlpha(InValue, InLayerIndex, true); })
 								.OnBeginSliderMovement_Lambda([this, InLayerIndex]()
-							{
-								CurrentSlider = InLayerIndex;
-								GEditor->BeginTransaction(LOCTEXT("Landscape_Layers_SetAlpha", "Set Layer Alpha"));
-							})
+								{
+									CurrentSlider = InLayerIndex;
+									GEditor->BeginTransaction(LOCTEXT("Landscape_Layers_SetAlpha", "Set Layer Alpha"));
+								})
 								.OnEndSliderMovement_Lambda([this](double)
-							{
-								GEditor->EndTransaction();
-								CurrentSlider = INDEX_NONE;
-							})
+								{
+									GEditor->EndTransaction();
+									CurrentSlider = INDEX_NONE;
+								})
 						]
 				]
 
+				+ SHorizontalBox::Slot()
+				.AutoWidth()
+				.HAlign(HAlign_Center)
+				.VAlign(VAlign_Center)
+				.Padding(2.0f, 0.0f)
+				[
+					InspectObjectButton
+				]
+
 				+ SHorizontalBox::Slot()
 				.AutoWidth()
 				.HAlign(HAlign_Center)
@@ -382,7 +402,10 @@ FText FLandscapeEditorCustomNodeBuilder_Layers::GetLayerDisplayName(int32 InLaye
 	ALandscape* Landscape = LandscapeEdMode ? LandscapeEdMode->GetLandscape() : nullptr;
 	if (LandscapeEdMode && Landscape)
 	{
-		return FText::FromName(LandscapeEdMode->GetLayerName(InLayerIndex));
+		const ULandscapeEditLayerBase* EditLayer = LandscapeEdMode->GetEditLayerConst(InLayerIndex);
+		check(EditLayer != nullptr);
+
+		return FText::FromName(EditLayer->GetName());
 	}
 
 	return FText::FromString(TEXT("None"));
@@ -401,8 +424,7 @@ bool FLandscapeEditorCustomNodeBuilder_Layers::IsLayerSelected(int32 InLayerInde
 
 bool FLandscapeEditorCustomNodeBuilder_Layers::CanRenameLayerTo(const FText& InNewText, FText& OutErrorMessage, int32 InLayerIndex)
 {
-	FEdModeLandscape* LandscapeEdMode = GetEditorMode();
-	if (LandscapeEdMode)
+	if (const FEdModeLandscape* LandscapeEdMode = GetEditorMode())
 	{
 		if (!LandscapeEdMode->CanRenameLayerTo(InLayerIndex, *InNewText.ToString()))
 		{
@@ -415,11 +437,13 @@ bool FLandscapeEditorCustomNodeBuilder_Layers::CanRenameLayerTo(const FText& InN
 
 void FLandscapeEditorCustomNodeBuilder_Layers::SetLayerName(const FText& InText, ETextCommit::Type InCommitType, int32 InLayerIndex)
 {
-	FEdModeLandscape* LandscapeEdMode = GetEditorMode();
-	if (LandscapeEdMode)
+	if (const FEdModeLandscape* LandscapeEdMode = GetEditorMode())
 	{
-		const FScopedTransaction Transaction(LOCTEXT("Landscape_Layers_Rename", "Rename Layer"));
-		LandscapeEdMode->SetLayerName(InLayerIndex, *InText.ToString());
+		const FScopedTransaction Transaction(LOCTEXT("Landscape_Layers_Rename", "Rename Edit Layer"));
+
+		ULandscapeEditLayerBase* EditLayer = LandscapeEdMode->GetEditLayer(InLayerIndex);
+		EditLayer->SetName(*InText.ToString(), /*bInModify =*/ true);
+
 		OnLayerSelectionChanged(InLayerIndex);
 	}
 }
@@ -513,7 +537,7 @@ TSharedPtr<SWidget> FLandscapeEditorCustomNodeBuilder_Layers::OnLayerContextMenu
 	ALandscape* Landscape = LandscapeEdMode ? LandscapeEdMode->GetLandscape() : nullptr;
 	if (Landscape && LandscapeEdMode && LandscapeEdMode->DoesCurrentToolAffectEditLayers())
 	{
-		const FLandscapeLayer* Layer = LandscapeEdMode->GetLayer(InLayerIndex);
+		const FLandscapeLayer* Layer = LandscapeEdMode->GetLayerConst(InLayerIndex);
 		FMenuBuilder MenuBuilder(true, nullptr);
 		MenuBuilder.BeginSection("LandscapeEditorLayerActions", LOCTEXT("LandscapeEditorLayerActions.Heading", "Edit Layers"));
 		{
@@ -578,7 +602,7 @@ TSharedPtr<SWidget> FLandscapeEditorCustomNodeBuilder_Layers::OnLayerContextMenu
 		{
 			if (Layer)
 			{
-				if (Layer->bVisible)
+				if (Layer->EditLayer->IsVisible())
 				{
 					// Hide Selected Layer
 					FUIAction HideSelectedLayerAction = FUIAction(FExecuteAction::CreateSPLambda(this, [this, InLayerIndex] { OnToggleVisibility(InLayerIndex); }),
@@ -631,13 +655,16 @@ TSharedPtr<SWidget> FLandscapeEditorCustomNodeBuilder_Layers::OnLayerContextMenu
 bool FLandscapeEditorCustomNodeBuilder_Layers::CanRenameLayer(int32 InLayerIndex, FText& OutReason) const
 {
 	FEdModeLandscape* LandscapeEdMode = GetEditorMode();
-	const FLandscapeLayer* Layer = LandscapeEdMode->GetLayer(InLayerIndex);
+	const FLandscapeLayer* Layer = LandscapeEdMode->GetLayerConst(InLayerIndex);
 	if (Layer == nullptr)
 	{
 		return false;
 	}
 
-	if (Layer->bLocked)
+	const ULandscapeEditLayerBase* EditLayer = Layer->EditLayer;
+	check(EditLayer != nullptr);
+
+	if (EditLayer->IsLocked())
 	{
 		OutReason = LOCTEXT("RenameLayer_CantRenameLocked", "Cannot rename a locked edit layer");
 		return false;
@@ -659,7 +686,7 @@ void FLandscapeEditorCustomNodeBuilder_Layers::RenameLayer(int32 InLayerIndex)
 bool FLandscapeEditorCustomNodeBuilder_Layers::CanClearTargetLayerOnLayer(int32 InLayerIndex, ULandscapeLayerInfoObject* InLayerInfo, FText& OutReason) const
 {
 	FEdModeLandscape* LandscapeEdMode = GetEditorMode();
-	const FLandscapeLayer* Layer = LandscapeEdMode->GetLayer(InLayerIndex);
+	const FLandscapeLayer* Layer = LandscapeEdMode->GetLayerConst(InLayerIndex);
 	if (Layer == nullptr)
 	{
 		return false;
@@ -674,7 +701,7 @@ bool FLandscapeEditorCustomNodeBuilder_Layers::CanClearTargetLayerOnLayer(int32
 		if (!EditLayer->SupportsTargetType(ELandscapeToolTargetType::Visibility))
 		{
 			OutReason = FText::Format(LOCTEXT("ClearTargetLayer_CantClearVisibilityOnLayerType", "Cannot clear visibility : the type of layer {0} ({1}) doesn't support visibility"),
-				FText::FromName(Layer->Name), EditLayer->GetClass()->GetDisplayNameText());
+				FText::FromName(Layer->EditLayer->GetName()), EditLayer->GetClass()->GetDisplayNameText());
 			return false;
 		}
 	}
@@ -682,7 +709,7 @@ bool FLandscapeEditorCustomNodeBuilder_Layers::CanClearTargetLayerOnLayer(int32
 	if (!EditLayer->SupportsTargetType(ELandscapeToolTargetType::Weightmap))
 	{
 		OutReason = FText::Format(LOCTEXT("ClearTargetLayer_CantClearWeightmapOnLayerType", "Cannot clear weightmap : the type of layer {0} ({1}) doesn't support weightmaps"),
-			FText::FromName(Layer->Name), EditLayer->GetClass()->GetDisplayNameText());
+			FText::FromName(Layer->EditLayer->GetName()), EditLayer->GetClass()->GetDisplayNameText());
 		return false;
 	}
 
@@ -696,10 +723,10 @@ void FLandscapeEditorCustomNodeBuilder_Layers::ClearTargetLayerOnLayer(int32 InL
 	ALandscape* Landscape = LandscapeEdMode ? LandscapeEdMode->GetLandscape() : nullptr;
 	if (Landscape)
 	{
-		const FLandscapeLayer* Layer = LandscapeEdMode->GetLayer(InLayerIndex);
+		const FLandscapeLayer* Layer = LandscapeEdMode->GetLayerConst(InLayerIndex);
 		if (Layer)
 		{
-			EAppReturnType::Type Result = FMessageDialog::Open(EAppMsgType::YesNo, FText::Format(LOCTEXT("Landscape_ClearTargetLayer_Message", "The layer {0} : {1} content will be completely cleared.  Continue?"), FText::FromName(Layer->Name), FText::FromName(InLayerInfo->LayerName)));
+			EAppReturnType::Type Result = FMessageDialog::Open(EAppMsgType::YesNo, FText::Format(LOCTEXT("Landscape_ClearTargetLayer_Message", "The layer {0} : {1} content will be completely cleared.  Continue?"), FText::FromName(Layer->EditLayer->GetName()), FText::FromName(InLayerInfo->LayerName)));
 			if (Result == EAppReturnType::Yes)
 			{
 				const FScopedTransaction Transaction(LOCTEXT("Landscape_Layers_TargetClear", "Clear Target Layer"));
@@ -713,7 +740,7 @@ void FLandscapeEditorCustomNodeBuilder_Layers::ClearTargetLayerOnLayer(int32 InL
 bool FLandscapeEditorCustomNodeBuilder_Layers::CanClearLayer(int32 InLayerIndex, FText& OutReason) const
 {
 	FEdModeLandscape* LandscapeEdMode = GetEditorMode();
-	const FLandscapeLayer* Layer = LandscapeEdMode->GetLayer(InLayerIndex);
+	const FLandscapeLayer* Layer = LandscapeEdMode->GetLayerConst(InLayerIndex);
 	if (Layer == nullptr)
 	{
 		return false;
@@ -722,7 +749,7 @@ bool FLandscapeEditorCustomNodeBuilder_Layers::CanClearLayer(int32 InLayerIndex,
 	const ULandscapeEditLayerBase* EditLayer = Layer->EditLayer;
 	check(EditLayer != nullptr);
 
-	if (Layer->bLocked)
+	if (EditLayer->IsLocked())
 	{
 		OutReason = LOCTEXT("ClearLayer_CantClearLocked", "Cannot clear a locked edit layer");
 		return false;
@@ -731,7 +758,7 @@ bool FLandscapeEditorCustomNodeBuilder_Layers::CanClearLayer(int32 InLayerIndex,
 	if (!EditLayer->SupportsEditingTools())
 	{
 		OutReason = FText::Format(LOCTEXT("ClearLayer_CantClearLayerWithoutPersistentTextures", "Cannot clear an edit layer which doesn't have editable textures (procedural)"),
-			FText::FromName(Layer->Name), EditLayer->GetClass()->GetDisplayNameText());
+			FText::FromName(Layer->EditLayer->GetName()), EditLayer->GetClass()->GetDisplayNameText());
 		return false;
 	}
 
@@ -742,7 +769,7 @@ bool FLandscapeEditorCustomNodeBuilder_Layers::CanClearLayer(int32 InLayerIndex,
 bool FLandscapeEditorCustomNodeBuilder_Layers::CanClearTargetLayersOnLayer(int32 InLayerIndex, ELandscapeClearMode InClearMode, FText& OutReason) const
 {
 	FEdModeLandscape* LandscapeEdMode = GetEditorMode();
-	const FLandscapeLayer* Layer = LandscapeEdMode->GetLayer(InLayerIndex);
+	const FLandscapeLayer* Layer = LandscapeEdMode->GetLayerConst(InLayerIndex);
 	if (Layer == nullptr)
 	{
 		return false;
@@ -757,14 +784,14 @@ bool FLandscapeEditorCustomNodeBuilder_Layers::CanClearTargetLayersOnLayer(int32
 		&& !EditLayer->SupportsTargetType(ELandscapeToolTargetType::Visibility))
 	{
 			OutReason = FText::Format(LOCTEXT("ClearTargetLayers_CantClearWeightmapsOnLayerType", "Cannot clear weightmaps : the type of layer {0} ({1}) doesn't support weightmaps"), 
-				FText::FromName(Layer->Name), EditLayer->GetClass()->GetDisplayNameText());
+				FText::FromName(Layer->EditLayer->GetName()), EditLayer->GetClass()->GetDisplayNameText());
 			return false;
 	}
 	else if ((InClearMode == ELandscapeClearMode::Clear_Heightmap)
 		&& !EditLayer->SupportsTargetType(ELandscapeToolTargetType::Heightmap))
 	{
 		OutReason = FText::Format(LOCTEXT("ClearTargetLayers_CantClearHeightmapOnLayerType", "Cannot clear heightmap : the type of layer {0} ({1}) doesn't support heightmaps"),
-			FText::FromName(Layer->Name), EditLayer->GetClass()->GetDisplayNameText());
+			FText::FromName(Layer->EditLayer->GetName()), EditLayer->GetClass()->GetDisplayNameText());
 			return false;
 	}
 	if ((InClearMode == ELandscapeClearMode::Clear_All)
@@ -774,7 +801,7 @@ bool FLandscapeEditorCustomNodeBuilder_Layers::CanClearTargetLayersOnLayer(int32
 	if ((InClearMode == ELandscapeClearMode::Clear_Heightmap) && !EditLayer->SupportsTargetType(ELandscapeToolTargetType::Heightmap))
 	{
 		OutReason = FText::Format(LOCTEXT("ClearTargetLayers_CantClearOnLayerType", "Cannot clear : the type of layer {0} ({1}) doesn't support heightmaps or weightmaps"),
-			FText::FromName(Layer->Name), EditLayer->GetClass()->GetDisplayNameText());
+			FText::FromName(Layer->EditLayer->GetName()), EditLayer->GetClass()->GetDisplayNameText());
 		return false;
 	}
 
@@ -788,10 +815,10 @@ void FLandscapeEditorCustomNodeBuilder_Layers::ClearTargetLayersOnLayer(int32 In
 	ALandscape* Landscape = LandscapeEdMode ? LandscapeEdMode->GetLandscape() : nullptr;
 	if (Landscape)
 	{
-		const FLandscapeLayer* Layer = LandscapeEdMode->GetLayer(InLayerIndex);
+		const FLandscapeLayer* Layer = LandscapeEdMode->GetLayerConst(InLayerIndex);
 		if (Layer)
 		{
-			EAppReturnType::Type Result = FMessageDialog::Open(EAppMsgType::YesNo, FText::Format(LOCTEXT("Landscape_ClearLayer_Message", "The layer {0} content will be completely cleared.  Continue?"), FText::FromName(Layer->Name)));
+			EAppReturnType::Type Result = FMessageDialog::Open(EAppMsgType::YesNo, FText::Format(LOCTEXT("Landscape_ClearLayer_Message", "The layer {0} content will be completely cleared.  Continue?"), FText::FromName(Layer->EditLayer->GetName())));
 			if (Result == EAppReturnType::Yes)
 			{
 				const FScopedTransaction Transaction(LOCTEXT("Landscape_Layers_Clean", "Clear Layer"));
@@ -810,19 +837,22 @@ bool FLandscapeEditorCustomNodeBuilder_Layers::CanDeleteLayer(int32 InLayerIndex
 {
 	FEdModeLandscape* LandscapeEdMode = GetEditorMode();
 	ALandscape* Landscape = LandscapeEdMode ? LandscapeEdMode->GetLandscape() : nullptr;
-	if (!Landscape || Landscape->GetLayers().Num() <= 1)
+	if (!Landscape || Landscape->GetLayersConst().Num() <= 1)
 	{
 		OutReason = LOCTEXT("DeleteLayer_CantDeleteLastLayer", "The last layer cannot be deleted");
 		return false;
 	}
 
-	const FLandscapeLayer* Layer = LandscapeEdMode->GetLayer(InLayerIndex);
+	const FLandscapeLayer* Layer = LandscapeEdMode->GetLayerConst(InLayerIndex);
 	if (Layer == nullptr)
 	{
 		return false;
 	}
 
-	if (Layer->bLocked)
+	const ULandscapeEditLayerBase* EditLayer = Layer->EditLayer;
+	check(EditLayer != nullptr);
+
+	if (EditLayer->IsLocked())
 	{
 		OutReason = LOCTEXT("DeleteLayer_CantDeleteLocked", "Cannot delete a locked edit layer");
 		return false;
@@ -837,12 +867,12 @@ void FLandscapeEditorCustomNodeBuilder_Layers::DeleteLayer(int32 InLayerIndex)
 {
 	FEdModeLandscape* LandscapeEdMode = GetEditorMode();
 	ALandscape* Landscape = LandscapeEdMode ? LandscapeEdMode->GetLandscape() : nullptr;
-	if (Landscape && Landscape->GetLayers().Num() > 1)
+	if (Landscape && Landscape->GetLayersConst().Num() > 1)
 	{
-		const FLandscapeLayer* Layer = LandscapeEdMode->GetLayer(InLayerIndex);
+		const FLandscapeLayer* Layer = LandscapeEdMode->GetLayerConst(InLayerIndex);
 		if (Layer)
 		{
-			EAppReturnType::Type Result = FMessageDialog::Open(EAppMsgType::YesNo, FText::Format(LOCTEXT("Landscape_DeleteLayer_Message", "The layer {0} will be deleted.  Continue?"), FText::FromName(Layer->Name)));
+			EAppReturnType::Type Result = FMessageDialog::Open(EAppMsgType::YesNo, FText::Format(LOCTEXT("Landscape_DeleteLayer_Message", "The layer {0} will be deleted.  Continue?"), FText::FromName(Layer->EditLayer->GetName())));
 			if (Result == EAppReturnType::Yes)
 			{
 				const FScopedTransaction Transaction(LOCTEXT("Landscape_Layers_Delete", "Delete Layer"));
@@ -859,7 +889,7 @@ bool FLandscapeEditorCustomNodeBuilder_Layers::CanCollapseLayer(int32 InLayerInd
 {
 	FEdModeLandscape* LandscapeEdMode = GetEditorMode();
 	ALandscape* Landscape = LandscapeEdMode ? LandscapeEdMode->GetLandscape() : nullptr;
-	if (!Landscape || Landscape->GetLayers().Num() <= 1)
+	if (!Landscape || Landscape->GetLayersConst().Num() <= 1)
 	{
 		OutReason = LOCTEXT("Landscape_CollapseLayer_Reason_NotEnoughLayersToCollapse", "Not enough layers to do collapse");
 		return false;
@@ -871,8 +901,8 @@ bool FLandscapeEditorCustomNodeBuilder_Layers::CanCollapseLayer(int32 InLayerInd
 		return false;
 	}
 	
-	const FLandscapeLayer* TopLayer = LandscapeEdMode->GetLayer(InLayerIndex);
-	const FLandscapeLayer* BottomLayer = LandscapeEdMode->GetLayer(InLayerIndex - 1);
+	const FLandscapeLayer* TopLayer = LandscapeEdMode->GetLayerConst(InLayerIndex);
+	const FLandscapeLayer* BottomLayer = LandscapeEdMode->GetLayerConst(InLayerIndex - 1);
 	if ((TopLayer == nullptr) || (BottomLayer == nullptr))
 	{
 		return false;
@@ -885,28 +915,28 @@ bool FLandscapeEditorCustomNodeBuilder_Layers::CanCollapseLayer(int32 InLayerInd
 	if (!TopEditLayer->SupportsCollapsingTo())
 	{
 		OutReason = FText::Format(LOCTEXT("Landscape_CollapseLayer_Reason_TopLayerDoesntSupportCollapsing", "Cannot collapse layer '{0}' onto layer '{1}'. The type of layer '{0}' ({2}) doesn't support collapsing to another one"), 
-			FText::FromName(TopLayer->Name), FText::FromName(BottomLayer->Name), TopEditLayer->GetClass()->GetDisplayNameText());
+			FText::FromName(TopLayer->EditLayer->GetName()), FText::FromName(BottomLayer->EditLayer->GetName()), TopEditLayer->GetClass()->GetDisplayNameText());
 		return false;
 	}
 
 	if (!BottomEditLayer->SupportsBeingCollapsedAway())
 	{
 		OutReason = FText::Format(LOCTEXT("Landscape_CollapseLayer_Reason_BottomLayerDoesntSupportCollapsing", "Cannot collapse layer '{0}' onto layer '{1}'. The type of layer '{1}' ({2}) doesn't support being collapsed away"),
-			FText::FromName(TopLayer->Name), FText::FromName(BottomLayer->Name), BottomEditLayer->GetClass()->GetDisplayNameText());
+			FText::FromName(TopLayer->EditLayer->GetName()), FText::FromName(BottomLayer->EditLayer->GetName()), BottomEditLayer->GetClass()->GetDisplayNameText());
 		return false;
 	}
 
-	if (TopLayer->bLocked)
+	if (TopEditLayer->IsLocked())
 	{
 		OutReason = FText::Format(LOCTEXT("Landscape_CollapseLayer_Reason_TopLayerIsLocked", "Cannot collapse layer '{0}' onto layer '{1}'. Layer '{0}' will be deleted in the operation but it is currently locked"),
-			FText::FromName(TopLayer->Name), FText::FromName(BottomLayer->Name));
+			FText::FromName(TopLayer->EditLayer->GetName()), FText::FromName(BottomLayer->EditLayer->GetName()));
 		return false;
 	}
 
-	if (BottomLayer->bLocked)
+	if (BottomEditLayer->IsLocked())
 	{
 		OutReason = FText::Format(LOCTEXT("Landscape_CollapseLayer_Reason_BottomLayerIsLocked", "Cannot collapse layer '{0}' onto layer '{1}'. Destination layer '{1}' is currently locked"),
-			FText::FromName(TopLayer->Name), FText::FromName(BottomLayer->Name));
+			FText::FromName(TopLayer->EditLayer->GetName()), FText::FromName(BottomLayer->EditLayer->GetName()));
 		return false;
 	}
 
@@ -914,19 +944,19 @@ bool FLandscapeEditorCustomNodeBuilder_Layers::CanCollapseLayer(int32 InLayerInd
 	if (TopLayer->Brushes.Num() > 0)
 	{
 		OutReason = FText::Format(LOCTEXT("Landscape_CollapseLayer_Reason_TopLayerHasBrush", "Cannot collapse layer '{0}' onto layer '{1}'. Layer '{0}' contains brush(es)"),
-			FText::FromName(TopLayer->Name), FText::FromName(BottomLayer->Name));
+			FText::FromName(TopLayer->EditLayer->GetName()), FText::FromName(BottomLayer->EditLayer->GetName()));
 		return false;
 	}
 
 	if (BottomLayer->Brushes.Num() > 0)
 	{
 		OutReason = FText::Format(LOCTEXT("Landscape_CollapseLayer_Reason_BottomLayerHasBrush", "Cannot collapse layer '{0}' onto layer '{1}'. Layer '{1}' contains brush(es)"),
-			FText::FromName(TopLayer->Name), FText::FromName(BottomLayer->Name));
+			FText::FromName(TopLayer->EditLayer->GetName()), FText::FromName(BottomLayer->EditLayer->GetName()));
 		return false;
 	}
 	
 	OutReason = FText::Format(LOCTEXT("Landscape_CollapseLayer_Reason_Collapse", "Collapse layer '{0}' onto layer '{1}'"),
-		FText::FromName(TopLayer->Name), FText::FromName(BottomLayer->Name));
+		FText::FromName(TopLayer->EditLayer->GetName()), FText::FromName(BottomLayer->EditLayer->GetName()));
 	return true;
 }
 
@@ -938,11 +968,11 @@ void FLandscapeEditorCustomNodeBuilder_Layers::CollapseLayer(int32 InLayerIndex)
 	if (FEdModeLandscape* LandscapeEdMode = GetEditorMode())
 	{
 		ALandscape* Landscape = LandscapeEdMode->GetLandscape();
-		const FLandscapeLayer* Layer = LandscapeEdMode->GetLayer(InLayerIndex);
-		const FLandscapeLayer* BaseLayer = LandscapeEdMode->GetLayer(InLayerIndex - 1);
+		const FLandscapeLayer* Layer = LandscapeEdMode->GetLayerConst(InLayerIndex);
+		const FLandscapeLayer* BaseLayer = LandscapeEdMode->GetLayerConst(InLayerIndex - 1);
 		if (Landscape && Layer && BaseLayer)
 		{
-			EAppReturnType::Type Result = FMessageDialog::Open(EAppMsgType::YesNo, FText::Format(LOCTEXT("Landscape_CollapseLayer_Message", "The layer {0} will be collapsed into layer {1}.  Continue?"), FText::FromName(Layer->Name), FText::FromName(BaseLayer->Name)));
+			EAppReturnType::Type Result = FMessageDialog::Open(EAppMsgType::YesNo, FText::Format(LOCTEXT("Landscape_CollapseLayer_Message", "The layer {0} will be collapsed into layer {1}.  Continue?"), FText::FromName(Layer->EditLayer->GetName()), FText::FromName(BaseLayer->EditLayer->GetName())));
 			if (Result == EAppReturnType::Yes)
 			{
 				const FScopedTransaction Transaction(LOCTEXT("Landscape_Layers_Collapse", "Collapse Layer"));
@@ -964,7 +994,7 @@ bool FLandscapeEditorCustomNodeBuilder_Layers::CanExecuteCustomLayerAction(int32
 		return false;
 	}
 
-	const FLandscapeLayer* Layer = LandscapeEdMode->GetLayer(InLayerIndex);
+	const FLandscapeLayer* Layer = LandscapeEdMode->GetLayerConst(InLayerIndex);
 	if (Layer == nullptr)
 	{
 		return false;
@@ -981,7 +1011,7 @@ void FLandscapeEditorCustomNodeBuilder_Layers::ExecuteCustomLayerAction(int32 In
 	FEdModeLandscape* LandscapeEdMode = GetEditorMode();
 	if (ALandscape* Landscape = LandscapeEdMode ? LandscapeEdMode->GetLandscape() : nullptr)
 	{
-		const FLandscapeLayer* Layer = LandscapeEdMode->GetLayer(InLayerIndex);
+		const FLandscapeLayer* Layer = LandscapeEdMode->GetLayerConst(InLayerIndex);
 		if (Layer != nullptr)
 		{
 			const ULandscapeEditLayerBase* EditLayer = Layer->EditLayer;
@@ -1082,7 +1112,7 @@ void FLandscapeEditorCustomNodeBuilder_Layers::CreateLayer()
 			// Disallow multiple layers of certain types : 
 			if (!EditLayerClass.GetDefaultObject()->SupportsMultiple())
 			{
-				if (int32 NumLayersOfThisType = Landscape->GetLayersOfType(EditLayerClass).Num(); NumLayersOfThisType > 0)
+				if (int32 NumLayersOfThisType = Landscape->GetLayersOfTypeConst(EditLayerClass).Num(); NumLayersOfThisType > 0)
 				{
 					FMessageDialog::Open(EAppMsgType::Ok, FText::Format(LOCTEXT("Landscape_CreateLayer_CannotCreateMultiple", "Cannot create layer of type {0} : {1} {1}|plural(one=layer, other=layers) of this type already {1}|plural(one=exists, other=exist) and only 1 is allowed"), 
 						EditLayerClass->GetDisplayNameText(), NumLayersOfThisType));
@@ -1092,7 +1122,7 @@ void FLandscapeEditorCustomNodeBuilder_Layers::CreateLayer()
 
 			const FScopedTransaction Transaction(LOCTEXT("Landscape_Layers_Create", "Create Layer"));
 			Landscape->CreateLayer(FName(EditLayerClass.GetDefaultObject()->GetDefaultName()), EditLayerClass);
-			OnLayerSelectionChanged(Landscape->GetLayerCount() - 1);
+			OnLayerSelectionChanged(Landscape->GetLayersConst().Num() - 1);
 
 			LandscapeEdMode->RefreshDetailPanel();
 		}
@@ -1105,7 +1135,7 @@ FText FLandscapeEditorCustomNodeBuilder_Layers::GetNumLayersText() const
 	ALandscape* Landscape = LandscapeEdMode ? LandscapeEdMode->GetLandscape() : nullptr;
 	if (Landscape)
 	{
-		return FText::Format(LOCTEXT("NumEditLayersText", "{0} Edit {0}|plural(one=Layer, other=Layers)"), Landscape->GetLayers().Num());
+		return FText::Format(LOCTEXT("NumEditLayersText", "{0} Edit {0}|plural(one=Layer, other=Layers)"), Landscape->GetLayersConst().Num());
 	}
 
 	return FText();
@@ -1130,8 +1160,7 @@ bool FLandscapeEditorCustomNodeBuilder_Layers::CanCreateLayer(FText& OutReason)
 
 void FLandscapeEditorCustomNodeBuilder_Layers::OnLayerSelectionChanged(int32 InLayerIndex)
 {
-	FEdModeLandscape* LandscapeEdMode = GetEditorMode();
-	if (LandscapeEdMode && LandscapeEdMode->GetCurrentLayerIndex() != InLayerIndex)
+	if (FEdModeLandscape* LandscapeEdMode = GetEditorMode(); LandscapeEdMode && LandscapeEdMode->GetCurrentLayerIndex() != InLayerIndex)
 	{
 		FScopedTransaction Transaction(LOCTEXT("Landscape_Layers_SetCurrentLayer", "Set Current Layer"));
 		LandscapeEdMode->SetCurrentLayer(InLayerIndex);
@@ -1141,26 +1170,58 @@ void FLandscapeEditorCustomNodeBuilder_Layers::OnLayerSelectionChanged(int32 InL
 
 TOptional<float> FLandscapeEditorCustomNodeBuilder_Layers::GetLayerAlpha(int32 InLayerIndex) const
 {
-	FEdModeLandscape* LandscapeEdMode = GetEditorMode();
-
-	if (LandscapeEdMode)
+	if (FEdModeLandscape* LandscapeEdMode = GetEditorMode())
 	{
-		return LandscapeEdMode->GetLayerAlpha(InLayerIndex);
+		const FLandscapeLayer* Layer = LandscapeEdMode->GetLayerConst(InLayerIndex);
+		check(Layer != nullptr);
+		const ULandscapeEditLayerBase* EditLayer = Layer->EditLayer;
+		check(EditLayer != nullptr);
+
+		return EditLayer->GetAlphaForTargetType(LandscapeEdMode->GetLandscapeToolTargetType());
 	}
 
 	return 1.0f;
 }
 
+float FLandscapeEditorCustomNodeBuilder_Layers::GetLayerAlphaMinValue() const
+{
+	if (const FEdModeLandscape* LandscapeEdMode = GetEditorMode(); LandscapeEdMode && LandscapeEdMode->GetLandscapeToolTargetType() == ELandscapeToolTargetType::Heightmap)
+	{
+		return -1.0f;
+	}
+	else
+	{
+		return 0.0f;
+	}
+}
+
 bool FLandscapeEditorCustomNodeBuilder_Layers::CanSetLayerAlpha(int32 InLayerIndex, FText& OutReason) const
 {
-	FEdModeLandscape* LandscapeEdMode = GetEditorMode();
-	const FLandscapeLayer* Layer = LandscapeEdMode->GetLayer(InLayerIndex);
+	const FEdModeLandscape* LandscapeEdMode = GetEditorMode();
+
+	if (LandscapeEdMode == nullptr)
+	{
+		return false;
+	}
+
+	const FLandscapeLayer* Layer = LandscapeEdMode->GetLayerConst(InLayerIndex);
 	if (Layer == nullptr)
 	{
 		return false;
 	}
 
-	if (Layer->bLocked)
+	const ULandscapeEditLayerBase* EditLayer = Layer->EditLayer;
+	check(EditLayer != nullptr);
+
+	if (!EditLayer->SupportsAlphaForTargetType(LandscapeEdMode->GetLandscapeToolTargetType()))
+	{
+		OutReason = FText::Format(LOCTEXT("SetLayerAlpha_LayerDoesntSupportAlpha", "Cannot change alpha : the type of layer {0} ({1}) doesn't support alpha"),
+			FText::FromName(Layer->EditLayer->GetName()), EditLayer->GetClass()->GetDisplayNameText());
+		return false;
+	}
+
+
+	if (EditLayer->IsLocked())
 	{
 		OutReason = LOCTEXT("SetLayerAlpha_LayerIsLocked", "Cannot change the alpha of a locked edit layer");
 		return false;
@@ -1175,28 +1236,28 @@ void FLandscapeEditorCustomNodeBuilder_Layers::SetLayerAlpha(float InAlpha, int3
 	FEdModeLandscape* LandscapeEdMode = GetEditorMode();
 	if (LandscapeEdMode)
 	{
-		// We get multiple commits when editing through the text box
-		if (LandscapeEdMode->GetLayerAlpha(InLayerIndex) == LandscapeEdMode->GetClampedLayerAlpha(InAlpha))
-		{
-			return;
-		}
+		ULandscapeEditLayerBase* EditLayer = LandscapeEdMode->GetEditLayer(InLayerIndex);
+		check(EditLayer != nullptr);
 
 		FScopedTransaction Transaction(LOCTEXT("Landscape_Layers_SetAlpha", "Set Layer Alpha"), CurrentSlider == INDEX_NONE && bCommit);
 		// Set Value when using slider or when committing text
-		LandscapeEdMode->SetLayerAlpha(InLayerIndex, InAlpha);
+		EditLayer->SetAlphaForTargetType(LandscapeEdMode->GetLandscapeToolTargetType(), InAlpha, /*bInModify = */true, bCommit ? EPropertyChangeType::ValueSet : EPropertyChangeType::Interactive);
 	}
 }
 
 bool FLandscapeEditorCustomNodeBuilder_Layers::CanToggleVisibility(int32 InLayerIndex, FText& OutReason) const
 {
 	FEdModeLandscape* LandscapeEdMode = GetEditorMode();
-	const FLandscapeLayer* Layer = LandscapeEdMode->GetLayer(InLayerIndex);
+	const FLandscapeLayer* Layer = LandscapeEdMode->GetLayerConst(InLayerIndex);
 	if (Layer == nullptr)
 	{
 		return false;
 	}
 
-	if (Layer->bLocked)
+	const ULandscapeEditLayerBase* EditLayer = Layer->EditLayer;
+	check(EditLayer != nullptr);
+
+	if (EditLayer->IsLocked())
 	{
 		OutReason = LOCTEXT("ToggleVisibility_CantToggleLocked", "Cannot change the visibility of a locked edit layer");
 		return false;
@@ -1208,12 +1269,15 @@ bool FLandscapeEditorCustomNodeBuilder_Layers::CanToggleVisibility(int32 InLayer
 
 FReply FLandscapeEditorCustomNodeBuilder_Layers::OnToggleVisibility(int32 InLayerIndex)
 {
-	FEdModeLandscape* LandscapeEdMode = GetEditorMode();
-	if (LandscapeEdMode)
+	if (FEdModeLandscape* LandscapeEdMode = GetEditorMode())
 	{
 		const FScopedTransaction Transaction(LOCTEXT("Landscape_Layers_SetVisibility", "Set Layer Visibility"));
-		LandscapeEdMode->SetLayerVisibility(!LandscapeEdMode->IsLayerVisible(InLayerIndex), InLayerIndex);
-		if (LandscapeEdMode->IsLayerVisible(InLayerIndex))
+		
+		ULandscapeEditLayerBase* EditLayer = LandscapeEdMode->GetEditLayer(InLayerIndex);
+		check(EditLayer != nullptr);
+		EditLayer->SetVisible(!EditLayer->IsVisible(), /*bInModify= */true);
+
+		if (EditLayer->IsVisible())
 		{
 			OnLayerSelectionChanged(InLayerIndex);
 		}
@@ -1223,18 +1287,40 @@ FReply FLandscapeEditorCustomNodeBuilder_Layers::OnToggleVisibility(int32 InLaye
 
 const FSlateBrush* FLandscapeEditorCustomNodeBuilder_Layers::GetVisibilityBrushForLayer(int32 InLayerIndex) const
 {
-	FEdModeLandscape* LandscapeEdMode = GetEditorMode();
-	bool bIsVisible = LandscapeEdMode && LandscapeEdMode->IsLayerVisible(InLayerIndex);
+	bool bIsVisible = false;
+
+	if (FEdModeLandscape* LandscapeEdMode = GetEditorMode())
+	{
+		const ULandscapeEditLayerBase* EditLayer = LandscapeEdMode->GetEditLayerConst(InLayerIndex);
+		check(EditLayer != nullptr);
+
+		bIsVisible = EditLayer->IsVisible();
+	}
+	
 	return bIsVisible ? FAppStyle::GetBrush("Level.VisibleIcon16x") : FAppStyle::GetBrush("Level.NotVisibleIcon16x");
 }
 
+void FLandscapeEditorCustomNodeBuilder_Layers::OnSetInspectedDetailsToEditLayer(int32 InLayerIndex) const
+{
+	if (FEdModeLandscape* LandscapeEdMode = GetEditorMode())
+	{
+		const ULandscapeEditLayerBase* EditLayer = LandscapeEdMode->GetEditLayerConst(InLayerIndex);
+
+		// Clear out all previously selected objects, this may change in the future
+		TArray<TWeakObjectPtr<UObject>> InspectedObjects;
+		InspectedObjects.Add(const_cast<ULandscapeEditLayerBase*>(EditLayer));
+		LandscapeEdMode->SetInspectedObjects(InspectedObjects);
+	}
+}
+
 FReply FLandscapeEditorCustomNodeBuilder_Layers::OnToggleLock(int32 InLayerIndex)
 {
-	FEdModeLandscape* LandscapeEdMode = GetEditorMode();
-	if (LandscapeEdMode)
+	if (FEdModeLandscape* LandscapeEdMode = GetEditorMode())
 	{
+		ULandscapeEditLayerBase* EditLayer =  LandscapeEdMode->GetEditLayer(InLayerIndex);
+		check(EditLayer != nullptr);
 		const FScopedTransaction Transaction(LOCTEXT("Landscape_Layers_Locked", "Set Layer Locked"));
-		LandscapeEdMode->SetLayerLocked(InLayerIndex, !LandscapeEdMode->IsLayerLocked(InLayerIndex));
+		EditLayer->SetLocked(!EditLayer->IsLocked(), /*bInModify = */true);
 	}
 	return FReply::Handled();
 }
@@ -1249,21 +1335,26 @@ EVisibility FLandscapeEditorCustomNodeBuilder_Layers::GetLayerAlphaVisibility(in
 TSharedPtr<IToolTip> FLandscapeEditorCustomNodeBuilder_Layers::GetEditLayerTypeTooltip(int32 InLayerIndex) const
 {
 	FEdModeLandscape* LandscapeEdMode = GetEditorMode();
-	const ULandscapeEditLayerBase* EditLayer = (LandscapeEdMode != nullptr) ? LandscapeEdMode->GetEditLayer(InLayerIndex) : nullptr;
+	const ULandscapeEditLayerBase* EditLayer = (LandscapeEdMode != nullptr) ? LandscapeEdMode->GetEditLayerConst(InLayerIndex) : nullptr;
 	return (EditLayer != nullptr) ? TSharedPtr<IToolTip>(FEditorClassUtils::GetTooltip(EditLayer->GetClass())) : nullptr;
 }
 
 const FSlateBrush* FLandscapeEditorCustomNodeBuilder_Layers::GetLockBrushForLayer(int32 InLayerIndex) const
 {
-	FEdModeLandscape* LandscapeEdMode = GetEditorMode();
-	bool bIsLocked = LandscapeEdMode && LandscapeEdMode->IsLayerLocked(InLayerIndex);
+	bool bIsLocked = false;
+	if (FEdModeLandscape* LandscapeEdMode = GetEditorMode())
+	{
+		const ULandscapeEditLayerBase* EditLayer = LandscapeEdMode->GetEditLayerConst(InLayerIndex);
+		check(EditLayer != nullptr);
+		bIsLocked = EditLayer->IsLocked();
+	}
 	return bIsLocked ? FAppStyle::GetBrush(TEXT("PropertyWindow.Locked")) : FAppStyle::GetBrush(TEXT("PropertyWindow.Unlocked"));
 }
 
 const FSlateBrush* FLandscapeEditorCustomNodeBuilder_Layers::GetEditLayerIconBrush(int32 InLayerIndex) const
 {
 	FEdModeLandscape* LandscapeEdMode = GetEditorMode();
-	const ULandscapeEditLayerBase* EditLayer = (LandscapeEdMode != nullptr) ? LandscapeEdMode->GetEditLayer(InLayerIndex) : nullptr;
+	const ULandscapeEditLayerBase* EditLayer = (LandscapeEdMode != nullptr) ? LandscapeEdMode->GetEditLayerConst(InLayerIndex) : nullptr;
 	return (EditLayer != nullptr) ? FSlateIconFinder::FindIconBrushForClass(EditLayer->GetClass()) : nullptr;
 }
 
@@ -1276,8 +1367,8 @@ int32 FLandscapeEditorCustomNodeBuilder_Layers::SlotIndexToLayerIndex(int32 Slot
 		return INDEX_NONE;
 	}
 	
-	check(Landscape->GetLayers().IsValidIndex(SlotIndex));
-	return Landscape->GetLayers().Num() - SlotIndex - 1;
+	check(Landscape->GetLayersConst().IsValidIndex(SlotIndex));
+	return Landscape->GetLayersConst().Num() - SlotIndex - 1;
 }
 
 FReply FLandscapeEditorCustomNodeBuilder_Layers::HandleDragDetected(const FGeometry& MyGeometry, const FPointerEvent& MouseEvent, int32 SlotIndex, SVerticalBox::FSlot* Slot)
@@ -1286,13 +1377,17 @@ FReply FLandscapeEditorCustomNodeBuilder_Layers::HandleDragDetected(const FGeome
 	if (LandscapeEdMode && LandscapeEdMode->DoesCurrentToolAffectEditLayers())
 	{
 		int32 LayerIndex = SlotIndexToLayerIndex(SlotIndex);
-		const FLandscapeLayer* Layer = LandscapeEdMode->GetLayer(LayerIndex);
-		if (Layer && !Layer->bLocked)
+		if (const FLandscapeLayer* Layer = LandscapeEdMode->GetLayerConst(LayerIndex))
 		{
-			TSharedPtr<SWidget> Row = GenerateRow(LayerIndex);
-			if (Row.IsValid())
+			const ULandscapeEditLayerBase* EditLayer = Layer->EditLayer;
+			check(EditLayer != nullptr);
+			if (!EditLayer->IsLocked())
 			{
-				return FReply::Handled().BeginDragDrop(FLandscapeListElementDragDropOp::New(SlotIndex, Slot, Row));
+				TSharedPtr<SWidget> Row = GenerateRow(LayerIndex);
+				if (Row.IsValid())
+				{
+					return FReply::Handled().BeginDragDrop(FLandscapeListElementDragDropOp::New(SlotIndex, Slot, Row));
+				}
 			}
 		}
 	}
@@ -1302,11 +1397,26 @@ FReply FLandscapeEditorCustomNodeBuilder_Layers::HandleDragDetected(const FGeome
 TOptional<SDragAndDropVerticalBox::EItemDropZone> FLandscapeEditorCustomNodeBuilder_Layers::HandleCanAcceptDrop(const FDragDropEvent& DragDropEvent, SDragAndDropVerticalBox::EItemDropZone DropZone, int32 SlotIndex, SVerticalBox::FSlot* Slot)
 {
 	TSharedPtr<FLandscapeListElementDragDropOp> DragDropOperation = DragDropEvent.GetOperationAs<FLandscapeListElementDragDropOp>();
-	if (DragDropOperation.IsValid())
+	const FEdModeLandscape* LandscapeEdMode = GetEditorMode();
+
+	if (DragDropOperation->IsOfType<FLandscapeBrushDragDropOp>() && DragDropOperation.IsValid() && LandscapeEdMode)
 	{
-		return DropZone;
+		const int32 DestinationLayerIndex = SlotIndexToLayerIndex(SlotIndex);
+		
+		if (const ULandscapeEditLayerBase* EditLayer = LandscapeEdMode->GetEditLayerConst(DestinationLayerIndex))
+		{
+			if (EditLayer->SupportsBlueprintBrushes())
+			{
+				return DropZone;
+			}
+			else 
+			{
+				return TOptional<SDragAndDropVerticalBox::EItemDropZone>();
+			}
+		}
 	}
-	return TOptional<SDragAndDropVerticalBox::EItemDropZone>();
+
+	return DropZone;
 }
 
 FReply FLandscapeEditorCustomNodeBuilder_Layers::HandleAcceptDrop(FDragDropEvent const& DragDropEvent, SDragAndDropVerticalBox::EItemDropZone DropZone, int32 SlotIndex, SVerticalBox::FSlot* Slot)
@@ -1318,8 +1428,8 @@ FReply FLandscapeEditorCustomNodeBuilder_Layers::HandleAcceptDrop(FDragDropEvent
 		return FReply::Unhandled();
 	}
 
-		FEdModeLandscape* LandscapeEdMode = GetEditorMode();
-		ALandscape* Landscape = LandscapeEdMode ? LandscapeEdMode->GetLandscape() : nullptr;
+	FEdModeLandscape* LandscapeEdMode = GetEditorMode();
+	ALandscape* Landscape = LandscapeEdMode ? LandscapeEdMode->GetLandscape() : nullptr;
 	if (!Landscape)
 	{
 		return FReply::Unhandled();
@@ -1350,7 +1460,6 @@ FReply FLandscapeEditorCustomNodeBuilder_Layers::HandleAcceptDrop(FDragDropEvent
 		Landscape->AddBrushToLayer(DestinationLayerIndex, Brush);
 
 		LandscapeEdMode->SetCurrentLayer(DestinationLayerIndex);
-		LandscapeEdMode->RefreshDetailPanel();
 
 		// HACK: We don't return FReply::Handled() here because otherwise, SDragAndDropVerticalBox::OnDrop
 		// will apply UI slot reordering after we return. Properly speaking, we should have a way to signal 
@@ -1366,7 +1475,6 @@ FReply FLandscapeEditorCustomNodeBuilder_Layers::HandleAcceptDrop(FDragDropEvent
 	if (Landscape->ReorderLayer(StartingLayerIndex, DestinationLayerIndex))
 	{
 		LandscapeEdMode->SetCurrentLayer(DestinationLayerIndex);
-		LandscapeEdMode->RefreshDetailPanel();
 		return FReply::Handled();
 	}
 
diff --git a/Engine/Source/Editor/LandscapeEditor/Private/LandscapeEditorDetailCustomization_Layers.h b/Engine/Source/Editor/LandscapeEditor/Private/LandscapeEditorDetailCustomization_Layers.h
--- a/Engine/Source/Editor/LandscapeEditor/Private/LandscapeEditorDetailCustomization_Layers.h
+++ b/Engine/Source/Editor/LandscapeEditor/Private/LandscapeEditorDetailCustomization_Layers.h
@@ -115,6 +115,7 @@ protected:
 	TSubclassOf<ULandscapeEditLayerBase> PickEditLayerClass() const;
 
 	TOptional<float> GetLayerAlpha(int32 InLayerIndex) const;
+	float GetLayerAlphaMinValue() const;
 	bool CanSetLayerAlpha(int32 InLayerIndex, FText& OutReason) const;
 	void SetLayerAlpha(float InAlpha, int32 InLayerIndex, bool bCommit);
 	
@@ -122,6 +123,8 @@ protected:
 	FReply OnToggleVisibility(int32 InLayerIndex);
 	const FSlateBrush* GetVisibilityBrushForLayer(int32 InLayerIndex) const;
 	
+	void OnSetInspectedDetailsToEditLayer(int32 InLayerIndex) const;
+
 	FReply OnToggleLock(int32 InLayerIndex);
 	const FSlateBrush* GetLockBrushForLayer(int32 InLayerIndex) const;
 
diff --git a/Engine/Source/Editor/LandscapeEditor/Private/LandscapeEditorDetailCustomization_LayersBrushStack.cpp b/Engine/Source/Editor/LandscapeEditor/Private/LandscapeEditorDetailCustomization_LayersBrushStack.cpp
--- a/Engine/Source/Editor/LandscapeEditor/Private/LandscapeEditorDetailCustomization_LayersBrushStack.cpp
+++ b/Engine/Source/Editor/LandscapeEditor/Private/LandscapeEditorDetailCustomization_LayersBrushStack.cpp
@@ -503,11 +503,13 @@ FReply FLandscapeEditorCustomNodeBuilder_LayersBrushStack::HandleDragDetected(co
 {
 	if (FEdModeLandscape* LandscapeEdMode = GetEditorMode())
 	{
-		const FLandscapeLayer* Layer = LandscapeEdMode->GetCurrentLayer();
-		if (Layer && !Layer->bLocked)
+		if (const FLandscapeLayer* Layer = LandscapeEdMode->GetCurrentLayerConst())
 		{
+			const ULandscapeEditLayerBase* EditLayer = Layer->EditLayer;
+			check(EditLayer != nullptr);
+
 			const TArray<ALandscapeBlueprintBrushBase*>& BrushStack = LandscapeEdMode->GetBrushesForCurrentLayer();
-			if (BrushStack.IsValidIndex(SlotIndex))
+			if (!EditLayer->IsLocked() && BrushStack.IsValidIndex(SlotIndex))
 			{
 				TSharedPtr<SWidget> Row = GenerateRow(SlotIndex);
 				if (Row.IsValid())
diff --git a/Engine/Source/Editor/LandscapeEditor/Private/LandscapeEditorDetailCustomization_NewLandscape.cpp b/Engine/Source/Editor/LandscapeEditor/Private/LandscapeEditorDetailCustomization_NewLandscape.cpp
--- a/Engine/Source/Editor/LandscapeEditor/Private/LandscapeEditorDetailCustomization_NewLandscape.cpp
+++ b/Engine/Source/Editor/LandscapeEditor/Private/LandscapeEditorDetailCustomization_NewLandscape.cpp
@@ -1111,11 +1111,11 @@ void FLandscapeEditorDetailCustomization_NewLandscape::AddComponents(ULandscapeI
 			TArray<ULandscapeComponent*> ComponentsUsingHeightmap;
 			ComponentsUsingHeightmap.Add(NewComponent);
 
-			for (const FLandscapeLayer& Layer : Landscape->GetLayers())
+			for (const FLandscapeLayer& Layer : Landscape->GetLayersConst())
 			{
 				// Since we do not share heightmap when adding new component, we will provided the required array, but they will only be used for 1 component
 				TMap<UTexture2D*, UTexture2D*> CreatedHeightmapTextures;
-				NewComponent->AddDefaultLayerData(Layer.Guid, ComponentsUsingHeightmap, CreatedHeightmapTextures);
+				NewComponent->AddDefaultLayerData(Layer.EditLayer->GetGuid(), ComponentsUsingHeightmap, CreatedHeightmapTextures);
 			}
 		}
 
diff --git a/Engine/Source/Editor/LandscapeEditor/Private/LandscapeEditorDetailCustomization_TargetLayers.cpp b/Engine/Source/Editor/LandscapeEditor/Private/LandscapeEditorDetailCustomization_TargetLayers.cpp
--- a/Engine/Source/Editor/LandscapeEditor/Private/LandscapeEditorDetailCustomization_TargetLayers.cpp
+++ b/Engine/Source/Editor/LandscapeEditor/Private/LandscapeEditorDetailCustomization_TargetLayers.cpp
@@ -1064,6 +1064,7 @@ TSharedPtr<SWidget> FLandscapeEditorCustomNodeBuilder_TargetLayers::GenerateRow(
 			.OnContextMenuOpening(this, &FLandscapeEditorCustomNodeBuilder_TargetLayers::OnTargetLayerContextMenuOpening, Target, static_cast<int32>(INDEX_NONE))
 			.OnSelected_Static(&FLandscapeEditorCustomNodeBuilder_TargetLayers::OnTargetSelectionChanged, Target)
 			.IsSelected_Static(&FLandscapeEditorCustomNodeBuilder_TargetLayers::GetTargetLayerIsSelected, Target)
+			.OnDoubleClick_Static(&FLandscapeEditorCustomNodeBuilder_TargetLayers::OnTargetLayerDoubleClicked, Target)
 			.Visibility(this, &FLandscapeEditorCustomNodeBuilder_TargetLayers::ShouldShowLayer, Target)
 			[
 				SNew(SHorizontalBox)
@@ -1484,6 +1485,18 @@ void FLandscapeEditorCustomNodeBuilder_TargetLayers::OnTargetSelectionChanged(co
 	}
 }
 
+void FLandscapeEditorCustomNodeBuilder_TargetLayers::OnTargetLayerDoubleClicked(const TSharedRef<FLandscapeTargetListInfo> Target)
+{
+	if (FEdModeLandscape* LandscapeEdMode = GetEditorMode())
+	{
+		LandscapeEdMode->CurrentToolTarget.TargetType = Target->TargetType;
+		if ((Target->TargetType == ELandscapeToolTargetType::Weightmap) && Target->LayerInfoObj.IsValid())
+		{
+			GEditor->GetEditorSubsystem<UAssetEditorSubsystem>()->OpenEditorForAsset(Target->LayerInfoObj.Get());
+		}
+	}
+}
+
 TSharedPtr<SWidget> FLandscapeEditorCustomNodeBuilder_TargetLayers::OnTargetLayerContextMenuOpening(const TSharedRef<FLandscapeTargetListInfo> Target, const int32 InLayerIndex)
 {
 	FMenuBuilder MenuBuilder(true, nullptr);
@@ -2004,10 +2017,15 @@ EVisibility FLandscapeEditorCustomNodeBuilder_TargetLayers::GetLayersSubstractiv
 
 ECheckBoxState FLandscapeEditorCustomNodeBuilder_TargetLayers::IsLayersSubstractiveBlendChecked(const TSharedRef<FLandscapeTargetListInfo> Target)
 {
-	FEdModeLandscape* LandscapeEdMode = GetEditorMode();
-	if (LandscapeEdMode)
+	if (const FEdModeLandscape* LandscapeEdMode = GetEditorMode())
 	{
-		return LandscapeEdMode->IsCurrentLayerBlendSubstractive(Target.Get().LayerInfoObj) ? ECheckBoxState::Checked : ECheckBoxState::Unchecked;
+		ULandscapeEditLayerBase* EditLayer = LandscapeEdMode->GetEditLayer(LandscapeEdMode->UISettings->CurrentLayerIndex);
+		const bool* AllocationBlend = EditLayer->GetWeightmapLayerAllocationBlend().Find(Target.Get().LayerInfoObj.Get());
+
+		if (AllocationBlend != nullptr)
+		{
+			return (*AllocationBlend) ? ECheckBoxState::Checked : ECheckBoxState::Unchecked;
+		}
 	}
 
 	return ECheckBoxState::Unchecked;
@@ -2015,11 +2033,12 @@ ECheckBoxState FLandscapeEditorCustomNodeBuilder_TargetLayers::IsLayersSubstract
 
 void FLandscapeEditorCustomNodeBuilder_TargetLayers::OnLayersSubstractiveBlendChanged(ECheckBoxState NewCheckedState, const TSharedRef<FLandscapeTargetListInfo> Target)
 {
-	FEdModeLandscape* LandscapeEdMode = GetEditorMode();
-	if (LandscapeEdMode)
+	if (const FEdModeLandscape* LandscapeEdMode = GetEditorMode())
 	{
 		FScopedTransaction Transaction(LOCTEXT("Undo_SubtractiveBlend", "Set Subtractive Blend Layer"));
-		LandscapeEdMode->SetCurrentLayerSubstractiveBlendStatus(NewCheckedState == ECheckBoxState::Checked, Target.Get().LayerInfoObj);
+		ULandscapeEditLayerBase* EditLayer = LandscapeEdMode->GetEditLayer(LandscapeEdMode->UISettings->CurrentLayerIndex);
+
+		EditLayer->AddOrUpdateWeightmapAllocationLayerBlend(Target.Get().LayerInfoObj.Get(), NewCheckedState == ECheckBoxState::Checked, /*bInModify =*/true);
 	}
 }
 
@@ -2132,6 +2151,22 @@ FReply SLandscapeEditorSelectableBorder::OnMouseButtonUp(const FGeometry& MyGeom
 	return FReply::Unhandled();
 }
 
+FReply SLandscapeEditorSelectableBorder::OnMouseDoubleClick(const FGeometry& MyGeometry, const FPointerEvent& MouseEvent)
+{
+	if (MyGeometry.IsUnderLocation(MouseEvent.GetScreenSpacePosition()))
+	{
+		if (MouseEvent.GetEffectingButton() == EKeys::LeftMouseButton &&
+			OnDoubleClick.IsBound())
+		{
+			OnDoubleClick.Execute();
+
+			return FReply::Handled().ReleaseMouseCapture();
+		}
+	}
+
+	return FReply::Unhandled();
+}
+
 const FSlateBrush* SLandscapeEditorSelectableBorder::GetBorder() const
 {
 	const bool bIsSelected = IsSelected.Get();
diff --git a/Engine/Source/Editor/LandscapeEditor/Private/LandscapeEditorDetailCustomization_TargetLayers.h b/Engine/Source/Editor/LandscapeEditor/Private/LandscapeEditorDetailCustomization_TargetLayers.h
--- a/Engine/Source/Editor/LandscapeEditor/Private/LandscapeEditorDetailCustomization_TargetLayers.h
+++ b/Engine/Source/Editor/LandscapeEditor/Private/LandscapeEditorDetailCustomization_TargetLayers.h
@@ -91,6 +91,7 @@ protected:
 
 	static bool GetTargetLayerIsSelected(const TSharedRef<FLandscapeTargetListInfo> Target);
 	static void OnTargetSelectionChanged(const TSharedRef<FLandscapeTargetListInfo> Target);
+	static void OnTargetLayerDoubleClicked(const TSharedRef<FLandscapeTargetListInfo> Target);
 	TSharedPtr<SWidget> OnTargetLayerContextMenuOpening(const TSharedRef<FLandscapeTargetListInfo> Target, const int32 InLayerIndex);
 	void OnRenameLayer(const int32 InLayerIndex);
 	static void OnExportLayer(const TSharedRef<FLandscapeTargetListInfo> Target);
@@ -172,18 +173,21 @@ public:
 
 		SLATE_EVENT(FOnContextMenuOpening, OnContextMenuOpening)
 		SLATE_EVENT(FSimpleDelegate, OnSelected)
+		SLATE_EVENT(FSimpleDelegate, OnDoubleClick)
 		SLATE_ATTRIBUTE(bool, IsSelected)
 	SLATE_END_ARGS()
 
 	void Construct(const FArguments& InArgs);
 
 	virtual FReply OnMouseButtonUp(const FGeometry& MyGeometry, const FPointerEvent& MouseEvent);
+	virtual FReply OnMouseDoubleClick(const FGeometry& MyGeometry, const FPointerEvent& MouseEvent);
 	const FSlateBrush* GetBorder() const;
 
 protected:
 	FOnContextMenuOpening OnContextMenuOpening;
 	FSimpleDelegate OnSelected;
 	TAttribute<bool> IsSelected;
+	FSimpleDelegate OnDoubleClick;
 };
 
 class FTargetLayerDragDropOp : public FDragAndDropVerticalBoxOp
diff --git a/Engine/Source/Editor/LandscapeEditor/Private/LandscapeEditorModule.cpp b/Engine/Source/Editor/LandscapeEditor/Private/LandscapeEditorModule.cpp
--- a/Engine/Source/Editor/LandscapeEditor/Private/LandscapeEditorModule.cpp
+++ b/Engine/Source/Editor/LandscapeEditor/Private/LandscapeEditorModule.cpp
@@ -10,6 +10,10 @@
 #include "EditorModeRegistry.h"
 #include "EditorModeManager.h"
 #include "EditorModes.h"
+#include "EditorWidgetsModule.h"
+#include "ObjectNameEditSinkRegistry.h"
+#include "LandscapeEditLayerObjectNameEditSink.h"
+
 #include "LandscapeFileFormatInterface.h"
 #include "LandscapeProxy.h"
 #include "LandscapeEdMode.h"
@@ -33,6 +37,9 @@
 #include "LandscapeEditorDetailCustomization_ImportLayers.h"
 #include "LandscapeEditorDetailCustomization_TargetLayers.h"
 #include "LandscapeSplineDetails.h"
+#include "LandscapeGrassTypeDetails.h"
+#include "LandscapeProxyUIDetails.h"
+#include "LandscapeUIDetails.h"
 #include "LandscapeModule.h"
 
 #include "LevelEditor.h"
@@ -84,6 +91,10 @@ public:
 			300
 			);
 
+		// register custom editor widgets 
+		FEditorWidgetsModule& EditorWidgetsModule = FModuleManager::LoadModuleChecked<FEditorWidgetsModule>("EditorWidgets");
+		EditorWidgetsModule.GetObjectNameEditSinkRegistry()->RegisterObjectNameEditSink(MakeShared<FLandscapeEditLayerObjectNameEditSink>());
+
 		// register customizations
 		FPropertyEditorModule& PropertyModule = FModuleManager::GetModuleChecked<FPropertyEditorModule>("PropertyEditor");
 		PropertyModule.RegisterCustomClassLayout("LandscapeEditorObject", FOnGetDetailCustomizationInstance::CreateStatic(&FLandscapeEditorDetails::MakeInstance));
@@ -94,6 +105,10 @@ public:
 		PropertyModule.RegisterCustomClassLayout("LandscapeSplineControlPoint", FOnGetDetailCustomizationInstance::CreateStatic(&FLandscapeSplineDetails::MakeInstance));
 		PropertyModule.RegisterCustomClassLayout("LandscapeSplineSegment", FOnGetDetailCustomizationInstance::CreateStatic(&FLandscapeSplineDetails::MakeInstance));
 
+		PropertyModule.RegisterCustomClassLayout("Landscape", FOnGetDetailCustomizationInstance::CreateStatic(&FLandscapeUIDetails::MakeInstance));
+		PropertyModule.RegisterCustomClassLayout("LandscapeProxy", FOnGetDetailCustomizationInstance::CreateStatic(&FLandscapeProxyUIDetails::MakeInstance));
+		PropertyModule.RegisterCustomClassLayout("LandscapeGrassType", FOnGetDetailCustomizationInstance::CreateStatic(&FLandscapeGrassTypeDetails::MakeInstance));
+
 		// register property sections
 		RegisterPropertySectionMappings();
 		
@@ -170,6 +185,10 @@ public:
 		// unregister the editor mode
 		FEditorModeRegistry::Get().UnregisterMode(FBuiltinEditorModes::EM_Landscape);
 
+		// unregister custom editor widgets 
+		FEditorWidgetsModule& EditorWidgetsModule = FModuleManager::LoadModuleChecked<FEditorWidgetsModule>("EditorWidgets");
+		EditorWidgetsModule.GetObjectNameEditSinkRegistry()->UnregisterObjectNameEditSink(MakeShared<FLandscapeEditLayerObjectNameEditSink>());
+
 		// unregister customizations
 		FPropertyEditorModule& PropertyModule = FModuleManager::GetModuleChecked<FPropertyEditorModule>("PropertyEditor");
 		PropertyModule.UnregisterCustomClassLayout("LandscapeEditorObject");
@@ -180,6 +199,10 @@ public:
 		PropertyModule.UnregisterCustomClassLayout("LandscapeSplineControlPoint");
 		PropertyModule.UnregisterCustomClassLayout("LandscapeSplineSegment");
 
+		PropertyModule.UnregisterCustomClassLayout("Landscape");
+		PropertyModule.UnregisterCustomClassLayout("LandscapeProxy");
+		PropertyModule.UnregisterCustomClassLayout("LandscapeGrassType");
+
 		// unregister property sections
 		UnregisterPropertySectionMappings();
 
diff --git a/Engine/Source/Editor/LandscapeEditor/Private/LandscapeGrassTypeDetails.cpp b/Engine/Source/Editor/LandscapeEditor/Private/LandscapeGrassTypeDetails.cpp
new file mode 100644
--- /dev/null
+++ b/Engine/Source/Editor/LandscapeEditor/Private/LandscapeGrassTypeDetails.cpp
@@ -0,0 +1,102 @@
+// Copyright Epic Games, Inc. All Rights Reserved.
+
+#include "LandscapeGrassTypeDetails.h"
+
+#include "Containers/Array.h"
+#include "Delegates/Delegate.h"
+#include "DetailLayoutBuilder.h"
+#include "Engine/Engine.h"
+#include "Engine/EngineTypes.h"
+#include "Misc/AssertionMacros.h"
+#include "PropertyHandle.h"
+#include "PropertyRestriction.h"
+#include "Templates/Casts.h"
+#include "UObject/Class.h"
+#include "UObject/NameTypes.h"
+#include "UObject/ObjectPtr.h"
+#include "UObject/ReflectedTypeAccessors.h"
+#include "UObject/UObjectGlobals.h"
+#include "UObject/UnrealType.h"
+#include "LandscapeGrassType.h"
+
+class UObject;
+
+static bool GShowBothPerQualityAndPerPlaformProperties = false;
+static FAutoConsoleVariableRef CVarShowBothPerQualityAndPerPlaform(
+	TEXT("r.grass.ShowBothPerQualityAndPerPlaformProperties"),
+	GShowBothPerQualityAndPerPlaformProperties,
+	TEXT("Show both per platform and per quality properties in the editor."));
+
+#define LOCTEXT_NAMESPACE "LandscapeGrassTypeDetails"
+
+TSharedRef<IDetailCustomization> FLandscapeGrassTypeDetails::MakeInstance()
+{
+	return MakeShareable(new FLandscapeGrassTypeDetails);
+}
+
+FLandscapeGrassTypeDetails::~FLandscapeGrassTypeDetails()
+{
+}
+
+void FLandscapeGrassTypeDetails::CustomizeDetails( IDetailLayoutBuilder& DetailLayout )
+{
+	if (GShowBothPerQualityAndPerPlaformProperties)
+	{
+		return;
+	}
+
+	TArray<TWeakObjectPtr<UObject>> Objects;
+	DetailLayout.GetObjectsBeingCustomized(Objects);
+
+	const TArray<FName> PerPlatformProperties =
+	{
+		GET_MEMBER_NAME_CHECKED(FGrassVariety, GrassDensity),
+		GET_MEMBER_NAME_CHECKED(FGrassVariety, StartCullDistance),
+		GET_MEMBER_NAME_CHECKED(FGrassVariety, EndCullDistance),
+	};
+
+	const TArray<FName> PerQualityProperties =
+	{
+		GET_MEMBER_NAME_CHECKED(FGrassVariety, GrassDensityQuality),
+		GET_MEMBER_NAME_CHECKED(FGrassVariety, StartCullDistanceQuality),
+		GET_MEMBER_NAME_CHECKED(FGrassVariety, EndCullDistanceQuality),
+	};
+
+	if (Objects.Num() > 0)
+	{
+		// get the grass variety array
+		TSharedPtr<IPropertyHandle> GrassVariety = DetailLayout.GetProperty(GET_MEMBER_NAME_CHECKED(ULandscapeGrassType, GrassVarieties));
+		TSharedPtr<IPropertyHandleArray> GrassVarietyArrayHandle = GrassVariety->AsArray();
+
+		if (GrassVarietyArrayHandle.IsValid())
+		{
+			uint32 NumElements;
+			GrassVarietyArrayHandle->GetNumElements(NumElements);
+
+			for (uint32 Index = 0; Index < NumElements; ++Index)
+			{
+				TSharedRef<IPropertyHandle> ElementHandle = GrassVarietyArrayHandle->GetElement(Index);
+
+				uint32 NumChildren;
+				ElementHandle->GetNumChildren(NumChildren);
+
+				for (uint32 ChildIndex = 0; ChildIndex < NumChildren; ++ChildIndex)
+				{
+					TSharedRef<IPropertyHandle> ChildHandle = ElementHandle->GetChildHandle(ChildIndex).ToSharedRef();
+					const FName& ChildPropertyName = ChildHandle->GetProperty()->GetFName();
+
+					if (PerPlatformProperties.Contains(ChildPropertyName) && GEngine && GEngine->UseGrassVarityPerQualityLevels)
+					{
+						DetailLayout.HideProperty(ChildHandle);
+					}
+
+					if (PerQualityProperties.Contains(ChildPropertyName) && GEngine && !GEngine->UseGrassVarityPerQualityLevels)
+					{
+						DetailLayout.HideProperty(ChildHandle);
+					}
+				}
+			}
+		}
+	}
+}
+#undef LOCTEXT_NAMESPACE
diff --git a/Engine/Source/Editor/LandscapeEditor/Private/LandscapeGrassTypeDetails.h b/Engine/Source/Editor/LandscapeEditor/Private/LandscapeGrassTypeDetails.h
new file mode 100644
--- /dev/null
+++ b/Engine/Source/Editor/LandscapeEditor/Private/LandscapeGrassTypeDetails.h
@@ -0,0 +1,27 @@
+// Copyright Epic Games, Inc. All Rights Reserved.
+
+#pragma once
+
+#include "Delegates/IDelegateInstance.h"
+#include "IDetailCustomization.h"
+#include "Templates/SharedPointer.h"
+#include "UObject/WeakObjectPtr.h"
+#include "UObject/WeakObjectPtrTemplates.h"
+
+class FPropertyRestriction;
+class IDetailLayoutBuilder;
+class UObject;
+struct FPropertyChangedEvent;
+
+class FLandscapeGrassTypeDetails : public IDetailCustomization
+{
+public:
+	/** Makes a new instance of this detail layout class for a specific detail view requesting it */
+	static TSharedRef<IDetailCustomization> MakeInstance();
+
+	virtual ~FLandscapeGrassTypeDetails();
+
+	/** IDetailCustomization interface */
+	virtual void CustomizeDetails( IDetailLayoutBuilder& DetailLayout ) override;
+};
+
diff --git a/Engine/Source/Editor/LandscapeEditor/Private/LandscapeProxyUIDetails.cpp b/Engine/Source/Editor/LandscapeEditor/Private/LandscapeProxyUIDetails.cpp
new file mode 100644
--- /dev/null
+++ b/Engine/Source/Editor/LandscapeEditor/Private/LandscapeProxyUIDetails.cpp
@@ -0,0 +1,613 @@
+// Copyright Epic Games, Inc. All Rights Reserved.
+
+#include "LandscapeProxyUIDetails.h"
+
+#include "Algo/AnyOf.h"
+#include "Algo/Count.h"
+#include "Algo/Find.h"
+#include "Algo/Transform.h"
+#include "Components/RuntimeVirtualTextureComponent.h"
+#include "Containers/Array.h"
+#include "Containers/Map.h"
+#include "DetailCategoryBuilder.h"
+#include "DetailLayoutBuilder.h"
+#include "DetailWidgetRow.h"
+#include "Engine/World.h"
+#include "Fonts/SlateFontInfo.h"
+#include "HAL/PlatformCrt.h"
+#include "IDetailPropertyRow.h"
+#include "Internationalization/Internationalization.h"
+#include "Internationalization/Text.h"
+#include "Landscape.h"
+#include "LandscapeEditTypes.h"
+#include "LandscapeInfo.h"
+#include "LandscapeProxy.h"
+#include "LandscapeStreamingProxy.h"
+#include "LandscapeSubsystem.h"
+#include "Layout/Margin.h"
+#include "Math/IntPoint.h"
+#include "Math/IntRect.h"
+#include "Math/UnrealMathSSE.h"
+#include "Misc/Attribute.h"
+#include "RuntimeVirtualTextureSetBounds.h"
+#include "ScopedTransaction.h"
+#include "Templates/Casts.h"
+#include "Types/SlateEnums.h"
+#include "UObject/LazyObjectPtr.h"
+#include "UObject/ObjectMacros.h"
+#include "UObject/ObjectPtr.h"
+#include "UObject/UObjectIterator.h"
+#include "UObject/WeakObjectPtr.h"
+#include "UObject/WeakObjectPtrTemplates.h"
+#include "VT/RuntimeVirtualTextureVolume.h"
+#include "Widgets/DeclarativeSyntaxSupport.h"
+#include "Widgets/Input/SButton.h"
+#include "Widgets/Input/SCheckBox.h"
+#include "Widgets/IToolTip.h"
+#include "Widgets/Layout/SBox.h"
+#include "Widgets/Text/STextBlock.h"
+#include "Widgets/Input/STextComboBox.h"
+
+class IPropertyHandle;
+class SWidget;
+class UObject;
+class URuntimeVirtualTexture;
+
+#define LOCTEXT_NAMESPACE "FLandscapeProxyUIDetails"
+
+FLandscapeProxyUIDetails::FLandscapeProxyUIDetails()
+{
+	// Position Precision options are copied from StaticMeshEditorTools.cpp 
+	auto PositionPrecisionValueToDisplayString = [](int32 Value)
+	{
+		if(Value <= 0)
+		{
+			return FString::Printf(TEXT("%dcm"), 1 << (-Value));
+		}
+		else
+		{
+			const float fValue = static_cast<float>(FMath::Exp2((double)-Value));
+			return FString::Printf(TEXT("1/%dcm (%.3gcm)"), 1 << Value, fValue);
+		}
+	};
+	
+	for (int32 i = MinNanitePrecision; i <= MaxNanitePrecision; i++)
+	{
+		TSharedPtr<FString> Option = MakeShared<FString>(PositionPrecisionValueToDisplayString(i));
+		
+		PositionPrecisionOptions.Add(Option);
+	}
+}
+
+TSharedRef<IDetailCustomization> FLandscapeProxyUIDetails::MakeInstance()
+{
+	return MakeShareable( new FLandscapeProxyUIDetails);
+}
+
+void FLandscapeProxyUIDetails::CustomizeDetails( IDetailLayoutBuilder& DetailBuilder )
+{
+	TArray<TWeakObjectPtr<UObject>> EditingObjects;
+	DetailBuilder.GetObjectsBeingCustomized(EditingObjects);
+
+	auto GenerateTextWidget = [](const FText& InText, bool bInBold = false) -> TSharedRef<SWidget>
+	{
+		return SNew(STextBlock)
+				.Font(bInBold? IDetailLayoutBuilder::GetDetailFontBold() : IDetailLayoutBuilder::GetDetailFont())
+				.Text(InText);
+	};
+
+	TArray<TWeakObjectPtr<ALandscapeProxy>> EditingProxies;
+	Algo::Transform(EditingObjects, EditingProxies, [](TWeakObjectPtr<UObject> InObject) { return TWeakObjectPtr<ALandscapeProxy>(Cast<ALandscapeProxy>(InObject.Get())); });
+
+	TArray<TWeakObjectPtr<ALandscape>> LandscapeActors;
+	TArray<TWeakObjectPtr<ALandscapeStreamingProxy>> EditingStreamingProxies;
+	for (TWeakObjectPtr<ALandscapeProxy> EditingProxy : EditingProxies)
+	{
+		if (EditingProxy.IsValid())
+		{
+			if (ALandscape* LandscapeActor = EditingProxy->GetLandscapeActor())
+			{
+				LandscapeActors.AddUnique(LandscapeActor);
+			}
+
+			if (ALandscapeStreamingProxy* LandscapeStreamingProxy = Cast<ALandscapeStreamingProxy>(EditingProxy.Get()))
+			{
+				EditingStreamingProxies.Add(LandscapeStreamingProxy);
+			}
+		}
+	}
+
+	ALandscapeStreamingProxy* LandscapeStreamingProxy = EditingStreamingProxies.IsEmpty() ? nullptr : EditingStreamingProxies[0].Get();
+
+	// Hide World Partition specific properties in non WP levels
+	const bool bShouldDisplayWorldPartitionProperties = Algo::AnyOf(EditingProxies, [](const TWeakObjectPtr<ALandscapeProxy> InProxy)
+	{
+		UWorld* World = InProxy.IsValid() ? InProxy->GetTypedOuter<UWorld>() : nullptr;
+		return UWorld::IsPartitionedWorld(World);
+	});
+
+	if (!bShouldDisplayWorldPartitionProperties)
+	{
+		DetailBuilder.HideProperty(DetailBuilder.GetProperty(GET_MEMBER_NAME_CHECKED(ALandscape, bAreNewLandscapeActorsSpatiallyLoaded), ALandscape::StaticClass()));
+	}
+
+	if (LandscapeActors.Num() == 1)
+	{
+		TWeakObjectPtr<ALandscape> LandscapeActor = LandscapeActors[0];
+		if (ULandscapeInfo* LandscapeInfo = LandscapeActor->GetLandscapeInfo())
+		{
+			IDetailCategoryBuilder& CategoryBuilder = DetailBuilder.EditCategory("Information", FText::GetEmpty(), ECategoryPriority::Important);
+
+			FText CreateRVTVolumesTooltip = LOCTEXT("Button_CreateVolumes_Tooltip", "Create volumes for the selected Runtime Virtual Textures.");
+				
+			FText RowDisplayText = LOCTEXT("LandscapeComponentResolution", "Component Resolution (Verts)");
+			CategoryBuilder.AddCustomRow(RowDisplayText)
+			.RowTag(TEXT("LandscapeComponentResolution"))
+			.NameContent()
+			[
+				GenerateTextWidget(RowDisplayText)
+			]
+			.ValueContent()
+			[
+				GenerateTextWidget(FText::Format(LOCTEXT("LandscapeComponentResolutionValue", "{0} x {0}"), LandscapeActor->ComponentSizeQuads+1), true) // Verts
+			];
+
+			RowDisplayText = LOCTEXT("LandscapeComponentCount", "Component Count");
+			int32 NumComponents = LandscapeActor->LandscapeComponents.Num();
+			// If displaying streaming proxies, don't show the main landscape actor's component count, which is always 0 : 
+			if (!EditingStreamingProxies.IsEmpty())
+			{
+				NumComponents = EditingStreamingProxies[0]->LandscapeComponents.Num();
+				if (Algo::AnyOf(EditingStreamingProxies, [NumComponents](const TWeakObjectPtr<ALandscapeStreamingProxy> InStreamingProxy) { return InStreamingProxy.Get()->LandscapeComponents.Num() != NumComponents; }))
+				{
+					NumComponents = -1;
+				}
+			}
+			CategoryBuilder.AddCustomRow(RowDisplayText)
+			.RowTag(TEXT("LandscapeComponentCount"))
+			.NameContent()
+			[
+				GenerateTextWidget(RowDisplayText)
+			]
+			.ValueContent()
+			[
+				GenerateTextWidget((NumComponents == -1) ? LOCTEXT("MultipleValues", "Multiple Values") : FText::Format(LOCTEXT("LandscapeComponentCountValue", "{0}"), NumComponents), true)
+			];
+
+			RowDisplayText = LOCTEXT("LandscapeComponentSubsections", "Component Subsections");
+			CategoryBuilder.AddCustomRow(RowDisplayText)
+			.RowTag(TEXT("LandscapeComponentSubsections"))
+			.NameContent()
+			[
+				GenerateTextWidget(RowDisplayText)
+			]
+			.ValueContent()
+			[
+				GenerateTextWidget(FText::Format(LOCTEXT("LandscapeComponentSubSectionsValue", "{0} x {0}"), LandscapeActor->NumSubsections), true)
+			];
+
+			FIntRect Rect = LandscapeActor->GetBoundingRect();
+			FIntPoint Size = Rect.Size();
+			RowDisplayText = LOCTEXT("LandscapeResolution", "Resolution (Verts)");
+			CategoryBuilder.AddCustomRow(RowDisplayText)
+			.RowTag(TEXT("LandscapeResolution"))
+			.NameContent()
+			[
+				GenerateTextWidget(RowDisplayText)
+			]
+			.ValueContent()
+			[
+				GenerateTextWidget(FText::Format(LOCTEXT("LandscapeResolutionValue", "{0} x {1}"), Size.X+1, Size.Y+1), true)
+			];
+
+			int32 LandscapeCount = LandscapeInfo->StreamingProxies.Num() + (LandscapeInfo->LandscapeActor.Get() ? 1 : 0);
+			RowDisplayText = LOCTEXT("LandscapeCount", "Landscape Proxy Count");
+			CategoryBuilder.AddCustomRow(RowDisplayText)
+			.RowTag(TEXT("LandscapeCount"))
+			.NameContent()
+			[
+				GenerateTextWidget(RowDisplayText)
+			]
+			.ValueContent()
+			[
+				GenerateTextWidget(FText::Format(LOCTEXT("LandscapeProxyCountValue", "{0}"), LandscapeCount), true)
+			];
+
+			int32 TotalComponentCount = LandscapeInfo->XYtoComponentMap.Num();
+			RowDisplayText = LOCTEXT("TotalLandscapeComponentCount", "Total Component Count");
+			CategoryBuilder.AddCustomRow(RowDisplayText)
+			.RowTag(TEXT("TotalLandscapeComponentCount"))
+			.NameContent()
+			[
+				GenerateTextWidget(RowDisplayText)
+			]
+			.ValueContent()
+			[
+				GenerateTextWidget(FText::Format(LOCTEXT("TotalLandscapeComponentCountValue", "{0}"), TotalComponentCount), true)
+			];
+
+			LandscapeInfo->GetLandscapeExtent(Rect.Min.X, Rect.Min.Y, Rect.Max.X, Rect.Max.Y);
+			Size = Rect.Size();
+			RowDisplayText = LOCTEXT("LandscapeOverallResolution", "Overall Resolution (Verts)");
+			CategoryBuilder.AddCustomRow(RowDisplayText)
+			.RowTag(TEXT("LandscapeOverallResolution"))
+			.NameContent()
+			[
+				GenerateTextWidget(RowDisplayText)
+			]
+			.ValueContent()
+			[
+				GenerateTextWidget(FText::Format(LOCTEXT("LandscapeOverallResolutionValue", "{0} x {1}"), Size.X + 1, Size.Y + 1), true)
+			];
+
+			// Apply custom widget for CreateVolume.
+			TSharedRef<IPropertyHandle> CreateVolumePropertyHandle = DetailBuilder.GetProperty(TEXT("bSetCreateRuntimeVirtualTextureVolumes"));
+			DetailBuilder.EditDefaultProperty(CreateVolumePropertyHandle)->CustomWidget()
+			.NameContent()
+			[
+				SNew(STextBlock)
+				.Font(IDetailLayoutBuilder::GetDetailFont())
+				.Text(LOCTEXT("Button_CreateVolumes", "Create Volumes"))
+				.ToolTipText(CreateRVTVolumesTooltip)
+			]
+			.ValueContent()
+			.MinDesiredWidth(125.f)
+			[
+				SNew(SButton)
+				.VAlign(VAlign_Center)
+				.HAlign(HAlign_Center)
+				.ContentPadding(2.f)
+				.Text(LOCTEXT("Button_CreateVolumes", "Create Volumes"))
+				.OnClicked_Lambda([LandscapeActor, this]() { return CreateRuntimeVirtualTextureVolume(LandscapeActor.Get()); })
+				.IsEnabled_Lambda([LandscapeActor, this]() { return IsCreateRuntimeVirtualTextureVolumeEnabled(LandscapeActor.Get()); })
+				.ToolTipText_Lambda([LandscapeActor, this, CreateRVTVolumesTooltip]()
+				{
+					return IsCreateRuntimeVirtualTextureVolumeEnabled(LandscapeActor.Get())
+						? CreateRVTVolumesTooltip
+						: LOCTEXT("Button_CreateVolumes_Tooltip_Invalid", "No Runtime Virtual Textures found in the world.");
+				})
+			];
+		}
+		
+		TSharedRef<IPropertyHandle> NanitePositionPrecisionHandle = DetailBuilder.GetProperty(TEXT("NanitePositionPrecision"));
+		IDetailPropertyRow* DetailRow = DetailBuilder.EditDefaultProperty(NanitePositionPrecisionHandle);
+
+		// todo don.boogert : this is the handling of disabling of inherited properties on Landscape Proxies. 
+		// todo don.boogert : be able to handle override & inherited properties which also have custom UI.
+		if (LandscapeStreamingProxy->IsPropertyInherited(NanitePositionPrecisionHandle->GetProperty()))
+		{
+			if (LandscapeStreamingProxy != nullptr)
+			{
+				if (DetailRow != nullptr)
+				{
+					// Extend the tool tip to indicate this property is inherited
+					FText ToolTipText = NanitePositionPrecisionHandle->GetToolTipText();
+					DetailRow->ToolTip(FText::Format(LOCTEXT("InheritedProperty", "{0} This property is inherited from the parent Landscape proxy."), ToolTipText));
+			
+					// Disable the property editing
+					DetailRow->IsEnabled(false);
+				}
+			}
+		}
+		
+		auto GetPositionPrecision = [MinNanitePrecision = this->MinNanitePrecision, &PositionOptions = this->PositionPrecisionOptions, LandscapeActor]()
+		{
+			return PositionOptions[LandscapeActor->GetNanitePositionPrecision() - MinNanitePrecision] ;
+		};
+
+		auto SetPositionPrecision = [NanitePositionPrecisionHandle, MinNanitePrecision = this->MinNanitePrecision, &PositionOptions = this->PositionPrecisionOptions, LandscapeActor](TSharedPtr<FString> NewValue, ESelectInfo::Type SelectInfo)
+		{
+			if (!LandscapeActor.IsValid())
+			{
+				return;
+			}
+					
+			if (const int32 Index = PositionOptions.Find(NewValue); Index != INDEX_NONE)
+			{
+				NanitePositionPrecisionHandle->SetValue(Index + MinNanitePrecision);
+			}
+		};
+		
+		DetailRow->CustomWidget()
+		.NameContent()
+		[
+			SNew(STextBlock)
+			.Font(IDetailLayoutBuilder::GetDetailFont())
+			.Text(LOCTEXT("LandscapeNanitePositionPrecision", "Nanite Position Precision"))
+			.ToolTipText(LOCTEXT("LandscapeNanitePositionPrecisionTooltip", "Precision of Nanite vertex positions in World Space."))
+		]
+		.ValueContent()
+		.VAlign(VAlign_Center)
+		[
+			SNew(STextComboBox)
+			.Font(IDetailLayoutBuilder::GetDetailFont())
+			.OptionsSource(&PositionPrecisionOptions)
+			.InitiallySelectedItem(GetPositionPrecision())
+			.OnSelectionChanged_Lambda(SetPositionPrecision)
+		];
+	}
+
+	// Add Nanite buttons :
+	if (!EditingProxies.IsEmpty())
+	{
+		auto BuildNaniteData = [EditingProxies](UE::Landscape::EBuildFlags InBuildFlags) -> FReply
+		{
+			TArray<ALandscapeProxy*> ProxiesToBuild;
+			Algo::TransformIf(EditingProxies, ProxiesToBuild, [](const TWeakObjectPtr<ALandscapeProxy>& InProxy) { return InProxy.IsValid(); }, [](const TWeakObjectPtr<ALandscapeProxy>& InProxy) { return InProxy.Get(); });
+			if (ProxiesToBuild.IsEmpty())
+			{
+				return FReply::Unhandled();
+			}
+
+			if (UWorld* World = ProxiesToBuild[0]->GetWorld())
+			{
+				if (ULandscapeSubsystem* LandscapeSubsystem = World->GetSubsystem<ULandscapeSubsystem>())
+				{
+					LandscapeSubsystem->BuildNanite(InBuildFlags, MakeArrayView(ProxiesToBuild));
+				}
+			}
+
+			return FReply::Handled();
+		};
+
+		auto GetNumProxiesNeedingRebuild = [EditingProxies](UE::Landscape::EBuildFlags InBuildFlags) -> int32
+		{
+			const bool bForceRebuild = EnumHasAnyFlags(InBuildFlags, UE::Landscape::EBuildFlags::ForceRebuild);
+			TSet<TWeakObjectPtr<ALandscapeProxy>> ProxiesToBuild;
+			for (TWeakObjectPtr<ALandscapeProxy> EditingProxy : EditingProxies)
+			{
+				if (EditingProxy.IsValid())
+				{
+					ProxiesToBuild.Add(EditingProxy);
+					// Build all streaming proxies in the case of a ALandscape :
+					if (ALandscape* Landscape = Cast<ALandscape>(EditingProxy.Get()))
+					{
+						ULandscapeInfo* LandscapeInfo = Landscape->GetLandscapeInfo();
+						if (LandscapeInfo != nullptr)
+						{
+							Algo::Transform(LandscapeInfo->StreamingProxies, ProxiesToBuild, [](const TWeakObjectPtr<ALandscapeStreamingProxy>& InStreamingProxy) { return InStreamingProxy; });
+						}
+					}
+				}
+			}
+
+			return static_cast<int32>(Algo::CountIf(ProxiesToBuild, [bForceRebuild](TWeakObjectPtr<ALandscapeProxy> InProxy) { return (InProxy.IsValid() && (bForceRebuild || !InProxy->IsNaniteMeshUpToDate())); }));
+		};
+
+		auto HasAtLeastOneNaniteLandscape = [LandscapeActors]() -> bool
+		{
+			return Algo::FindByPredicate(LandscapeActors, [](TWeakObjectPtr<ALandscape> InLandscape) { return (InLandscape.IsValid() && InLandscape->IsNaniteEnabled()); }) != nullptr;
+		};
+
+		IDetailCategoryBuilder& CategoryBuilder = DetailBuilder.EditCategory("Nanite", FText::GetEmpty(), ECategoryPriority::Default);
+		CategoryBuilder.AddCustomRow(FText::FromString(TEXT("Rebuild Nanite Data")))
+		.RowTag("RebuildNaniteData")
+		[
+			SNew(SHorizontalBox)
+			.IsEnabled_Lambda([HasAtLeastOneNaniteLandscape] { return HasAtLeastOneNaniteLandscape(); })
+			+ SHorizontalBox::Slot()
+			.FillWidth(1)
+			[
+				SNew(SButton)
+				.Text(LOCTEXT("BuildNaniteData", "Build Data"))
+				.HAlign(HAlign_Center)
+				.VAlign(VAlign_Center)
+				.ToolTipText_Lambda([GetNumProxiesNeedingRebuild]()
+				{ 
+					return FText::Format(LOCTEXT("BuildNaniteDataTooltip", "Builds the Nanite mesh representation from the Landscape data if it's not up to date ({0} {0}|plural(one=actors, other=actors) to build)"), GetNumProxiesNeedingRebuild(UE::Landscape::EBuildFlags::None));
+				})
+				.OnClicked_Lambda([BuildNaniteData]() { return BuildNaniteData(UE::Landscape::EBuildFlags::WriteFinalLog); })
+				.IsEnabled_Lambda([GetNumProxiesNeedingRebuild]() { return (GetNumProxiesNeedingRebuild(UE::Landscape::EBuildFlags::None) > 0); })
+			]
+			+ SHorizontalBox::Slot()
+			.FillWidth(1)
+			[
+				SNew(SButton)
+				.Text(LOCTEXT("RebuildNaniteData", "Rebuild Data"))
+				.HAlign(HAlign_Center)
+				.VAlign(VAlign_Center)
+				.ToolTipText_Lambda([GetNumProxiesNeedingRebuild]()
+				{ 
+					return FText::Format(LOCTEXT("RebuildNaniteDataTooltip", "Rebuilds the Nanite mesh representation from the Landscape data ({0} {0}|plural(one=actors, other=actors) to build)"), GetNumProxiesNeedingRebuild(UE::Landscape::EBuildFlags::ForceRebuild));
+				})
+				.OnClicked_Lambda([BuildNaniteData]() { return BuildNaniteData(UE::Landscape::EBuildFlags::WriteFinalLog | UE::Landscape::EBuildFlags::ForceRebuild); })
+			]
+		];
+	}
+
+	if (LandscapeStreamingProxy != nullptr)
+	{
+		for (TFieldIterator<FProperty> PropertyIterator(LandscapeStreamingProxy->GetClass()); PropertyIterator; ++PropertyIterator)
+		{
+			FProperty* Property = *PropertyIterator;
+
+			if (Property == nullptr)
+			{
+				continue;
+			}
+
+			if (LandscapeStreamingProxy->IsPropertyInherited(Property))
+			{
+				TSharedPtr<IPropertyHandle> PropertyHandle = DetailBuilder.GetProperty(Property->GetFName());
+				
+				if (PropertyHandle->IsValidHandle())
+				{
+					IDetailPropertyRow* DetailRow = DetailBuilder.EditDefaultProperty(PropertyHandle);
+
+					// Extend the tool tip to indicate this property is inherited
+					FText ToolTipText = PropertyHandle->GetToolTipText();
+					DetailRow->ToolTip(FText::Format(LOCTEXT("InheritedProperty", "{0} This property is inherited from the parent Landscape proxy."), ToolTipText));
+
+					// Disable the property editing
+					DetailRow->IsEnabled(false);
+				}
+			}
+			else if (LandscapeStreamingProxy->IsPropertyOverridable(Property))
+			{
+				TSharedPtr<IPropertyHandle> PropertyHandle = DetailBuilder.GetProperty(Property->GetFName());
+
+				if (PropertyHandle->IsValidHandle())
+				{
+					const FText TooltipText = LOCTEXT("OverriddenProperty", "Check this box to override the parent landscape's property.");
+					FName PropertyName = Property->GetFName();
+					IDetailPropertyRow* DetailRow = DetailBuilder.EditDefaultProperty(PropertyHandle);
+					TSharedPtr<SWidget> NameWidget = nullptr;
+					TSharedPtr<SWidget> ValueWidget = nullptr;
+
+					DetailRow->GetDefaultWidgets(NameWidget, ValueWidget);
+
+					TAttribute<bool> EnabledAttribute = TAttribute<bool>::Create([LandscapeStreamingProxy, PropertyName]() -> bool
+					{
+						return LandscapeStreamingProxy->IsSharedPropertyOverridden(PropertyName);
+					});
+
+					DetailRow->CustomWidget(/*bShowChildren = */true)
+					.NameContent()
+					[
+						SNew(SHorizontalBox)
+						+ SHorizontalBox::Slot()
+						.AutoWidth()
+						[
+							SNew(SCheckBox)
+							.ToolTipText(TooltipText)
+							.IsChecked_Lambda([EditingStreamingProxies, PropertyName]() -> ECheckBoxState
+							{
+								bool bContainsOverriddenProperty = false;
+								bool bContainsDefaultProperty = false;
+
+								for (const TWeakObjectPtr<ALandscapeStreamingProxy> StreamingProxy : EditingStreamingProxies)
+								{
+									if (!StreamingProxy.IsValid())
+									{
+										continue;
+									}
+
+									const bool bIsPropertyOverridden = StreamingProxy->IsSharedPropertyOverridden(PropertyName);
+
+									bContainsOverriddenProperty |= bIsPropertyOverridden;
+									bContainsDefaultProperty |= !bIsPropertyOverridden;
+								}
+
+								if (bContainsOverriddenProperty && bContainsDefaultProperty)
+								{
+									return ECheckBoxState::Undetermined;
+								}
+
+								return bContainsOverriddenProperty ? ECheckBoxState::Checked : ECheckBoxState::Unchecked;
+							})
+							.OnCheckStateChanged_Lambda([EditingStreamingProxies, PropertyName](ECheckBoxState NewState)
+							{
+								if (NewState == ECheckBoxState::Undetermined)
+								{
+									return;
+								}
+
+								const bool bChecked = NewState == ECheckBoxState::Checked;
+
+								const FScopedTransaction Transaction(LOCTEXT("SetSharedPropertyOverride", "Change Property Override"));
+
+								for (const TWeakObjectPtr<ALandscapeStreamingProxy> StreamingProxy : EditingStreamingProxies)
+								{
+									if (!StreamingProxy.IsValid())
+									{
+										continue;
+									}
+
+									StreamingProxy->SetSharedPropertyOverride(PropertyName, bChecked);
+								}
+							})
+						]
+						+ SHorizontalBox::Slot()
+						.AutoWidth()
+						[
+							SNew(SBox)
+							.IsEnabled(EnabledAttribute)
+							[
+								NameWidget->AsShared()
+							]
+						]
+					]
+					.ValueContent()
+					[
+						SNew(SBox)
+						.IsEnabled(EnabledAttribute)
+						[
+							ValueWidget->AsShared()
+						]
+					]
+					.IsValueEnabled(EnabledAttribute);
+				}
+			}
+		}
+	}
+}
+
+static void GetMissingRuntimeVirtualTextureVolumes(ALandscape* InLandscapeActor, TArray<URuntimeVirtualTexture*>& OutVirtualTextures)
+{
+	UWorld* World = InLandscapeActor != nullptr ? InLandscapeActor->GetWorld() : nullptr;
+	if (World == nullptr)
+	{
+		return;
+	}
+
+	TArray<URuntimeVirtualTexture*> FoundVolumes;
+	for (TObjectIterator<URuntimeVirtualTextureComponent> It(RF_ClassDefaultObject, false, EInternalObjectFlags::Garbage); It; ++It)
+	{
+		if (It->GetWorld() == World)
+		{
+			if (URuntimeVirtualTexture* VirtualTexture = It->GetVirtualTexture())
+			{
+				FoundVolumes.Add(VirtualTexture);
+			}
+		}
+	}
+
+	for (URuntimeVirtualTexture* VirtualTexture : InLandscapeActor->RuntimeVirtualTextures)
+	{
+		if (VirtualTexture != nullptr && FoundVolumes.Find(VirtualTexture) == INDEX_NONE)
+		{
+			OutVirtualTextures.Add(VirtualTexture);
+		}
+	}
+}
+
+bool FLandscapeProxyUIDetails::IsCreateRuntimeVirtualTextureVolumeEnabled(ALandscape* InLandscapeActor) const
+{
+	if (InLandscapeActor == nullptr)
+	{
+		return false;
+	}
+
+	TArray<URuntimeVirtualTexture*> VirtualTextureVolumesToCreate;
+	GetMissingRuntimeVirtualTextureVolumes(InLandscapeActor, VirtualTextureVolumesToCreate);
+	return VirtualTextureVolumesToCreate.Num() > 0;
+}
+
+FReply FLandscapeProxyUIDetails::CreateRuntimeVirtualTextureVolume(ALandscape* InLandscapeActor)
+{
+	if (InLandscapeActor == nullptr)
+	{
+		return FReply::Unhandled();
+	}
+
+	TArray<URuntimeVirtualTexture*> VirtualTextureVolumesToCreate;
+	GetMissingRuntimeVirtualTextureVolumes(InLandscapeActor, VirtualTextureVolumesToCreate);
+	if (VirtualTextureVolumesToCreate.Num() == 0)
+	{
+		return FReply::Unhandled();
+	}
+
+	const FScopedTransaction Transaction(LOCTEXT("Transaction_CreateVolumes", "Create Runtime Virtual Texture Volumes"));
+	
+	for (URuntimeVirtualTexture* VirtualTexture : VirtualTextureVolumesToCreate)
+	{
+		ARuntimeVirtualTextureVolume* NewVolume = InLandscapeActor->GetWorld()->SpawnActor<ARuntimeVirtualTextureVolume>();
+		NewVolume->VirtualTextureComponent->SetVirtualTexture(VirtualTexture);
+		NewVolume->VirtualTextureComponent->SetBoundsAlignActor(InLandscapeActor);
+		RuntimeVirtualTexture::SetBounds(NewVolume->VirtualTextureComponent);
+	}
+
+	return FReply::Handled();
+}
+
+#undef LOCTEXT_NAMESPACE
diff --git a/Engine/Source/Editor/LandscapeEditor/Private/LandscapeProxyUIDetails.h b/Engine/Source/Editor/LandscapeEditor/Private/LandscapeProxyUIDetails.h
new file mode 100644
--- /dev/null
+++ b/Engine/Source/Editor/LandscapeEditor/Private/LandscapeProxyUIDetails.h
@@ -0,0 +1,38 @@
+// Copyright Epic Games, Inc. All Rights Reserved.
+
+#pragma once
+
+#include "IDetailCustomization.h"
+#include "Input/Reply.h"
+#include "Templates/SharedPointer.h"
+
+class ALandscape;
+class IDetailLayoutBuilder;
+
+class FLandscapeProxyUIDetails : public IDetailCustomization
+{
+public:
+	/** Makes a new instance of this detail layout class for a specific detail view requesting it */
+	static TSharedRef<IDetailCustomization> MakeInstance();
+
+	/** IDetailCustomization interface */
+	virtual void CustomizeDetails( IDetailLayoutBuilder& DetailBuilder ) override;
+	/** End IDetailCustomization interface */
+
+private:
+	/** Use MakeInstance to create an instance of this class */
+	FLandscapeProxyUIDetails();
+
+	/** Returns true if SetBounds button is enabled */
+	bool IsCreateRuntimeVirtualTextureVolumeEnabled(ALandscape* InLandscapeActor) const;
+	
+	/** Callback for Set Bounds button */
+	FReply CreateRuntimeVirtualTextureVolume(ALandscape* InLandscapeActor);
+	
+	IDetailLayoutBuilder* DetailLayoutBuilder = nullptr;
+
+	/** Nanite Position Precision Options */
+	TArray<TSharedPtr<FString>> PositionPrecisionOptions;
+	static constexpr int MinNanitePrecision = -6;
+	static constexpr int MaxNanitePrecision = 13;
+};
diff --git a/Engine/Source/Editor/LandscapeEditor/Private/LandscapeUIDetails.cpp b/Engine/Source/Editor/LandscapeEditor/Private/LandscapeUIDetails.cpp
new file mode 100644
--- /dev/null
+++ b/Engine/Source/Editor/LandscapeEditor/Private/LandscapeUIDetails.cpp
@@ -0,0 +1,150 @@
+// Copyright Epic Games, Inc. All Rights Reserved.
+
+#include "LandscapeUIDetails.h"
+
+#include "Containers/Array.h"
+#include "Delegates/Delegate.h"
+#include "DetailLayoutBuilder.h"
+#include "DetailWidgetRow.h"
+#include "Editor.h"
+#include "Editor/EditorEngine.h"
+#include "HAL/Platform.h"
+#include "HAL/PlatformCrt.h"
+#include "HAL/PlatformMisc.h"
+#include "Input/Reply.h"
+#include "Internationalization/Internationalization.h"
+#include "Internationalization/Text.h"
+#include "Landscape.h"
+#include "LandscapeEditLayer.h"
+#include "LandscapeInfo.h"
+#include "LandscapeProxy.h"
+#include "LandscapeSettings.h"
+#include "Layout/Visibility.h"
+#include "Misc/AssertionMacros.h"
+#include "Misc/Attribute.h"
+#include "Misc/MessageDialog.h"
+#include "Misc/Optional.h"
+#include "PropertyHandle.h"
+#include "Styling/SlateTypes.h"
+#include "Templates/Casts.h"
+#include "Types/SlateEnums.h"
+#include "UObject/ObjectPtr.h"
+#include "UObject/UObjectGlobals.h"
+#include "UObject/WeakObjectPtr.h"
+#include "UObject/WeakObjectPtrTemplates.h"
+#include "Widgets/DeclarativeSyntaxSupport.h"
+#include "Widgets/Input/SButton.h"
+#include "Widgets/Input/SCheckBox.h"
+#include "Misc/ScopedSlowTask.h"
+
+class UObject;
+
+#define LOCTEXT_NAMESPACE "FLandscapeUIDetails"
+
+FLandscapeUIDetails::FLandscapeUIDetails()
+{
+}
+
+FLandscapeUIDetails::~FLandscapeUIDetails()
+{
+}
+
+TSharedRef<IDetailCustomization> FLandscapeUIDetails::MakeInstance()
+{
+	return MakeShareable( new FLandscapeUIDetails);
+}
+
+void FLandscapeUIDetails::CustomizeDetails( IDetailLayoutBuilder& DetailBuilder )
+{
+	TArray<TWeakObjectPtr<UObject>> EditingObjects;
+	DetailBuilder.GetObjectsBeingCustomized(EditingObjects);
+
+	if (EditingObjects.Num() == 1)
+	{
+		TWeakObjectPtr<ALandscape> Landscape(Cast<ALandscape>(EditingObjects[0].Get()));
+		if (!Landscape.IsValid())
+		{
+			return;
+		}
+
+		TSharedRef<IPropertyHandle> CanHaveLayersPropertyHandle = DetailBuilder.GetProperty(GET_MEMBER_NAME_CHECKED(ALandscape, bCanHaveLayersContent));
+		DetailBuilder.HideProperty(CanHaveLayersPropertyHandle);
+		const FText DisplayAndFilterText(LOCTEXT("LandscapeToggleLayerName", "Enable Edit Layers"));
+		const FText ToolTipText(LOCTEXT("LandscapeToggleLayerToolTip", "Toggle whether or not to support edit layers on this Landscape. Toggling this will clear the undo stack."));
+		DetailBuilder.AddCustomRowToCategory(CanHaveLayersPropertyHandle, DisplayAndFilterText)
+		.RowTag(TEXT("EnableEditLayers"))
+		.NameContent()
+		[
+			CanHaveLayersPropertyHandle->CreatePropertyNameWidget(DisplayAndFilterText, ToolTipText)
+		]
+		.Visibility(MakeAttributeLambda([]()
+		{
+			const ULandscapeSettings* Settings = GetDefault<ULandscapeSettings>();
+			return Settings->InRestrictiveMode() ? EVisibility::Hidden : EVisibility::Visible;
+		}))
+		.ValueContent()
+		[
+			SNew(SCheckBox)
+			.ToolTipText(ToolTipText)
+			.Type(ESlateCheckBoxType::CheckBox)
+			.IsChecked_Lambda([=]()
+			{
+				return Landscape.IsValid() && Landscape->CanHaveLayersContent() ? ECheckBoxState::Checked : ECheckBoxState::Unchecked;
+			})
+			.OnCheckStateChanged_Lambda([this, Landscape](ECheckBoxState NewState)
+			{
+				bool bChecked = (NewState == ECheckBoxState::Checked);
+				if (Landscape.IsValid() && (Landscape->CanHaveLayersContent() != bChecked))
+				{
+					ToggleCanHaveLayersContent(Landscape.Get());
+				}
+			})
+		];
+	}
+}
+
+void FLandscapeUIDetails::ToggleCanHaveLayersContent(ALandscape* Landscape)
+{
+	bool bToggled = false;
+
+	if (Landscape->bCanHaveLayersContent)
+	{
+		bool bHasHiddenLayers = false;
+		for (const FLandscapeLayer& Layer : Landscape->GetLayersConst())
+		{
+			if (!Layer.EditLayer->IsVisible())
+			{
+				bHasHiddenLayers = true;
+				break;
+			}
+		}
+				
+		FText Reason;
+		if (bHasHiddenLayers)
+		{
+			Reason = LOCTEXT("LandscapeDisableLayers_HiddenLayers", "Are you sure you want to disable the edit layers on this Landscape?\n\nDoing so, will result in losing the data stored for each edit layer, but the current visual output will be kept. Be aware that some edit layers are currently hidden, continuing will result in their data being lost. Undo/redo buffer will also be cleared.");
+		}
+		else
+		{
+			Reason = LOCTEXT("LandscapeDisableLayers", "Are you sure you want to disable the edit layers on this Landscape?\n\nDoing so, will result in losing the data stored for each edit layers, but the current visual output will be kept. Undo/redo buffer will also be cleared.");
+		}
+
+		bToggled = FMessageDialog::Open(EAppMsgType::YesNo, Reason) == EAppReturnType::Yes;
+	}
+	else
+	{
+		
+		bToggled = FMessageDialog::Open(EAppMsgType::YesNo, LOCTEXT("LandscapeEnableLayers", "Are you sure you want to enable edit layers on this landscape? Doing so will clear the undo/redo buffer.")) == EAppReturnType::Yes;
+	}
+
+	if (bToggled)
+	{
+		Landscape->ToggleCanHaveLayersContent();
+		if (GEditor)
+		{
+			GEditor->ResetTransaction(LOCTEXT("ToggleLanscapeLayers", "Toggling Landscape Edit Layers"));
+		}
+	}
+}
+
+#undef LOCTEXT_NAMESPACE
diff --git a/Engine/Source/Editor/LandscapeEditor/Private/LandscapeUIDetails.h b/Engine/Source/Editor/LandscapeEditor/Private/LandscapeUIDetails.h
new file mode 100644
--- /dev/null
+++ b/Engine/Source/Editor/LandscapeEditor/Private/LandscapeUIDetails.h
@@ -0,0 +1,28 @@
+// Copyright Epic Games, Inc. All Rights Reserved.
+
+#pragma once
+
+#include "IDetailCustomization.h"
+#include "Templates/SharedPointer.h"
+
+class ALandscape;
+class IDetailLayoutBuilder;
+
+class FLandscapeUIDetails : public IDetailCustomization
+{
+public:
+	~FLandscapeUIDetails();
+
+	/** Makes a new instance of this detail layout class for a specific detail view requesting it */
+	static TSharedRef<IDetailCustomization> MakeInstance();
+
+	/** IDetailCustomization interface */
+	virtual void CustomizeDetails( IDetailLayoutBuilder& DetailBuilder ) override;
+	/** End IDetailCustomization interface */
+
+private:
+	void ToggleCanHaveLayersContent(ALandscape* Landscape);
+
+	/** Use MakeInstance to create an instance of this class */
+	FLandscapeUIDetails();
+};
diff --git a/Engine/Source/Editor/LandscapeEditor/Private/SLandscapeEditor.cpp b/Engine/Source/Editor/LandscapeEditor/Private/SLandscapeEditor.cpp
--- a/Engine/Source/Editor/LandscapeEditor/Private/SLandscapeEditor.cpp
+++ b/Engine/Source/Editor/LandscapeEditor/Private/SLandscapeEditor.cpp
@@ -4,8 +4,10 @@
 #include "Framework/MultiBox/MultiBoxDefs.h"
 #include "Framework/MultiBox/MultiBoxBuilder.h"
 #include "Materials/Material.h"
+#include "Widgets/Docking/SDockTab.h"
 #include "Widgets/Layout/SBorder.h"
 #include "Widgets/Input/SMultiLineEditableTextBox.h"
+#include "Widgets/Layout/SExpandableArea.h"
 #include "Modules/ModuleManager.h"
 #include "SlateOptMacros.h"
 #include "Widgets/Layout/SBox.h"
@@ -20,6 +22,7 @@
 #include "PropertyEditorModule.h"
 #include "LandscapeSettings.h"
 #include "HAL/IConsoleManager.h"
+#include "Toolkits/AssetEditorModeUILayer.h"
 
 #define LOCTEXT_NAMESPACE "LandscapeEditor"
 
@@ -331,6 +334,23 @@ void FLandscapeToolKit::RefreshDetailPanel()
 	LandscapeEditorWidgets->RefreshDetailPanel();
 }
 
+void FLandscapeToolKit::RefreshInspectedObjectsDetailPanel()
+{
+	// Focus or re-open the inspected objects details view
+	if (TSharedPtr<FAssetEditorModeUILayer> ModeUILayerPtr = ModeUILayer.Pin())
+	{
+		TSharedPtr<FTabManager> TabManagerPtr = ModeUILayerPtr->GetTabManager();
+		if (!TabManagerPtr)
+		{
+			return;
+		}
+		InspectedObjectsTab = TabManagerPtr->TryInvokeTab(UAssetEditorUISubsystem::BottomLeftTabID);
+
+		// Refresh the selected objects
+		InspectedObjectsDetailsView->RefreshDetailPanel();
+	}
+}
+
 void FLandscapeToolKit::OnChangeMode(FName ModeName)
 {
 	FEdModeLandscape* LandscapeEdMode = GetEditorMode();
@@ -492,6 +512,53 @@ bool FLandscapeToolKit::IsBrushActive(FName BrushName) const
 	return false;
 }
 
+void FLandscapeToolKit::RequestModeUITabs()
+{
+	FModeToolkit::RequestModeUITabs();
+	if (TSharedPtr<FAssetEditorModeUILayer> ModeUILayerPtr = ModeUILayer.Pin())
+	{
+		TSharedPtr<FWorkspaceItem> MenuModeCategoryPtr = ModeUILayerPtr->GetModeMenuCategory();
+
+		if (!MenuModeCategoryPtr)
+		{
+			return;
+		}
+
+		InspectedObjectsTabInfo.OnSpawnTab = FOnSpawnTab::CreateSP(SharedThis(this), &FLandscapeToolKit::CreateInspectedObjectsDetailsViewTab);
+		InspectedObjectsTabInfo.TabLabel = LOCTEXT("LandscapeModeDetailsViewToolboxTabLabel", "Landscape Details");
+		InspectedObjectsTabInfo.TabTooltip = LOCTEXT("LandscapeModeDetailsViewToolboxTabTooltipText", "Open the Landscape Details tab, which contains the landscape editor mode's inspected object details.");
+		InspectedObjectsTabInfo.TabIcon = GetEditorModeIcon();
+		ModeUILayerPtr->SetModePanelInfo(UAssetEditorUISubsystem::BottomLeftTabID, InspectedObjectsTabInfo);
+	}
+}
+
+void FLandscapeToolKit::InvokeUI()
+{
+	FModeToolkit::InvokeUI();
+
+	RefreshInspectedObjectsDetailPanel();
+}
+
+TSharedRef<SDockTab> FLandscapeToolKit::CreateInspectedObjectsDetailsViewTab(const FSpawnTabArgs& Args)
+{
+	FMargin MorePadding = FMargin(10.0f, 2.0f);
+
+	TSharedPtr<SDockTab> CreatedTab = SNew(SDockTab)
+		.Label(LOCTEXT("LandscapeDetailsTab", "Landscape Details"))
+		.ToolTipText(LOCTEXT("LandscapeDetailsTabToolTip", "Landscape Details"))
+		[
+			SNew(SVerticalBox)
+				+ SVerticalBox::Slot()
+				.Padding(2.0)
+				[
+					SAssignNew(InspectedObjectsDetailsView, SLandscapeEditorInspectedDetailsView)
+				]
+		];
+
+	InspectedObjectsTab = CreatedTab;
+	return CreatedTab.ToSharedRef();
+}
+
 //////////////////////////////////////////////////////////////////////////
 
 BEGIN_SLATE_FUNCTION_BUILD_OPTIMIZATION
@@ -509,7 +576,7 @@ void SLandscapeEditor::Construct(const FArguments& InArgs, TSharedRef<FLandscape
 	DetailsPanel = PropertyEditorModule.CreateDetailView(DetailsViewArgs);
 	DetailsPanel->SetIsPropertyVisibleDelegate(FIsPropertyVisible::CreateSP(this, &SLandscapeEditor::GetIsPropertyVisible));
 
-	FEdModeLandscape* LandscapeEdMode = GetEditorMode();
+	const FEdModeLandscape* LandscapeEdMode = GetEditorMode();
 	if (LandscapeEdMode)
 	{
 		DetailsPanel->SetObject(LandscapeEdMode->UISettings);
@@ -518,23 +585,24 @@ void SLandscapeEditor::Construct(const FArguments& InArgs, TSharedRef<FLandscape
 	ChildSlot
 	[
 		SNew(SVerticalBox)
-		+ SVerticalBox::Slot()
-		.AutoHeight()
-		.Padding(0, 0, 0, 5)
-		[
-			SAssignNew(Error, SErrorText)
-		]
-		+ SVerticalBox::Slot()
-		.Padding(0)
-		[
-			SNew(SVerticalBox)
-			.IsEnabled(this, &SLandscapeEditor::GetLandscapeEditorIsEnabled)
 			+ SVerticalBox::Slot()
+			.AutoHeight()
+			.Padding(0, 0, 0, 5)
+			[
+				SAssignNew(Error, SErrorText)
+			]
+			+ SVerticalBox::Slot()
+			.FillHeight(1.0f)
 			.Padding(0)
 			[
-				DetailsPanel.ToSharedRef()
+				SNew(SVerticalBox)
+					.IsEnabled(this, &SLandscapeEditor::GetLandscapeEditorIsEnabled)
+					+ SVerticalBox::Slot()
+					.Padding(0)
+					[
+						DetailsPanel.ToSharedRef()
+					]
 			]
-		]
 	];
 }
 END_SLATE_FUNCTION_BUILD_OPTIMIZATION
@@ -610,6 +678,15 @@ bool SLandscapeEditor::GetLandscapeEditorIsEnabled() const
 	return false;
 }
 
+bool SLandscapeEditor::GetInspectedObjectsDetailsIsVisible() const
+{
+	if (FEdModeLandscape* LandscapeEdMode = GetEditorMode())
+	{
+		return !LandscapeEdMode->GetInspectedObjects().IsEmpty();
+	}
+	return false;
+}
+
 bool SLandscapeEditor::GetIsPropertyVisible(const FPropertyAndParent& PropertyAndParent) const
 {
 	return ParentToolkit.Pin()->GetIsPropertyVisibleFromProperty(PropertyAndParent.Property);
@@ -721,12 +798,56 @@ void SLandscapeEditor::NotifyBrushChanged()
 
 void SLandscapeEditor::RefreshDetailPanel()
 {
-	FEdModeLandscape* LandscapeEdMode = GetEditorMode();
-	if (LandscapeEdMode)
+	if (const FEdModeLandscape* LandscapeEdMode = GetEditorMode())
 	{
 		// Refresh details panel
 		DetailsPanel->SetObject(LandscapeEdMode->UISettings, true);
 	}
 }
 
+//////////////////////////////////////////////////////////////////////////
+
+void SLandscapeEditorInspectedDetailsView::Construct(const FArguments& InArgs)
+{
+	FPropertyEditorModule& PropertyEditorModule = FModuleManager::LoadModuleChecked<FPropertyEditorModule>("PropertyEditor");
+
+	FDetailsViewArgs DetailsViewArgs;
+	DetailsViewArgs.NameAreaSettings = FDetailsViewArgs::ObjectsUseNameArea;
+	DetailsViewArgs.bAllowSearch = false;
+	DetailsViewArgs.bCustomNameAreaLocation = true;
+	// Minimum size to allow the ResetToDefault button to be hit testable.
+	DetailsViewArgs.RightColumnMinWidth = 35;
+
+	DetailsView = PropertyEditorModule.CreateDetailView(DetailsViewArgs);
+
+	ChildSlot
+		[
+			SNew(SVerticalBox)
+				+ SVerticalBox::Slot()
+				.Padding(10.f, 6.f, 0.f, 4.f)
+				.AutoHeight()
+				[
+					DetailsView->GetNameAreaWidget().ToSharedRef()
+				]
+				+ SVerticalBox::Slot()
+				.AutoHeight()
+				[
+					DetailsView.ToSharedRef()
+				]
+		];
+}
+
+FEdModeLandscape* SLandscapeEditorInspectedDetailsView::GetEditorMode() const
+{
+	return (FEdModeLandscape*)GLevelEditorModeTools().GetActiveMode(FBuiltinEditorModes::EM_Landscape);
+}
+
+void SLandscapeEditorInspectedDetailsView::RefreshDetailPanel()
+{
+	if (const FEdModeLandscape* LandscapeEdMode = GetEditorMode())
+	{
+		DetailsView->SetObjects(LandscapeEdMode->GetInspectedObjects());
+	}
+}
+
 #undef LOCTEXT_NAMESPACE
diff --git a/Engine/Source/Editor/LandscapeEditor/Private/SLandscapeEditor.h b/Engine/Source/Editor/LandscapeEditor/Private/SLandscapeEditor.h
--- a/Engine/Source/Editor/LandscapeEditor/Private/SLandscapeEditor.h
+++ b/Engine/Source/Editor/LandscapeEditor/Private/SLandscapeEditor.h
@@ -16,6 +16,7 @@ class Error;
 class IDetailsView;
 class SErrorText;
 class SLandscapeEditor;
+class SLandscapeEditorInspectedDetailsView;
 struct FPropertyAndParent;
 
 /**
@@ -56,7 +57,7 @@ namespace LandscapeEditorNames
 class FLandscapeToolKit : public FModeToolkit
 {
 public:
-	/** Initializes the geometry mode toolkit */
+	/** Initializes the landscape mode toolkit */
 	virtual void Init(const TSharedPtr< class IToolkitHost >& InitToolkitHost) override;
 
 	/** IToolkit interface */
@@ -68,6 +69,7 @@ public:
 	void NotifyToolChanged();
 	void NotifyBrushChanged();
 	void RefreshDetailPanel();
+	void RefreshInspectedObjectsDetailPanel();
 
 	/** Mode Toolbar Palettes **/
 	virtual void GetToolPaletteNames(TArray<FName>& InPaletteName) const;
@@ -98,10 +100,20 @@ protected:
 	void OnChangeBrush(FName BrushName);
 	bool IsBrushActive(FName BrushName) const;
 
+	/** FModeToolkit interface */
+	virtual void RequestModeUITabs() override;
+	virtual void InvokeUI() override;
+
+	TSharedRef<SDockTab> CreateInspectedObjectsDetailsViewTab(const FSpawnTabArgs& Args);
+
 private:
 	TSharedPtr<SLandscapeEditor> LandscapeEditorWidgets;
 	TSharedPtr<FLandscapeEditorDetails> BrushesWidgets;
 
+	TWeakPtr<SDockTab> InspectedObjectsTab;
+	FMinorTabConfig InspectedObjectsTabInfo;
+	TSharedPtr<SLandscapeEditorInspectedDetailsView> InspectedObjectsDetailsView;
+
 	const static TArray<FName> PaletteNames;
 };
 
@@ -126,6 +138,7 @@ protected:
 	FText GetErrorText() const;
 
 	bool GetLandscapeEditorIsEnabled() const;
+	bool GetInspectedObjectsDetailsIsVisible() const;
 
 	bool GetIsPropertyVisible(const FPropertyAndParent& PropertyAndParent) const;
 
@@ -134,3 +147,24 @@ protected:
 	TSharedPtr<IDetailsView> DetailsPanel;
 	TWeakPtr<FLandscapeToolKit> ParentToolkit;
 };
+
+/**
+ * Slate widget for the Landscape Editor Inspected Objects Details View
+ */
+class SLandscapeEditorInspectedDetailsView : public SCompoundWidget
+{
+public:
+	SLATE_BEGIN_ARGS(SLandscapeEditorInspectedDetailsView)
+	{
+	}
+	SLATE_END_ARGS()
+
+	void Construct(const FArguments& InArgs);
+	void RefreshDetailPanel();
+
+protected:
+	class FEdModeLandscape* GetEditorMode() const;
+
+private:
+	TSharedPtr<IDetailsView> DetailsView;
+};
diff --git a/Engine/Source/Editor/LandscapeEditor/Private/SLandscapeLayerListDialog.cpp b/Engine/Source/Editor/LandscapeEditor/Private/SLandscapeLayerListDialog.cpp
--- a/Engine/Source/Editor/LandscapeEditor/Private/SLandscapeLayerListDialog.cpp
+++ b/Engine/Source/Editor/LandscapeEditor/Private/SLandscapeLayerListDialog.cpp
@@ -2,6 +2,7 @@
 
 #include "SLandscapeLayerListDialog.h"
 #include "Landscape.h"
+#include "LandscapeEditLayer.h"
 #include "Widgets/Text/STextBlock.h"
 #include "Widgets/Views/STableViewBase.h"
 #include "Widgets/Views/STableRow.h"
@@ -15,11 +16,12 @@
 struct FWidgetLayerListItem
 {
 	FWidgetLayerListItem(const FLandscapeLayer& InLayer, const TFunction<void(void)>& InOnLayerListUpdated, TArray<TSharedPtr<FWidgetLayerListItem>>* InWidgetLayerList)
-	: LayerName(InLayer.Name)
-	, LayerGuid(InLayer.Guid)
+	: LayerName(InLayer.EditLayer->GetName())
+	, LayerGuid(InLayer.EditLayer->GetGuid())
 	, OnLayerListUpdated(InOnLayerListUpdated)
 	, WidgetLayerList(InWidgetLayerList)
-	{}
+	{ 
+	}
 
 	// Duplicated Layer Info
 	FName LayerName;
@@ -180,7 +182,7 @@ void SLandscapeLayerListDialog::Construct(const FArguments& InArgs, const TWeakO
 	check(Landscape.IsValid());
 
 	int32 LayerIndex = 0;
-	TArrayView<const FLandscapeLayer> LandscapeLayers = Landscape->GetLayers();
+	TArrayView<const FLandscapeLayer> LandscapeLayers = Landscape->GetLayersConst();
 	for (const FLandscapeLayer& Layer : LandscapeLayers)
 	{
 		WidgetLayerList.Add(MakeShared<FWidgetLayerListItem>(LandscapeLayers[LandscapeLayers.Num() - 1 - LayerIndex], [this](){ OnLayerListUpdated(); }, &WidgetLayerList));
@@ -256,7 +258,7 @@ TSharedRef<ITableRow> SLandscapeLayerListDialog::OnGenerateRow(TSharedPtr<FWidge
 void SLandscapeLayerListDialog::OnComplete()
 {
 	check(Landscape.IsValid());
-	TArrayView<const FLandscapeLayer> LandscapeLayers = Landscape->GetLayers();
+	TArrayView<const FLandscapeLayer> LandscapeLayers = Landscape->GetLayersConst();
 	// Find the draggable layer in WidgetLayerList
 	for (int WidgetLayerListIndex = 0; WidgetLayerListIndex < WidgetLayerList.Num(); ++WidgetLayerListIndex)
 	{
@@ -270,7 +272,7 @@ void SLandscapeLayerListDialog::OnComplete()
 				const FLandscapeLayer& Layer = LandscapeLayers[LayerIndex];
 
 				// Once found, remove and copy edit layer into correct spot in LayerList
-				if (WidgetLayer->LayerGuid == Layer.Guid)
+				if (WidgetLayer->LayerGuid == Layer.EditLayer->GetGuid())
 				{
 					InsertedLayerIndex = WidgetLayerListIndex;
 					Landscape->ReorderLayer(LayerIndex, InsertedLayerIndex);
diff --git a/Engine/Source/Runtime/Core/Public/UObject/FortniteMainBranchObjectVersions.inl b/Engine/Source/Runtime/Core/Public/UObject/FortniteMainBranchObjectVersions.inl
--- a/Engine/Source/Runtime/Core/Public/UObject/FortniteMainBranchObjectVersions.inl
+++ b/Engine/Source/Runtime/Core/Public/UObject/FortniteMainBranchObjectVersions.inl
@@ -625,5 +625,8 @@ FORTNITE_MAIN_VERSION(RectLightFixedEVUnitConversion, 203)
 // Add particle bounds to data exported to CVD
 FORTNITE_MAIN_VERSION(ParticleInflatedBoundsInChaosVisualDebugger, 204)
 
+// Migrate properties from FLandscapeLayer to ULandscapeEditLayer
+FORTNITE_MAIN_VERSION(MigrateLandscapeEditLayerProperties, 205)
+
 // -----<new versions can be added above this line>-------------------------------------------------
 #undef FORTNITE_MAIN_VERSION
\ No newline at end of file
diff --git a/Engine/Source/Runtime/Landscape/Classes/Landscape.h b/Engine/Source/Runtime/Landscape/Classes/Landscape.h
--- a/Engine/Source/Runtime/Landscape/Classes/Landscape.h
+++ b/Engine/Source/Runtime/Landscape/Classes/Landscape.h
@@ -38,6 +38,7 @@ struct FTextureToComponentHelper;
 struct FUpdateLayersContentContext;
 struct FEditLayersHeightmapMergeParams;
 struct FEditLayersWeightmapMergeParams;
+struct FOnLandscapeEditLayerDataChangedParams;
 enum class ELandscapeNotificationType;
 
 namespace EditLayersHeightmapLocalMerge_RenderThread
@@ -154,6 +155,7 @@ private:
 #endif // WITH_EDITORONLY_DATA
 };
 
+// TODO [jonathan.bard] : deprecate this
 UENUM()
 enum ELandscapeBlendMode : int
 {
@@ -171,13 +173,6 @@ struct FLandscapeLayer
 	GENERATED_USTRUCT_BODY()
 
 	FLandscapeLayer()
-		: Guid(FGuid::NewGuid())
-		, Name(NAME_None)
-		, bVisible(true)
-		, bLocked(false)
-		, HeightmapAlpha(1.0f)
-		, WeightmapAlpha(1.0f)
-		, BlendMode(LSBM_AdditiveBlend)
 	{}
 
 	virtual ~FLandscapeLayer() = default;
@@ -187,38 +182,34 @@ struct FLandscapeLayer
 	LANDSCAPE_EDIT_LAYERS_BATCHED_MERGE_EXPERIMENTAL
 	LANDSCAPE_API virtual TArray<UE::Landscape::EditLayers::FEditLayerRendererState> GetEditLayerRendererStates(const UE::Landscape::EditLayers::FMergeContext* InMergeContext) override;
 	//~ End UE::Landscape::EditLayers::IEditLayerRendererProvider implementation
-
-	ELandscapeToolTargetTypeFlags GetEnabledTargetTypeMask() const;
 #endif // WITH_EDITOR
-
-	// TODO [jonathan.bard] : make those private
 public:
 	UPROPERTY(meta = (IgnoreForMemberInitializationTest))
-	FGuid Guid;
+	FGuid Guid_DEPRECATED = FGuid::NewGuid();
 
 	UPROPERTY()
-	FName Name;
+	FName Name_DEPRECATED = NAME_None;
 
 	UPROPERTY(Transient)
-	bool bVisible;
+	bool bVisible_DEPRECATED = true;
 
 	UPROPERTY()
-	bool bLocked;
+	bool bLocked_DEPRECATED = false;
 
 	UPROPERTY()
-	float HeightmapAlpha;
+	float HeightmapAlpha_DEPRECATED = 1.0f;
 
 	UPROPERTY()
-	float WeightmapAlpha;
+	float WeightmapAlpha_DEPRECATED = 1.0f;
 
 	UPROPERTY()
-	TEnumAsByte<enum ELandscapeBlendMode> BlendMode;
+	TEnumAsByte<enum ELandscapeBlendMode> BlendMode_DEPRECATED = ELandscapeBlendMode::LSBM_AdditiveBlend;
 
 	UPROPERTY()
 	TArray<FLandscapeLayerBrush> Brushes;
 
 	UPROPERTY()
-	TMap<TObjectPtr<ULandscapeLayerInfoObject>, bool> WeightmapLayerAllocationBlend; // True -> Substractive, False -> Additive
+	TMap<TObjectPtr<ULandscapeLayerInfoObject>, bool> WeightmapLayerAllocationBlend_DEPRECATED; // True -> Substractive, False -> Additive
 
 	UPROPERTY(Instanced)
 	TObjectPtr<ULandscapeEditLayerBase> EditLayer;
@@ -408,21 +399,50 @@ public:
 	LANDSCAPE_API void AddLayersToProxy(ALandscapeProxy* InProxy);
 	LANDSCAPE_API FIntPoint ComputeComponentCounts() const;
 	LANDSCAPE_API bool IsLayerNameUnique(const FName& InName) const;
+
+	UE_DEPRECATED(5.6, "Use SetName on the ULandscapeEditLayerBase object")
 	LANDSCAPE_API void SetLayerName(int32 InLayerIndex, const FName& InName);
+	UE_DEPRECATED(5.6, "Use SetAlphaForTargetType on the ULandscapeEditLayerBase object")
 	LANDSCAPE_API void SetLayerAlpha(int32 InLayerIndex, const float InAlpha, bool bInHeightmap);
+	UE_DEPRECATED(5.6, "Use GetAlphaForTargetType on the ULandscapeEditLayerBase object")
 	LANDSCAPE_API float GetLayerAlpha(int32 InLayerIndex, bool bInHeightmap) const;
+	UE_DEPRECATED(5.6, "Unused : The edit layer class clamps the alpha already")
 	LANDSCAPE_API float GetClampedLayerAlpha(float InAlpha, bool bInHeightmap) const;
+	UE_DEPRECATED(5.6, "Use SetVisibility on the ULandscapeEditLayerBase object")
 	LANDSCAPE_API void SetLayerVisibility(int32 InLayerIndex, bool bInVisible, bool bInForIntermediateRender = false);
+	UE_DEPRECATED(5.6, "Use SetLocked on the ULandscapeEditLayerBase object")
 	LANDSCAPE_API void SetLayerLocked(int32 InLayerIndex, bool bLocked);
+	UE_DEPRECATED(5.6, "Unused: Override the GetBlendMode virtual method in ULandscapeEditLayerBase instead")
 	LANDSCAPE_API void SetLayerBlendMode(int32 InLayerIndex, ELandscapeBlendMode InBlendMode);
-	LANDSCAPE_API uint8 GetLayerCount() const;
-	LANDSCAPE_API TArrayView<const FLandscapeLayer> GetLayers() const { return MakeArrayView(LandscapeEditLayers); }
+
+	// FLandscapeLayer accessors : only the const version is provided because we don't want to let them be mutated freely without the landscape being aware
+	LANDSCAPE_API TArrayView<const FLandscapeLayer> GetLayersConst() const;
 	LANDSCAPE_API const FLandscapeLayer* GetLayerConst(int32 InLayerIndex) const;
 	LANDSCAPE_API const FLandscapeLayer* GetLayerConst(const FGuid& InLayerGuid) const;
 	LANDSCAPE_API const FLandscapeLayer* GetLayerConst(const FName& InLayerName) const;
+	LANDSCAPE_API const FLandscapeLayer* FindLayerOfTypeConst(const TSubclassOf<ULandscapeEditLayerBase>& InLayerClass) const;
+	LANDSCAPE_API TArray<const FLandscapeLayer*> GetLayersOfTypeConst(const TSubclassOf<ULandscapeEditLayerBase>& InLayerClass) const;
 	LANDSCAPE_API int32 GetLayerIndex(const FGuid& InLayerGuid) const;
 	LANDSCAPE_API int32 GetLayerIndex(FName InLayerName) const;
 
+	// ULandscapeEditLayerBase accessors : both const and non-const versions are provided because the landscape listens to data changes on the edit layer object 
+	//  and can therefore react to any change
+	LANDSCAPE_API const ULandscapeEditLayerBase* GetEditLayerConst(int32 InLayerIndex) const;
+	LANDSCAPE_API ULandscapeEditLayerBase* GetEditLayer(int32 InLayerIndex) const;
+	LANDSCAPE_API const ULandscapeEditLayerBase* GetEditLayerConst(const FGuid& InLayerGuid) const;
+	LANDSCAPE_API ULandscapeEditLayerBase* GetEditLayer(const FGuid& InLayerGuid) const;
+	LANDSCAPE_API const ULandscapeEditLayerBase* GetEditLayerConst(const FName& InLayerName) const;
+	LANDSCAPE_API ULandscapeEditLayerBase* GetEditLayer(const FName& InLayerName) const;
+
+	LANDSCAPE_API const ULandscapeEditLayerBase* FindEditLayerOfTypeConst(const TSubclassOf<ULandscapeEditLayerBase>& InLayerClass) const;
+	LANDSCAPE_API ULandscapeEditLayerBase* FindEditLayerOfType(const TSubclassOf<ULandscapeEditLayerBase>& InLayerClass) const;
+	LANDSCAPE_API TArray<const ULandscapeEditLayerBase*> GetEditLayersOfTypeConst(const TSubclassOf<ULandscapeEditLayerBase>& InLayerClass) const;
+	LANDSCAPE_API TArray<ULandscapeEditLayerBase*> GetEditLayersOfType(const TSubclassOf<ULandscapeEditLayerBase>& InLayerClass) const;
+
+	UE_DEPRECATED(5.6, "Use GetLayersConst().Num()")
+	LANDSCAPE_API uint8 GetLayerCount() const;
+	UE_DEPRECATED(5.6, "Use GetLayersConst")
+	LANDSCAPE_API TArrayView<const FLandscapeLayer> GetLayers() const { return MakeArrayView(LandscapeEditLayers); }
 	UE_DEPRECATED(5.5, "This has moved to private (GetLayerInternal), use the appropriate setters to mutate the edit layer")
 	LANDSCAPE_API FLandscapeLayer* GetLayer(int32 InLayerIndex) { return nullptr; }
 	UE_DEPRECATED(5.5, "Use GetLayerConst")
@@ -445,9 +465,6 @@ public:
 	 */
 	LANDSCAPE_API void ForEachLayerConst(TFunctionRef<bool(const FLandscapeLayer&)> Fn);
 
-	LANDSCAPE_API const FLandscapeLayer* FindLayerOfType(const TSubclassOf<ULandscapeEditLayerBase>& InLayerClass) const;
-	LANDSCAPE_API TArray<const FLandscapeLayer*> GetLayersOfType(const TSubclassOf<ULandscapeEditLayerBase>& InLayerClass) const;
-
 	LANDSCAPE_API void GetUsedPaintLayers(int32 InLayerIndex, TArray<ULandscapeLayerInfoObject*>& OutUsedLayerInfos) const;
 	LANDSCAPE_API void GetUsedPaintLayers(const FGuid& InLayerGuid, TArray<ULandscapeLayerInfoObject*>& OutUsedLayerInfos) const;
 	LANDSCAPE_API void ClearPaintLayer(int32 InLayerIndex, ULandscapeLayerInfoObject* InLayerInfo);
@@ -474,7 +491,9 @@ public:
 	UE_DEPRECATED(5.5, "Use FindLayerOfType(ULandscapeEditLayerSplines::StaticClass()).Guid == GetEditingLayer() instead")
 	LANDSCAPE_API bool IsEditingLayerReservedForSplines() const;
 
+	UE_DEPRECATED(5.6, "Use ULandscapeEditLayerSplines GetWeightmapLayerAllocationBlend().Find instead")
 	LANDSCAPE_API bool IsLayerBlendSubstractive(int32 InLayerIndex, const TWeakObjectPtr<ULandscapeLayerInfoObject>& InLayerInfoObj) const;
+	UE_DEPRECATED(5.6, "Use ULandscapeEditLayerSplines AddOrUpdateWeightmapAllocationLayerBlend instead")
 	LANDSCAPE_API void SetLayerSubstractiveBlendStatus(int32 InLayerIndex, bool InStatus, const TWeakObjectPtr<ULandscapeLayerInfoObject>& InLayerInfoObj);
 	LANDSCAPE_API void ReplaceLayerSubstractiveBlendStatus(ULandscapeLayerInfoObject* InFromLayerInfo, ULandscapeLayerInfoObject* InToLayerInfo, bool bInShouldDirtyPackage);
 
@@ -619,6 +638,8 @@ private:
 	void OnDirtyHeightmap(FTextureToComponentHelper const& MapHelper, UTexture2D const* InWeightmap, FColor const* InOldData, FColor const* InNewData, int32 InMipLevel);
 
 	static bool IsMaterialResourceCompiled(FMaterialResource* InMaterialResource, bool bInWaitForCompilation);
+
+	void OnEditLayerDataChanged(const FOnLandscapeEditLayerDataChangedParams& InParams);
 #endif // WITH_EDITOR
 
 private:
diff --git a/Engine/Source/Runtime/Landscape/Classes/LandscapeEditLayer.h b/Engine/Source/Runtime/Landscape/Classes/LandscapeEditLayer.h
--- a/Engine/Source/Runtime/Landscape/Classes/LandscapeEditLayer.h
+++ b/Engine/Source/Runtime/Landscape/Classes/LandscapeEditLayer.h
@@ -2,14 +2,37 @@
 
 #pragma once
 
+#include "Landscape.h"
 #include "LandscapeEditLayerRenderer.h"
 #include "LandscapeEditTypes.h"
+#include "Math/Interval.h"
 
 #include "LandscapeEditLayer.generated.h"
 
 enum class ELandscapeToolTargetType : uint8;
-struct FLandscapeLayer;
-class ALandscape;
+
+struct FOnLandscapeEditLayerDataChangedParams
+{
+	LANDSCAPE_API FOnLandscapeEditLayerDataChangedParams(const FPropertyChangedEvent& InPropertyChangedEvent = FPropertyChangedEvent(/*InProperty = */nullptr))
+		: PropertyChangedEvent(InPropertyChangedEvent)
+	{}
+
+	/** Provides some additional context about how data has changed (property, type of change...) */
+	FPropertyChangedEvent PropertyChangedEvent;
+
+	/** Indicates a user-initiated property change */
+	bool bUserTriggered = false;
+
+	/** Indicates the change requires a full landscape update (e.g. parameter affecting heightmap or weightmap...) */
+	bool bRequiresLandscapeUpdate = true;
+
+	/** The delegate is triggered each time a data change is requested, even when the data didn't actually change. This indicates that the 
+	 * was actually modified. This can occur for example when several EPropertyChangeType::Interactive changes are triggered because of the user
+	 * manipulating a slider : this will be followed by a final EPropertyChangeType::ValueSet but when this occurs, the data usually is not actually 
+	 * modified so, to be consistent, we'll still trigger the delegate but indicate that the value didn't actually change, to let the user react appropriately
+	 */
+	bool bHasValueChanged = true;
+};
 
 /** 
 * Base class for all landscape edit layers. By implementing the various virtual functions, we are able to customize the behavior of the edit layer
@@ -82,12 +105,8 @@ public:
 	};
 
 public:
-
-	// TODO: This might be removed once the guid is stored here and subclasses have a way to request landscape updates.
-	// Otherwise, it might be better made private and then we friend ALandscape?
-	void SetBackPointer(ALandscape* Landscape);
-
 	/**
+	 * @param InType : tool target type (Heightmap, Weightmap, Visibility) 
 	 * @return true if the this edit layer has support for the target type (heightmap, weightmap, visibility)
 	 */
 	virtual bool SupportsTargetType(ELandscapeToolTargetType InType) const
@@ -123,35 +142,143 @@ public:
 	virtual bool SupportsCollapsingTo() const
 		PURE_VIRTUAL(ULandscapeEditLayerBase::SupportsCollapsingTo, return true; );
 
+	/**
+	 * @return true if the layer supports blueprint brushes
+	 */
+	virtual bool SupportsBlueprintBrushes() const
+		PURE_VIRTUAL(ULandscapeEditLayerBase::SupportsBlueprintBrushes(), return false; );
+
 	/**
 	* @return the default name to use when creating a new layer of this type
 	*/
 	virtual FString GetDefaultName() const
 		PURE_VIRTUAL(ULandscapeEditLayerBase::GetDefaultName, return FString(); );
 
+#if WITH_EDITOR
+	/**
+	 * @param InType : tool target type (Heightmap, Weightmap, Visibility) 
+	 * @return true if the layer supports alpha for a given target type
+	 */
+	LANDSCAPE_API virtual bool SupportsAlphaForTargetType(ELandscapeToolTargetType InType) const;
+
+	/**
+	 * Sets the alpha value for a given target type
+	 * @param InType : the type of target (Heightmap, Weightmap) on which to set the alpha
+	 * @param InNewValue : the new alpha value
+	 * @param bInModify : true when Modify() needs to be called
+	 * @param InChangeType : type of change (normal, interactive, ...)
+	 */
+	LANDSCAPE_API virtual void SetAlphaForTargetType(ELandscapeToolTargetType InType, float InNewValue, bool bInModify, EPropertyChangeType::Type InChangeType);
+
+	/**
+	 * Gets the alpha value for a given target type
+	 * @param InType : tool target type (Heightmap, Weightmap, Visibility) 
+	 * @return the alpha value for a given target type
+	 */
+	LANDSCAPE_API virtual float GetAlphaForTargetType(ELandscapeToolTargetType InType) const;
+
+	/**
+	* @return the valid alpha value interval for a given target type
+	*/
+	LANDSCAPE_API virtual FFloatInterval GetAlphaRangeForTargetType(ELandscapeToolTargetType InType) const;
+
+	/**
+	 * Sets the layer's Guid value
+	 * @param InGuid : the new Guid value
+	 * @param bInModify : true when Modify() needs to be called
+	 */
+	LANDSCAPE_API virtual void SetGuid(const FGuid& InGuid, bool bInModify);
+
+	LANDSCAPE_API virtual const FGuid& GetGuid() const;
+
+	/**
+	 * Sets the layer's Name value
+	 * @param InName : the new Name value
+	 * @param bInModify : true when Modify() needs to be called
+	 */
+	LANDSCAPE_API virtual void SetName(FName InName, bool bInModify);
+
+	LANDSCAPE_API virtual FName GetName() const;
+
+	/**
+	 * Sets the layer's visibility value
+	 * @param bInVisible : true to set visible, false to set invisible
+	 * @param bInModify : true when Modify() needs to be called
+	 */
+	LANDSCAPE_API virtual void SetVisible(bool bInVisible, bool bInModify);
+
+	LANDSCAPE_API virtual bool IsVisible() const;
+
+	/**
+	 * Sets the layer's locked value 
+	 * @param bInLocked : true to lock layer, false to unlock
+	 * @param bInModify : true when Modify() needs to be called
+	 */
+	LANDSCAPE_API virtual void SetLocked(bool bInLocked, bool bInModify);
+
+	LANDSCAPE_API virtual bool IsLocked() const;
+
+	/**
+	 * @return the layer's BlendMode - LSBM_AdditiveBlend by default
+	 * Marked as UE_INTERNAL to prevent external usage before blend refactor with the introduction of blend groups and premultiplied alpha blending
+	 */
+	UE_INTERNAL LANDSCAPE_API virtual ELandscapeBlendMode GetBlendMode() const;
+
+	/**
+	 * Remove and Copy a weightmap layer allocation for a given LayerInfoObj
+	 * @param InKey : the ULandscapeLayerInfoObj key to be removed
+	 * @param bOutValue: the layer allocation value stored in the map for InKey
+	 * @param bInModify : true when Modify() needs to be called
+	 * Marked as UE_INTERNAL to prevent external usage before blend refactor with the introduction of blend groups and premultiplied alpha blending
+	 */
+	UE_INTERNAL LANDSCAPE_API virtual bool RemoveAndCopyWeightmapAllocationLayerBlend(TObjectPtr<ULandscapeLayerInfoObject> InKey, bool& bOutValue, bool bInModify);
+
+	/**
+	 * Updates the value of an existing layer allocation or adds a weightmap layer allocation if no entry is found
+	 * @param InKey : the ULandscapeLayerInfoObj map key
+	 * @param bInValue: the new allocation value
+	 * @param bInModify : true when Modify() needs to be called
+	 * Marked as UE_INTERNAL to prevent external usage before blend refactor with the introduction of blend groups and premultiplied alpha blending
+	 */
+	UE_INTERNAL LANDSCAPE_API virtual void AddOrUpdateWeightmapAllocationLayerBlend(TObjectPtr<ULandscapeLayerInfoObject> InKey, bool InValue, bool bInModify);
+
+	/**
+	 * @return the layer's weightmap layer allocation blend map. Layers do not have an entry in the map until a user sets the blend mode (selects Subtractive)
+	 * Marked as UE_INTERNAL to prevent external usage before blend refactor with the introduction of blend groups and premultiplied alpha blending
+	 */
+	UE_INTERNAL LANDSCAPE_API virtual const TMap<TObjectPtr<ULandscapeLayerInfoObject>, bool>& GetWeightmapLayerAllocationBlend() const;
+
+	/**
+	 * Sets the layers WeightmapLayerAllocation map
+	 * @param InWeightmapLayerAllocationBlend : the new WeightmapLayerAllocation map
+	 * @param bInModify : true when Modify() needs to be called
+	 * Marked as UE_INTERNAL to prevent external usage before blend refactor with the introduction of blend groups and premultiplied alpha blending
+	 */
+	UE_INTERNAL LANDSCAPE_API virtual void SetWeightmapLayerAllocationBlend(const TMap<TObjectPtr<ULandscapeLayerInfoObject>, bool>& InWeightmapLayerAllocationBlend, bool bInModify);
+
+#endif // WITH_EDITOR
+
 	/**
 	* @return a list of actions that can be triggered on this edit layer
 	*/
-	virtual TArray<FEditLayerAction> GetActions() const { return {}; }
+	virtual TArray<FEditLayerAction> GetActions() const 
+	{ 
+		return {}; 
+	}
 
 	/**
 	* @return a list of UObjects that this layers needs in order to render properly. This is necessary to avoid trying to render a layer while some of its
 	*  resources are not fully ready. 
 	*  These can be UTexture (requires all mips to be fully loaded) or UMaterialInterface (requires shader maps to be fully compiled)
 	*/
-	virtual void GetRenderDependencies(TSet<UObject*>& OutDependencies) const { }
-
-	// TODO [jonathan.bard] : Remove : temporary method to give the edit layer an opportunity to change some settings on the old struct (FLandscapeLayer) upon creation. 
-	//  The better way would be to move most of the settings to ULandscapeEditLayer(Base?) and expose a "property changed" event that both the UI and runtime code 
-	//  could listen to in order to generically react to any change on the layer settings, including on the derived class (e.g. if there were some splines layer-specific 
-	//  settings, on property change, the event would be triggered and the landscape layers would be updated as a result), instead of providing ad-hoc functions on 
-	//  ALandscape, like we do currently (e.g. SetLayerAlpha, SetLayerVisibility, SetLayerName, etc.) on both the runtime code (ALandscape) and the UI code (FEdModeLandscape)
-	virtual void OnLayerCreated(FLandscapeLayer& Layer) {}
+	virtual void GetRenderDependencies(TSet<UObject*>& OutDependencies) const 
+	{}
 
 	// Called by landscape after removing this layer from its list so that the layer can do
 	// any cleanup that it might need to do.
 	// TODO: Should this be protected and then we friend ALandscape?
-	virtual void OnLayerRemoved() {}
+	virtual void OnLayerRemoved() 
+	{}
 
 #if WITH_EDITOR
 	//~ Begin IEditLayerRendererProvider implementation
@@ -159,21 +286,85 @@ public:
 	//  they would like to provide additional renderers.
 	virtual TArray<UE::Landscape::EditLayers::FEditLayerRendererState> GetEditLayerRendererStates(const UE::Landscape::EditLayers::FMergeContext* InMergeContext) {  return {};  };
 	//~ End IEditLayerRendererProvider implementation
-#endif
 
-	// UObject
-	LANDSCAPE_API virtual void PostLoad() override;
+
+	/** Delegate triggered whenever a change occurred on the edit layer's data */
+	DECLARE_MULTICAST_DELEGATE_OneParam(FOnLayerDataChanged, const FOnLandscapeEditLayerDataChangedParams& /*InParams*/);
+	FOnLayerDataChanged::RegistrationType& OnLayerDataChanged() const
+	{
+		return OnLayerDataChangedDelegate;
+	}
+
+	/**
+	* @return the mask of the target types currently enabled on this edit layer 
+	*/
+	LANDSCAPE_API virtual ELandscapeToolTargetTypeFlags GetEnabledTargetTypeMask() const;
+#endif // WITH_EDITOR
 
 protected:
 #if WITH_EDITOR
-	// TODO [jonathan.bard] remove this dependency to FLandscapeLayer ASAP (once all data from there has moved to the ULandscapeEditLayer class) :
-	const FLandscapeLayer* GetOwningLayer() const;
+	// Begin UObject implementation
+	LANDSCAPE_API virtual void PostLoad() override;
+	LANDSCAPE_API virtual void PostEditChangeProperty(FPropertyChangedEvent& PropertyChangedEvent) override;
+	LANDSCAPE_API virtual void PostEditUndo() override;
+	LANDSCAPE_API virtual bool CanEditChange(const FProperty* InProperty) const override;
+	// End UObject implementation
 #endif // WITH_EDITOR
 
+#if WITH_EDITORONLY_DATA
+	// Setters for UPROPERTY (should be used for blueprint setters eventually) :  
+	void SetGuidInternal(const FGuid& InGuid);
+	void SetNameInternal(FName InName);
+	void SetVisibleInternal(bool bInVisible);
+	void SetLockedInternal(bool bInLocked);
+	void SetHeightmapAlphaInternal(float InNewValue);
+	void SetWeightmapAlphaInternal(float InNewValue);
+	// Marked as UE_INTERNAL to prevent external usage before blend refactor with the introduction of blend groups and premultiplied alpha blending
+	UE_INTERNAL void SetWeightmapLayerAllocationBlendInternal(const TMap<TObjectPtr<ULandscapeLayerInfoObject>, bool>& InWeightmapLayerAllocationBlend);
+#endif // WITH_EDITORONLY_DATA
+
+#if WITH_EDITOR
+	/** Internal function that returns a reference to the alpha value corresponding to a given target type (HeightmapAlpha, WeightmapAlpha, ...) */
+	float& GetAlphaForTargetTypeRef(ELandscapeToolTargetType InType);
+	/** Internal function that returns the alpha FProperty corresponding to a given target type (HeightmapAlpha, WeightmapAlpha, ...)  */
+	FProperty* GetAlphaPropertyForTargetType(ELandscapeToolTargetType InType) const;
+	/** Internal function meant to be called whenever the edit layer data changes, broadcast the OnLayerDataChanged event */
+	void BroadcastOnLayerDataChanged(FName InPropertyName, bool bInUserTriggered, bool bRequiresLandscapeUpdate, bool bInHasValueChanged, EPropertyChangeType::Type InChangeType);
+#endif // WITH_EDITOR
 protected:
+#if WITH_EDITORONLY_DATA
+	UPROPERTY(Setter = "SetGuidInternal", meta = (EditCondition = "!bLocked"))
+	FGuid Guid = FGuid::NewGuid();
+
+	UPROPERTY(Setter = "SetNameInternal", meta = (EditCondition = "!bLocked"))
+	FName LayerName = NAME_None;
+
+	UPROPERTY(Category = "Edit Layer", EditAnywhere, Setter = "SetVisibleInternal", meta = (EditCondition = "!bLocked"))
+	bool bVisible = true;
+
+	UPROPERTY(Category = "Edit Layer", EditAnywhere, Setter = "SetLockedInternal")
+	bool bLocked = false;
+
+	UPROPERTY(Category = "Edit Layer", EditAnywhere, Setter = "SetHeightmapAlphaInternal", meta = (UIMin = "-1.0", UIMax = "1.0", ClampMin = "-1.0", ClampMax = "1.0", EditCondition = "!bLocked"))
+	float HeightmapAlpha = 1.0f;
+
+	UPROPERTY(Category = "Edit Layer", EditAnywhere, Setter = "SetWeightmapAlphaInternal", meta = (UIMin = "0.0", UIMax = "1.0", ClampMin = "0.0", ClampMax = "1.0", EditCondition = "!bLocked"))
+	float WeightmapAlpha = 1.0f;
+
 	// TODO: This might be removed once more things are moved from FLandscapeLayer to ULandscapeLayer
 	UPROPERTY()
 	TWeakObjectPtr<ALandscape> OwningLandscape;
+
+private:
+	// TODO: This might be removed once the guid is stored here and subclasses have a way to request landscape updates.
+	friend class ALandscape;
+	void SetBackPointer(ALandscape* Landscape);
+
+	mutable FOnLayerDataChanged OnLayerDataChangedDelegate;
+
+	UPROPERTY(Setter = "SetWeightmapLayerAllocationBlendInternal", meta = (EditCondition = "!bLocked"))
+	TMap<TObjectPtr<ULandscapeLayerInfoObject>, bool> WeightmapLayerAllocationBlend; // True -> Substractive, False -> Add
+#endif //WITH_EDITORONLY_DATA
 };
 
 /** 
@@ -189,6 +380,7 @@ public:
 	// Begin ULandscapeEditLayerBase implementation
 	virtual bool NeedsPersistentTextures() const override { return true; };
 	virtual bool SupportsCollapsingTo() const override { return true; } // If the layer has persistent textures, it can be collapsed to another layer (one that supports being collapsed away, that is)
+	virtual bool SupportsBlueprintBrushes() const override { return false; }
 	// End ULandscapeEditLayerBase implementation
 
 #if WITH_EDITOR
@@ -223,7 +415,8 @@ public:
 	virtual bool SupportsTargetType(ELandscapeToolTargetType InType) const override;
 	virtual bool SupportsEditingTools() const override { return true; }
 	virtual bool SupportsMultiple() const override { return true; }
-	virtual bool SupportsBeingCollapsedAway() const override { return true; }
+	virtual bool SupportsBeingCollapsedAway() const override { return true; } 
+	virtual bool SupportsBlueprintBrushes() const override { return true; }
 	virtual FString GetDefaultName() const { return TEXT("Layer"); }
 	// End ULandscapeEditLayerBase implementation
 
@@ -244,6 +437,7 @@ public:
 	virtual bool SupportsEditingTools() const override { return false; } // procedural layers cannot be edited through standard editing tools
 	virtual bool SupportsCollapsingTo() const override { return false; } // for now, don't support collapsing to a layer underneath for a procedural layer (this may become unneeded if we make the collapse happen on the GPU)
 	virtual bool SupportsBeingCollapsedAway() const override { return false; } // this is a procedural and therefore cannot be collapsed 
+	virtual bool SupportsBlueprintBrushes() const override { return false; }
 	// End ULandscapeEditLayerBase implementation
 };
 
@@ -263,10 +457,14 @@ public:
 	virtual bool SupportsMultiple() const override { return false; } // only one layer of this type is allowed
 	virtual bool SupportsBeingCollapsedAway() const override { return false; } // this is a procedural and therefore cannot be collapsed 
 	virtual FString GetDefaultName() const override { return TEXT("Splines"); }
-	virtual void OnLayerCreated(FLandscapeLayer& Layer) override;
 	virtual TArray<FEditLayerAction> GetActions() const override;
+#if WITH_EDITOR
+	virtual ELandscapeBlendMode GetBlendMode() const override { return ELandscapeBlendMode::LSBM_AlphaBlend; };
+	bool SupportsAlphaForTargetType(ELandscapeToolTargetType InType) const override { return false; };
+	virtual float GetAlphaForTargetType(ELandscapeToolTargetType InType) const override { return 1.0f; };
+	virtual void SetAlphaForTargetType(ELandscapeToolTargetType InType, float InNewValue, bool bInModify, EPropertyChangeType::Type InChangeType) override { /* do nothing */ };
 	// End ULandscapeEditLayerBase implementation
-
+#endif // WITH_EDITOR
 protected:
 
 };
diff --git a/Engine/Source/Runtime/Landscape/Private/Landscape.cpp b/Engine/Source/Runtime/Landscape/Private/Landscape.cpp
--- a/Engine/Source/Runtime/Landscape/Private/Landscape.cpp
+++ b/Engine/Source/Runtime/Landscape/Private/Landscape.cpp
@@ -1408,11 +1408,11 @@ PRAGMA_ENABLE_DEPRECATION_WARNINGS
 		}
 
 		// Fixup missing/mismatching edit layer names :
-		if (const FLandscapeLayer* EditLayer = GetLandscapeActor() ? GetLandscapeActor()->GetLayerConst(ItPair.Key) : nullptr)
+		if (const FLandscapeLayer* LandscapeLayer = GetLandscapeActor() ? GetLandscapeActor()->GetLayerConst(ItPair.Key) : nullptr)
 		{
-			if (LayerComponentData.DebugName != EditLayer->Name)
+			if (LayerComponentData.DebugName != LandscapeLayer->EditLayer->GetName())
 			{
-				LayerComponentData.DebugName = EditLayer->Name;
+				LayerComponentData.DebugName = LandscapeLayer->EditLayer->GetName();
 			}
 		}
 	}
@@ -2566,8 +2566,8 @@ void ULandscapeComponent::AddDefaultLayerData(const FGuid& InLayerGuid, const TA
 
 	if (LayerData == nullptr || !LayerData->IsInitialized())
 	{
-		const FLandscapeLayer* EditLayer = GetLandscapeActor() ? GetLandscapeActor()->GetLayerConst(InLayerGuid) : nullptr;
-		FLandscapeLayerComponentData NewData(EditLayer ? EditLayer->Name : FName());
+		const FLandscapeLayer* LandscapeLayer = GetLandscapeActor() ? GetLandscapeActor()->GetLayerConst(InLayerGuid) : nullptr;
+		FLandscapeLayerComponentData NewData(LandscapeLayer ? LandscapeLayer->EditLayer->GetName() : FName());
 
 		// Setup Heightmap data
 		UTexture2D** LayerHeightmap = InOutCreatedHeightmapTextures.Find(ComponentHeightmap);
@@ -2953,13 +2953,13 @@ PRAGMA_DISABLE_DEPRECATION_WARNINGS
 		for (FLandscapeLayer& Layer : LandscapeEditLayers)
 		{
 			UClass* EditLayerClass = ULandscapeEditLayer::StaticClass();
-			if (Layer.Guid == LandscapeSplinesTargetLayerGuid_DEPRECATED)
+
+			if (Layer.Guid_DEPRECATED == LandscapeSplinesTargetLayerGuid_DEPRECATED)
 			{
 				EditLayerClass = ULandscapeEditLayerSplines::StaticClass();
 			}
 			check(Layer.EditLayer == nullptr);
-			Layer.EditLayer = NewObject<ULandscapeEditLayerBase>(this, EditLayerClass, MakeUniqueObjectName(this, EditLayerClass));
-			Layer.EditLayer->OnLayerCreated(Layer);
+			Layer.EditLayer = NewObject<ULandscapeEditLayerBase>(this, EditLayerClass, MakeUniqueObjectName(this, EditLayerClass), RF_Transactional);
 		}
 
 		// Empty the old property now that we've moved them over, else we'll accidentally keep references to brushes etc.
@@ -2972,19 +2972,34 @@ PRAGMA_ENABLE_DEPRECATION_WARNINGS
 		FLandscapeLayer& Layer = LandscapeEditLayers[LayerIndex];
 		if (Layer.EditLayer != nullptr)
 		{
-			// For now, only Layer reserved for Landscape Spline uses AlphaBlend
-			Layer.BlendMode = Layer.EditLayer->IsA<ULandscapeEditLayerSplines>() ? LSBM_AlphaBlend : LSBM_AdditiveBlend;
+			if (GetLinkerCustomVersion(FFortniteMainBranchObjectVersion::GUID) < FFortniteMainBranchObjectVersion::MigrateLandscapeEditLayerProperties)
+			{
+				// Set Owning Landscape before other fields so Setter checks succeed
+				Layer.EditLayer->SetBackPointer(this);
+
+PRAGMA_DISABLE_DEPRECATION_WARNINGS
+				Layer.EditLayer->SetFlags(RF_Transactional); // Bug fix, this flag was forgotten in the initial version
+				Layer.EditLayer->SetGuid(Layer.Guid_DEPRECATED, /*bInModify = */false);
+				Layer.EditLayer->SetName(Layer.Name_DEPRECATED, /*bInModify = */false);
+				Layer.EditLayer->SetVisible(Layer.bVisible_DEPRECATED, /*bInModify = */false);
+				Layer.EditLayer->SetLocked(Layer.bLocked_DEPRECATED, /*bInModify = */false);
+				Layer.EditLayer->SetAlphaForTargetType(ELandscapeToolTargetType::Heightmap, Layer.HeightmapAlpha_DEPRECATED, /*bInModify = */false, EPropertyChangeType::ValueSet);
+				Layer.EditLayer->SetAlphaForTargetType(ELandscapeToolTargetType::Weightmap, Layer.WeightmapAlpha_DEPRECATED, /*bInModify = */false, EPropertyChangeType::ValueSet);
+				Layer.EditLayer->SetWeightmapLayerAllocationBlend(Layer.WeightmapLayerAllocationBlend_DEPRECATED, /*bInModify = */false);
+PRAGMA_ENABLE_DEPRECATION_WARNINGS
+			}
+
+			// Register to data change events on the edit layer so that we can update the landscape accordingly : 
+			Layer.EditLayer->OnLayerDataChanged().AddUObject(this, &ALandscape::OnEditLayerDataChanged);
 
 			for (FLandscapeLayerBrush& Brush : Layer.Brushes)
 			{
 				Brush.SetOwner(this);
 			}
-
-			Layer.EditLayer->SetBackPointer(this);
 		}
 		else
 		{
-			UE_LOG(LogLandscape, Error, TEXT("Couldn't load edit layer object associated with layer %s for landscape %s. This may happen when the edit layer class cannot be found (for example, when a plugin is removed from the project). The layer will be deleted."), *Layer.Name.ToString(), *GetFullName());
+			UE_LOG(LogLandscape, Error, TEXT("Couldn't load edit layer object associated with layer %s for landscape %s. This may happen when the edit layer class cannot be found (for example, when a plugin is removed from the project). The layer will be deleted."), *Layer.EditLayer->GetName().ToString(), *GetFullName());
 			ensure(DeleteLayer(LayerIndex));
 		}
 	}
@@ -3797,7 +3812,7 @@ void ALandscapeProxy::PreSave(FObjectPreSaveContext ObjectSaveContext)
 			{
 				if (const FLandscapeLayer* EditLayer = Landscape->GetLayerConst(LayerGuid))
 				{
-					LayerData.DebugName = EditLayer->Name;
+					LayerData.DebugName = EditLayer->EditLayer->GetName();
 				}
 			});
 		}
@@ -4996,12 +5011,12 @@ void ALandscapeProxy::FixupSharedData(ALandscape* Landscape, const bool bMapChec
 		bool bUpdated = !SynchronizedProperties.IsEmpty();
 
 		TSet<FGuid> LayerGuids;
-		Algo::Transform(Landscape->GetLayers(), LayerGuids, [](const FLandscapeLayer& Layer) { return Layer.Guid; });
+		Algo::Transform(Landscape->GetLayersConst(), LayerGuids, [](const FLandscapeLayer& Layer) { return Layer.EditLayer->GetGuid(); });
 		bUpdated |= RemoveObsoleteLayers(LayerGuids);
 
-		for (const FLandscapeLayer& Layer : Landscape->GetLayers())
+		for (const FLandscapeLayer& Layer : Landscape->GetLayersConst())
 		{
-			bUpdated |= AddLayer(Layer.Guid);
+			bUpdated |= AddLayer(Layer.EditLayer->GetGuid());
 		}
 
 		if (bUpdated)
diff --git a/Engine/Source/Runtime/Landscape/Private/LandscapeBlueprintSupport.cpp b/Engine/Source/Runtime/Landscape/Private/LandscapeBlueprintSupport.cpp
--- a/Engine/Source/Runtime/Landscape/Private/LandscapeBlueprintSupport.cpp
+++ b/Engine/Source/Runtime/Landscape/Private/LandscapeBlueprintSupport.cpp
@@ -12,6 +12,7 @@
 #include "Components/SplineComponent.h"
 #include "LandscapeComponent.h"
 #include "Landscape.h"
+#include "LandscapeEditLayer.h"
 #include "LandscapePrivate.h"
 #include "Materials/MaterialInstanceDynamic.h"
 #include "LandscapeProxy.h"
@@ -30,7 +31,7 @@ void ALandscapeProxy::EditorApplySpline(USplineComponent* InSplineComponent, flo
 				return;
 			}
 
-			FScopedSetLandscapeEditingLayer Scope(Landscape, Layer ? Layer->Guid : FGuid(), [=] { Landscape->RequestLayersContentUpdate(ELandscapeLayerUpdateMode::Update_All); });
+			FScopedSetLandscapeEditingLayer Scope(Landscape, Layer ? Layer->EditLayer->GetGuid() : FGuid(), [=] { Landscape->RequestLayersContentUpdate(ELandscapeLayerUpdateMode::Update_All); });
 
 			TArray<FLandscapeSplineInterpPoint> Points;
 			LandscapeSplineRaster::FPointifyFalloffs Falloffs(StartSideFalloff, EndSideFalloff);
diff --git a/Engine/Source/Runtime/Landscape/Private/LandscapeConfigHelper.cpp b/Engine/Source/Runtime/Landscape/Private/LandscapeConfigHelper.cpp
--- a/Engine/Source/Runtime/Landscape/Private/LandscapeConfigHelper.cpp
+++ b/Engine/Source/Runtime/Landscape/Private/LandscapeConfigHelper.cpp
@@ -380,7 +380,7 @@ void FLandscapeConfigHelper::CopyRegionToComponent(ULandscapeInfo* InLandscapeIn
 		}
 		InComponent->InitHeightmapData(HeightInitData, false);
 
-		for (const FLandscapeLayer& LandscapeLayer : InLandscapeInfo->LandscapeActor->GetLayers())
+		for (const FLandscapeLayer& LandscapeLayer : InLandscapeInfo->LandscapeActor->GetLayersConst())
 		{		
 			check(LandscapeLayer.EditLayer != nullptr);
 			if (!LandscapeLayer.EditLayer->NeedsPersistentTextures())
@@ -389,7 +389,7 @@ void FLandscapeConfigHelper::CopyRegionToComponent(ULandscapeInfo* InLandscapeIn
 			}
 
 			TMap<UTexture2D*, UTexture2D*> CreatedTextures;
-			InComponent->AddDefaultLayerData(LandscapeLayer.Guid, { InComponent }, CreatedTextures);
+			InComponent->AddDefaultLayerData(LandscapeLayer.EditLayer->GetGuid(), { InComponent }, CreatedTextures);
 
 			// No need to copy the persistent textures if they are not manually edited (e.g. fully procedural layers with persistent textures : e.g. the splines layer) :
 			if (!LandscapeLayer.EditLayer->SupportsEditingTools())
@@ -398,10 +398,10 @@ void FLandscapeConfigHelper::CopyRegionToComponent(ULandscapeInfo* InLandscapeIn
 			}
 
 			SrcWeightData.Empty();
-			ExtractLandscapeData(InLandscapeInfo, InRegion, LandscapeLayer.Guid, SrcHeightData, SrcWeightData);
+			ExtractLandscapeData(InLandscapeInfo, InRegion, LandscapeLayer.EditLayer->GetGuid(), SrcHeightData, SrcWeightData);
 			CopyData(SrcHeightData, NewHeightData, InRegion, NewRegion, bInResample);
 
-			FScopedSetLandscapeEditingLayer LayerScope(NewLandscapeInfo->LandscapeActor.Get(), LandscapeLayer.Guid);
+			FScopedSetLandscapeEditingLayer LayerScope(NewLandscapeInfo->LandscapeActor.Get(), LandscapeLayer.EditLayer->GetGuid());
 			{
 				FLandscapeEditDataInterface LandscapeEdit(NewLandscapeInfo);
 				LandscapeEdit.SetHeightData(NewRegion.Min.X, NewRegion.Min.Y, NewRegion.Max.X, NewRegion.Max.Y, NewHeightData.GetData(), 0, false);
@@ -573,7 +573,7 @@ ULandscapeInfo* FLandscapeConfigHelper::ChangeConfiguration(ULandscapeInfo* InLa
 	if (OldLandscape->HasLayersContent())
 	{
 		NewLandscape->bCanHaveLayersContent = true;
-		for (const FLandscapeLayer& OldLayer : OldLandscape->GetLayers())
+		for (const FLandscapeLayer& OldLayer : OldLandscape->GetLayersConst())
 		{
 			const FLandscapeLayer* NewLayer = NewLandscape->DuplicateLayerAndMoveBrushes(OldLayer);
 			check((NewLayer == nullptr) || (NewLayer->EditLayer != nullptr)); // it's possible DuplicateLayerAndMoveBrushes fails (e.g. max number of layers reached), but if not, we should always have an EditLayer
diff --git a/Engine/Source/Runtime/Landscape/Private/LandscapeEdit.cpp b/Engine/Source/Runtime/Landscape/Private/LandscapeEdit.cpp
--- a/Engine/Source/Runtime/Landscape/Private/LandscapeEdit.cpp
+++ b/Engine/Source/Runtime/Landscape/Private/LandscapeEdit.cpp
@@ -15,6 +15,7 @@ LandscapeEdit.cpp: Landscape editing
 #include "UObject/Package.h"
 #include "Misc/PackageName.h"
 #include "Landscape.h"
+#include "LandscapeEditLayer.h"
 #include "LandscapeEditReadback.h"
 #include "LandscapeStreamingProxy.h"
 #include "LandscapeInfo.h"
@@ -3901,7 +3902,7 @@ void ALandscapeProxy::Import(const FGuid& InGuid, int32 InMinX, int32 InMinY, in
 		// Create the default layer first
 		ALandscape* LandscapeActor = GetLandscapeActor();
 		check(LandscapeActor != nullptr);
-		if (LandscapeActor->GetLayerCount() == 0 && InImportLayers.IsEmpty())
+		if (LandscapeActor->GetLayersConst().IsEmpty() && InImportLayers.IsEmpty())
 		{
 			LandscapeActor->CreateDefaultLayer();
 		}
@@ -3931,8 +3932,8 @@ void ALandscapeProxy::Import(const FGuid& InGuid, int32 InMinX, int32 InMinY, in
 					check(NewLayer->EditLayer != nullptr); // it's possible DuplicateLayerAndMoveBrushes fails (e.g. max number of layers reached), but if not, we should always have an EditLayer
 
 					FLayerImportSettings ImportSettings;
-					ImportSettings.SourceLayerGuid = OldLayer.Guid;
-					ImportSettings.DestinationLayerGuid = NewLayer->Guid;
+					ImportSettings.SourceLayerGuid = OldLayer.EditLayer->GetGuid();
+					ImportSettings.DestinationLayerGuid = NewLayer->EditLayer->GetGuid();
 					LayerImportSettings.Add(ImportSettings);
 				}
 			}
@@ -3953,7 +3954,7 @@ void ALandscapeProxy::Import(const FGuid& InGuid, int32 InMinX, int32 InMinY, in
 
 			FLayerImportSettings ImportSettings;
 			ImportSettings.SourceLayerGuid = FinalLayerGuid;
-			ImportSettings.DestinationLayerGuid = DefaultLayer->Guid;
+			ImportSettings.DestinationLayerGuid = DefaultLayer->EditLayer->GetGuid();
 			LayerImportSettings.Add(ImportSettings);
 
 			ComponentsToProcess.Append(ToRawPtrTArrayUnsafe(LandscapeComponents));
diff --git a/Engine/Source/Runtime/Landscape/Private/LandscapeEditInterface.cpp b/Engine/Source/Runtime/Landscape/Private/LandscapeEditInterface.cpp
--- a/Engine/Source/Runtime/Landscape/Private/LandscapeEditInterface.cpp
+++ b/Engine/Source/Runtime/Landscape/Private/LandscapeEditInterface.cpp
@@ -11,6 +11,7 @@ LandscapeEditInterface.cpp: Landscape editing interface
 #include "RenderingThread.h"
 #include "LandscapeProxy.h"
 #include "Landscape.h"
+#include "LandscapeEditLayer.h"
 #include "LandscapeInfo.h"
 #include "LandscapeComponent.h"
 #include "LandscapeLayerInfoObject.h"
@@ -2363,7 +2364,7 @@ void FLandscapeEditDataInterface::ReplaceLayer(ULandscapeLayerInfoObject* FromLa
 		LandscapeInfo->LandscapeActor->ReplaceLayerSubstractiveBlendStatus(FromLayerInfo, ToLayerInfo, GetShouldDirtyPackage());
 		LandscapeInfo->LandscapeActor->ForEachLayerConst([DoReplace](const FLandscapeLayer& CurrentLayer)
 		{
-			DoReplace(CurrentLayer.Guid);
+			DoReplace(CurrentLayer.EditLayer->GetGuid());
 			return true;
 		});
 		LandscapeInfo->LandscapeActor->RequestLayersContentUpdateForceAll(ELandscapeLayerUpdateMode::Update_Weightmap_All);
diff --git a/Engine/Source/Runtime/Landscape/Private/LandscapeEditLayer.cpp b/Engine/Source/Runtime/Landscape/Private/LandscapeEditLayer.cpp
--- a/Engine/Source/Runtime/Landscape/Private/LandscapeEditLayer.cpp
+++ b/Engine/Source/Runtime/Landscape/Private/LandscapeEditLayer.cpp
@@ -6,55 +6,421 @@
 
 #define LOCTEXT_NAMESPACE "LandscapeEditLayer"
 
-void ULandscapeEditLayerBase::SetBackPointer(ALandscape* Landscape)
+// ----------------------------------------------------------------------------------
+
+#if WITH_EDITOR
+bool ULandscapeEditLayerBase::SupportsAlphaForTargetType(ELandscapeToolTargetType InType) const
 {
-	OwningLandscape = Landscape;
+	return (InType == ELandscapeToolTargetType::Heightmap) || (InType == ELandscapeToolTargetType::Weightmap);
+}
+
+void ULandscapeEditLayerBase::SetAlphaForTargetType(ELandscapeToolTargetType InType, float InNewValue, bool bInModify, EPropertyChangeType::Type InChangeType)
+{
+	check(SupportsAlphaForTargetType(InType));
+
+	const FFloatInterval AlphaInterval = GetAlphaRangeForTargetType(InType);
+	const float ClampedNewValue = FMath::Clamp(InNewValue, AlphaInterval.Min, AlphaInterval.Max);
+	float& AlphaValueRef = GetAlphaForTargetTypeRef(InType);
+	const bool bHasValueChanged = (AlphaValueRef != ClampedNewValue);
+	if (bHasValueChanged)
+	{
+		if (bInModify)
+		{
+			Modify();
+		}
+
+		AlphaValueRef = ClampedNewValue;
+	}
+
+	FProperty* AlphaProperty = GetAlphaPropertyForTargetType(InType);
+	check(AlphaProperty != nullptr);
+	BroadcastOnLayerDataChanged(FName(AlphaProperty->GetName()), /*bInUserTriggered = */true, /*bInRequiresLandscapeUpdate =*/ true, bHasValueChanged, InChangeType);
+}
+
+float ULandscapeEditLayerBase::GetAlphaForTargetType(ELandscapeToolTargetType InType) const
+{
+	return const_cast<ULandscapeEditLayerBase*>(this)->GetAlphaForTargetTypeRef(InType);
+}
+
+float& ULandscapeEditLayerBase::GetAlphaForTargetTypeRef(ELandscapeToolTargetType InType) 
+{
+	switch (InType)
+	{
+	case ELandscapeToolTargetType::Heightmap: 
+	{
+		return HeightmapAlpha;
+	}
+	case ELandscapeToolTargetType::Weightmap: 
+	{
+		return WeightmapAlpha;
+	}
+	default:
+	{
+		static float DefaultValue = 1.0f;
+		return DefaultValue;
+	}
+	}
+}
+
+FProperty* ULandscapeEditLayerBase::GetAlphaPropertyForTargetType(ELandscapeToolTargetType InType) const
+{
+	switch (InType)
+	{
+	case ELandscapeToolTargetType::Heightmap:
+	{
+		return FindFProperty<FProperty>(StaticClass(), GET_MEMBER_NAME_CHECKED(ULandscapeEditLayerBase, HeightmapAlpha));
+	}
+	case ELandscapeToolTargetType::Weightmap:
+	{
+		return FindFProperty<FProperty>(StaticClass(), GET_MEMBER_NAME_CHECKED(ULandscapeEditLayerBase, WeightmapAlpha));
+	}
+	default:
+	{
+		return nullptr;
+	}
+	}
+}
+
+FFloatInterval ULandscapeEditLayerBase::GetAlphaRangeForTargetType(ELandscapeToolTargetType InType) const
+{
+	switch (InType)
+	{
+	case ELandscapeToolTargetType::Heightmap:
+	{
+		return FFloatInterval(-1.0f, 1.0f);
+	}
+	case ELandscapeToolTargetType::Weightmap:
+	{
+		return FFloatInterval(0.0f, 1.0f);
+	}
+	default:
+	{
+		return FFloatInterval(0.0f, 1.0f);
+	}
+	}
+}
+
+void ULandscapeEditLayerBase::SetGuid(const FGuid& InGuid, bool bInModify)
+{
+	const bool bHasValueChanged = (InGuid != Guid);
+	if (bHasValueChanged)
+	{
+		if (bInModify)
+		{
+			Modify();
+		}
+
+		Guid = InGuid;
+	}
+
+	BroadcastOnLayerDataChanged(GET_MEMBER_NAME_CHECKED(ULandscapeEditLayerBase, Guid), /*bInUserTriggered = */true, /*bInRequiresLandscapeUpdate =*/ true, bHasValueChanged, /*InChangeType = */EPropertyChangeType::ValueSet);
+}
+
+const FGuid& ULandscapeEditLayerBase::GetGuid() const
+{
+	return Guid;
+}
+
+void ULandscapeEditLayerBase::SetName(FName InName, bool bInModify)
+{
+	check(OwningLandscape.IsValid())
+	if (!OwningLandscape->IsLayerNameUnique(InName) || LayerName == InName)
+	{
+		return;
+	}
+
+	const bool bHasValueChanged = (InName != LayerName);
+	if (bHasValueChanged)
+	{
+		if (bInModify)
+		{
+			Modify();
+		}
+
+		LayerName = InName;
+	}
+
+	BroadcastOnLayerDataChanged(GET_MEMBER_NAME_CHECKED(ULandscapeEditLayerBase, LayerName), /*bInUserTriggered = */true, /*bInRequiresLandscapeUpdate =*/ true, bHasValueChanged, /*InChangeType = */EPropertyChangeType::ValueSet);
+}
+
+FName ULandscapeEditLayerBase::GetName() const
+{
+	return LayerName;
+}
+
+void ULandscapeEditLayerBase::SetVisible(bool bInVisible, bool bInModify)
+{
+	const bool bHasValueChanged = (bInVisible != bVisible);
+	if (bHasValueChanged)
+	{
+		if (bInModify)
+		{
+			Modify();
+		}
+
+		bVisible = bInVisible;
+	}
+
+	BroadcastOnLayerDataChanged(GET_MEMBER_NAME_CHECKED(ULandscapeEditLayerBase, bVisible), /*bInUserTriggered = */true, /*bInRequiresLandscapeUpdate =*/ true, bHasValueChanged, /*InChangeType = */EPropertyChangeType::ValueSet);
+}
+
+bool ULandscapeEditLayerBase::IsVisible() const
+{
+	return bVisible;
+}
+
+void ULandscapeEditLayerBase::SetLocked(bool bInLocked, bool bInModify)
+{
+	const bool bHasValueChanged = (bInLocked != bLocked);
+	if (bHasValueChanged)
+	{
+		if (bInModify)
+		{
+			Modify();
+		}
+
+		bLocked = bInLocked;
+	}
+	
+	BroadcastOnLayerDataChanged(GET_MEMBER_NAME_CHECKED(ULandscapeEditLayerBase, bLocked), /*bInUserTriggered = */true, /*bInRequiresLandscapeUpdate =*/ true, bHasValueChanged, /*InChangeType = */EPropertyChangeType::ValueSet);
+}
+
+bool ULandscapeEditLayerBase::IsLocked() const
+{
+	return bLocked;
+}
+
+ELandscapeBlendMode ULandscapeEditLayerBase::GetBlendMode() const
+{
+	return ELandscapeBlendMode::LSBM_AdditiveBlend;
+}
+
+bool ULandscapeEditLayerBase::RemoveAndCopyWeightmapAllocationLayerBlend(TObjectPtr<ULandscapeLayerInfoObject> InKey, bool& bOutValue, bool bInModify)
+{
+	bool bHasValueChanged = false;
+
+	if (WeightmapLayerAllocationBlend.RemoveAndCopyValue(InKey, bOutValue))
+	{
+		Modify(bInModify);
+
+		bHasValueChanged = true;
+	}
+
+	BroadcastOnLayerDataChanged(GET_MEMBER_NAME_CHECKED(ULandscapeEditLayerBase, WeightmapLayerAllocationBlend), /*bInUserTriggered = */true, /*bInRequiresLandscapeUpdate =*/ true, bHasValueChanged, /*InChangeType = */EPropertyChangeType::ValueSet);
+
+	// return true if and only if RemoveAndCopy returns true
+	return bHasValueChanged;
+}
+
+void ULandscapeEditLayerBase::AddOrUpdateWeightmapAllocationLayerBlend(TObjectPtr<ULandscapeLayerInfoObject> InKey, bool bInValue, bool bInModify)
+{ 
+	const bool* bValueExists = WeightmapLayerAllocationBlend.Find(InKey);
+	bool& bFoundValue = WeightmapLayerAllocationBlend.FindOrAdd(InKey, bInValue);
+
+	// changed if existing value has been toggled or a new entry is added to the map
+	const bool bHasValueChanged = ((bValueExists == nullptr) || (bFoundValue != bInValue));
+
+	bFoundValue = bInValue;
+	
+	if (bInModify)
+	{
+		Modify();
+	}
+
+	BroadcastOnLayerDataChanged(GET_MEMBER_NAME_CHECKED(ULandscapeEditLayerBase, WeightmapLayerAllocationBlend), /*bInUserTriggered = */true, /*bInRequiresLandscapeUpdate =*/ true, /*bHasValueChanged =*/ true, /*InChangeType = */EPropertyChangeType::ValueSet);
+}
+
+const TMap<TObjectPtr<ULandscapeLayerInfoObject>, bool>& ULandscapeEditLayerBase::GetWeightmapLayerAllocationBlend() const
+{
+	return WeightmapLayerAllocationBlend;
+}
+
+void ULandscapeEditLayerBase::SetWeightmapLayerAllocationBlend(const TMap<TObjectPtr<ULandscapeLayerInfoObject>, bool>& InWeightmapLayerAllocationBlend, bool bInModify)
+{
+	WeightmapLayerAllocationBlend = InWeightmapLayerAllocationBlend;
+
+	if (bInModify)
+	{
+		Modify();
+	}
+
+	BroadcastOnLayerDataChanged(GET_MEMBER_NAME_CHECKED(ULandscapeEditLayerBase, WeightmapLayerAllocationBlend), /*bInUserTriggered = */true, /*bInRequiresLandscapeUpdate =*/ true, /*bHasValueChanged =*/ true, /*InChangeType = */EPropertyChangeType::ValueSet);
+}
+
+void ULandscapeEditLayerBase::PostEditChangeProperty(FPropertyChangedEvent& PropertyChangedEvent)
+{
+	Super::PostEditChangeProperty(PropertyChangedEvent);
+
+	const FName MemberPropertyName = PropertyChangedEvent.MemberProperty ? PropertyChangedEvent.MemberProperty->GetFName() : NAME_None;
+
+	if (MemberPropertyName == GET_MEMBER_NAME_CHECKED(ULandscapeEditLayerBase, Guid))
+	{
+		SetGuid(Guid, /*bInModify = */true);
+	}
+	else if (MemberPropertyName == GET_MEMBER_NAME_CHECKED(ULandscapeEditLayerBase, LayerName))
+	{
+		SetName(LayerName, /*bInModify = */true);
+	}
+	else if (MemberPropertyName == GET_MEMBER_NAME_CHECKED(ULandscapeEditLayerBase, HeightmapAlpha))
+	{
+		SetAlphaForTargetType(ELandscapeToolTargetType::Heightmap, HeightmapAlpha, /*bInModify = */true, PropertyChangedEvent.ChangeType);
+	}
+	else if (MemberPropertyName == GET_MEMBER_NAME_CHECKED(ULandscapeEditLayerBase, WeightmapAlpha))
+	{
+		SetAlphaForTargetType(ELandscapeToolTargetType::Weightmap, WeightmapAlpha, /*bInModify = */true, PropertyChangedEvent.ChangeType);
+	}
+	else if (MemberPropertyName == GET_MEMBER_NAME_CHECKED(ULandscapeEditLayerBase, bLocked))
+	{
+		SetLocked(bLocked, /*bInModify = */true);
+	}
+	else if (MemberPropertyName == GET_MEMBER_NAME_CHECKED(ULandscapeEditLayerBase, bVisible))
+	{
+		SetVisible(bVisible, /*bInModify = */true);
+	}
+	else if (MemberPropertyName == GET_MEMBER_NAME_CHECKED(ULandscapeEditLayerBase, WeightmapLayerAllocationBlend))
+	{
+		SetWeightmapLayerAllocationBlend(WeightmapLayerAllocationBlend, /*bInModify = */true);
+	}
+}
+
+void ULandscapeEditLayerBase::PostEditUndo()
+{
+	Super::PostEditUndo();
+
+	BroadcastOnLayerDataChanged(/*InPropertyName = */NAME_None, /*bInUserTriggered = */true, /*bInRequiresLandscapeUpdate =*/ true, /*bInHasValueChanged = */true, /*InChangeType = */EPropertyChangeType::ValueSet);
+}
+
+bool ULandscapeEditLayerBase::CanEditChange(const FProperty* InProperty) const
+{
+	bool bCanEdit = Super::CanEditChange(InProperty);
+
+	const FName MemberPropertyName = InProperty ? InProperty->GetFName() : NAME_None;
+
+	// Always able to update locked property
+	if (MemberPropertyName == GET_MEMBER_NAME_CHECKED(ULandscapeEditLayerBase, bLocked))
+	{
+		return true;
+	}
+
+	// All other properties disabled when layer is locked
+	if (IsLocked())
+	{
+		return false;
+	}
+
+	if (MemberPropertyName == GET_MEMBER_NAME_CHECKED(ULandscapeEditLayerBase, HeightmapAlpha))
+	{
+		bCanEdit &= SupportsAlphaForTargetType(ELandscapeToolTargetType::Heightmap);
+	}
+	else if (MemberPropertyName == GET_MEMBER_NAME_CHECKED(ULandscapeEditLayerBase, WeightmapAlpha))
+	{
+		bCanEdit &= SupportsAlphaForTargetType(ELandscapeToolTargetType::Weightmap);
+	}
+
+	return bCanEdit;
 }
 
 void ULandscapeEditLayerBase::PostLoad()
 {
 	Super::PostLoad();
 
-	// TODO[jonathan.bard] Remove
 	// Needed because we might have saved some layers before we realized we were missing this flag
 	SetFlags(RF_Transactional);
 }
 
-#if WITH_EDITOR
+void ULandscapeEditLayerBase::BroadcastOnLayerDataChanged(FName InPropertyName, bool bInUserTriggered, bool bInRequiresLandscapeUpdate, bool bInHasValueChanged, EPropertyChangeType::Type InChangeType)
+{
+	FProperty* Property = nullptr;
+	if (InPropertyName != NAME_None)
+	{
+		Property = FindFProperty<FProperty>(StaticClass(), InPropertyName);
+		check(Property != nullptr);
+	}
+
+	FOnLandscapeEditLayerDataChangedParams OnLayerDataChangedParams(FPropertyChangedEvent(Property, InChangeType));
+	OnLayerDataChangedParams.bUserTriggered = bInUserTriggered;
+	OnLayerDataChangedParams.bRequiresLandscapeUpdate = bInRequiresLandscapeUpdate;
+	OnLayerDataChangedParams.bHasValueChanged = bInHasValueChanged;
+	OnLayerDataChangedDelegate.Broadcast(OnLayerDataChangedParams);
+}
 
-const FLandscapeLayer* ULandscapeEditLayerBase::GetOwningLayer() const
+ELandscapeToolTargetTypeFlags ULandscapeEditLayerBase::GetEnabledTargetTypeMask() const
 {
-	if (OwningLandscape != nullptr)
+	// Compute the default state of each target type : 
+	ELandscapeToolTargetTypeFlags EnabledTargetTypeMask = ELandscapeToolTargetTypeFlags::None;
+	if (IsVisible())
 	{
-		TArrayView<const FLandscapeLayer> Layers = OwningLandscape->GetLayers();
-		return Layers.FindByPredicate([this](const FLandscapeLayer& Struct) { return Struct.EditLayer == this; });
+		//  HeightmapAlpha might still be set to a different value)
+		if (GetAlphaForTargetType(ELandscapeToolTargetType::Heightmap) != 0.0f)
+		{
+			EnabledTargetTypeMask |= ELandscapeToolTargetTypeFlags::Heightmap;
+		}
+		if (GetAlphaForTargetType(ELandscapeToolTargetType::Weightmap) > 0.0f)
+		{
+			EnabledTargetTypeMask |= ELandscapeToolTargetTypeFlags::Weightmap;
+		}
+		// If the layer is visible, visibility is always considered enabled since it does not depend on WeightmapAlpha 
+		EnabledTargetTypeMask |= ELandscapeToolTargetTypeFlags::Visibility;
 	}
+	return EnabledTargetTypeMask;
+}
+
+void ULandscapeEditLayerBase::SetBackPointer(ALandscape* Landscape)
+{
+	OwningLandscape = Landscape;
+}
+#endif //WITH_EDITOR
 
-	return nullptr;
+#if WITH_EDITORONLY_DATA
+void ULandscapeEditLayerBase::SetHeightmapAlphaInternal(float InNewValue)
+{
+	SetAlphaForTargetType(ELandscapeToolTargetType::Heightmap, InNewValue, /*bInModify = */true, /*InChangeType = */EPropertyChangeType::ValueSet);
 }
 
-#endif // WITH_EDITOR
+void ULandscapeEditLayerBase::SetWeightmapAlphaInternal(float InNewValue)
+{
+	SetAlphaForTargetType(ELandscapeToolTargetType::Weightmap, InNewValue, /*bInModify = */true, /*InChangeType = */EPropertyChangeType::ValueSet);
+}
 
+void ULandscapeEditLayerBase::SetGuidInternal(const FGuid& InGuid)
+{
+	SetGuid(InGuid, /*bInModify = */true);
+}
 
-// ----------------------------------------------------------------------------------
+void ULandscapeEditLayerBase::SetNameInternal(FName InName)
+{
+	SetName(InName, /*bInModify = */true);
+}
 
-bool ULandscapeEditLayer::SupportsTargetType(ELandscapeToolTargetType InType) const
+void ULandscapeEditLayerBase::SetVisibleInternal(bool bInVisible)
 {
-	return (InType == ELandscapeToolTargetType::Heightmap) || (InType == ELandscapeToolTargetType::Weightmap) || (InType == ELandscapeToolTargetType::Visibility);
+	SetVisible(bInVisible, /*bInModify = */true);
 }
 
+void ULandscapeEditLayerBase::SetLockedInternal(bool bInLocked)
+{
+	SetLocked(bInLocked, /*bInModify = */true);
+}
+
+void ULandscapeEditLayerBase::SetWeightmapLayerAllocationBlendInternal(const TMap<TObjectPtr<ULandscapeLayerInfoObject>, bool>& InWeightmapLayerAllocationBlend)
+{
+	SetWeightmapLayerAllocationBlend(InWeightmapLayerAllocationBlend, /*bInModify = */true);
+}
+
+#endif // WITH_EDITORONLY_DATA
 
 // ----------------------------------------------------------------------------------
 
-bool ULandscapeEditLayerSplines::SupportsTargetType(ELandscapeToolTargetType InType) const
+bool ULandscapeEditLayer::SupportsTargetType(ELandscapeToolTargetType InType) const
 {
 	return (InType == ELandscapeToolTargetType::Heightmap) || (InType == ELandscapeToolTargetType::Weightmap) || (InType == ELandscapeToolTargetType::Visibility);
 }
 
-void ULandscapeEditLayerSplines::OnLayerCreated(FLandscapeLayer& Layer)
+// ----------------------------------------------------------------------------------
+
+bool ULandscapeEditLayerSplines::SupportsTargetType(ELandscapeToolTargetType InType) const
 {
-	// Splines edit layer is always using alpha blend mode
-	Layer.BlendMode = LSBM_AlphaBlend;
+	return (InType == ELandscapeToolTargetType::Heightmap) || (InType == ELandscapeToolTargetType::Weightmap) || (InType == ELandscapeToolTargetType::Visibility);
 }
 
 TArray<ULandscapeEditLayerSplines::FEditLayerAction> ULandscapeEditLayerSplines::GetActions() const
@@ -72,9 +438,11 @@ TArray<ULandscapeEditLayerSplines::FEditLayerAction> ULandscapeEditLayerSplines:
 		}),
 		FEditLayerAction::FCanExecuteDelegate::CreateWeakLambda(this, [](const FEditLayerAction::FExecuteParams& InParams, FText& OutReason)
 		{
-			if (InParams.GetLayer()->bLocked)
+			const ULandscapeEditLayerBase* EditLayer = InParams.GetLayer()->EditLayer;
+			check(EditLayer != nullptr);
+			if (EditLayer->IsLocked())
 			{
-				OutReason = FText::Format(LOCTEXT("LandscapeEditLayerSplines_CannotUpdateSplinesOnLockedLayer", "Cannot update splines on layer '{0}' : the layer is currently locked"), FText::FromName(InParams.GetLayer()->Name));
+				OutReason = FText::Format(LOCTEXT("LandscapeEditLayerSplines_CannotUpdateSplinesOnLockedLayer", "Cannot update splines on layer '{0}' : the layer is currently locked"), FText::FromName(InParams.GetLayer()->EditLayer->GetName()));
 				return false;
 			}
 
diff --git a/Engine/Source/Runtime/Landscape/Private/LandscapeEditLayers.cpp b/Engine/Source/Runtime/Landscape/Private/LandscapeEditLayers.cpp
--- a/Engine/Source/Runtime/Landscape/Private/LandscapeEditLayers.cpp
+++ b/Engine/Source/Runtime/Landscape/Private/LandscapeEditLayers.cpp
@@ -2376,7 +2376,7 @@ ELandscapeEditLayersMergeMode ALandscape::GetEditLayersMergeMode()
 	for (const FLandscapeLayer& Layer : Landscape->LandscapeEditLayers)
 	{
 		// No BP brush is supported for local merge of edit layers : 
-		if (Layer.bVisible && Algo::AnyOf(Layer.Brushes, [](const FLandscapeLayerBrush& Brush) { return (Brush.GetBrush() != nullptr) && Brush.GetBrush()->IsVisible(); }))
+		if (Layer.EditLayer->IsVisible() && Algo::AnyOf(Layer.Brushes, [](const FLandscapeLayerBrush& Brush) { return (Brush.GetBrush() != nullptr) && Brush.GetBrush()->IsVisible(); }))
 		{
 				return ELandscapeEditLayersMergeMode::GlobalMerge;
 			}
@@ -2560,7 +2560,7 @@ void ALandscape::ToggleCanHaveLayersContent()
 	}
 	else
 	{
-		check(GetLayerCount() == 0);
+		check(LandscapeEditLayers.IsEmpty());
 		CreateDefaultLayer();
 		CopyOldDataToDefaultLayer();
 	}
@@ -2716,7 +2716,7 @@ void ALandscape::CopyOldDataToDefaultLayer(ALandscapeProxy* InProxy)
 
 	for (ULandscapeComponent* Component : InProxy->LandscapeComponents)
 	{
-		FLandscapeLayerComponentData* LayerData = Component->GetLayerData(DefaultLayer->Guid);
+		FLandscapeLayerComponentData* LayerData = Component->GetLayerData(DefaultLayer->EditLayer->GetGuid());
 
 		if (ensure(LayerData != nullptr && LayerData->IsInitialized()))
 		{
@@ -2762,7 +2762,7 @@ void ALandscape::CopyOldDataToDefaultLayer(ALandscapeProxy* InProxy)
 				{
 					LayerData->WeightmapData.Textures[TextureIndex] = WeightmapTextureData->Texture;
 					LayerData->WeightmapData.TextureUsages[TextureIndex] = WeightmapTextureData->Usage;
-					check(WeightmapTextureData->Usage->LayerGuid == DefaultLayer->Guid);
+					check(WeightmapTextureData->Usage->LayerGuid == DefaultLayer->EditLayer->GetGuid());
 
 					for (int32 ChannelIndex = 0; ChannelIndex < ULandscapeWeightmapUsage::NumChannels; ++ChannelIndex)
 					{
@@ -2804,7 +2804,7 @@ void ALandscape::CopyOldDataToDefaultLayer(ALandscapeProxy* InProxy)
 						LayerData->WeightmapData.TextureUsages[TextureIndex]->ChannelUsage[ChannelIndex] = ComponentWeightmapTexturesUsage[TextureIndex]->ChannelUsage[ChannelIndex];
 					}
 
-					LayerData->WeightmapData.TextureUsages[TextureIndex]->LayerGuid = DefaultLayer->Guid;
+					LayerData->WeightmapData.TextureUsages[TextureIndex]->LayerGuid = DefaultLayer->EditLayer->GetGuid();
 
 					// Create new Usage for the "final" layer as the other one will now be used by the Default layer
 					for (const FWeightmapLayerAllocationInfo& Allocation : ComponentLayerAllocations)
@@ -2897,10 +2897,10 @@ void ALandscapeProxy::InitializeProxyLayersWeightmapUsage()
 			// Reinitialize the weightmap usages for the base (final) paint layers allocations :
 			Component->InitializeLayersWeightmapUsage(FGuid());
 
-			for (const FLandscapeLayer& Layer : Landscape->GetLayers())
+			for (const FLandscapeLayer& Layer : Landscape->GetLayersConst())
 			{
 				// Reinitialize each edit layer's weightmap usages list :
-				Component->InitializeLayersWeightmapUsage(Layer.Guid);
+				Component->InitializeLayersWeightmapUsage(Layer.EditLayer->GetGuid());
 			}
 		}
 	}
@@ -2915,8 +2915,8 @@ void ULandscapeComponent::InitializeLayersWeightmapUsage(const FGuid& InLayerGui
 	check(Proxy);
 	ALandscape* Landscape = GetLandscapeActor();
 	check(Landscape);
-	const FLandscapeLayer* SplinesLayer = Landscape->FindLayerOfType(ULandscapeEditLayerSplines::StaticClass());
-	FGuid SplinesEditLayerGuid = SplinesLayer ? SplinesLayer->Guid : FGuid();
+	const FLandscapeLayer* SplinesLayer = Landscape->FindLayerOfTypeConst(ULandscapeEditLayerSplines::StaticClass());
+	FGuid SplinesEditLayerGuid = SplinesLayer ? SplinesLayer->EditLayer->GetGuid() : FGuid();
 
 	// Don't consider invalid edit layers : 
 	if (InLayerGuid.IsValid())
@@ -3062,10 +3062,10 @@ void ALandscapeProxy::ValidateProxyLayersWeightmapUsage() const
 
 			// Validate edit layers weightmap allocations : 
 			{
-				const FLandscapeLayer* SplinesLayer = Landscape->FindLayerOfType(ULandscapeEditLayerSplines::StaticClass());
-				for (const FLandscapeLayer& Layer : Landscape->GetLayers())
+				const FLandscapeLayer* SplinesLayer = Landscape->FindLayerOfTypeConst(ULandscapeEditLayerSplines::StaticClass());
+				for (const FLandscapeLayer& Layer : Landscape->GetLayersConst())
 				{
-					FLandscapeLayerComponentData* LayerData = Component->GetLayerData(Layer.Guid);
+					FLandscapeLayerComponentData* LayerData = Component->GetLayerData(Layer.EditLayer->GetGuid());
 
 					// Skip validation on SplinesLayer since it can momentarily contain duplicated layer allocations after undo (since it's updated outside of a transaction) :
 					if (LayerData != nullptr && LayerData->IsInitialized() && (&Layer != SplinesLayer))
@@ -3078,7 +3078,7 @@ void ALandscapeProxy::ValidateProxyLayersWeightmapUsage() const
 								UTexture2D* WeightmapTexture = LayerData->WeightmapData.Textures[Allocation.WeightmapTextureIndex];
 								if (ULandscapeWeightmapUsage* Usage = LayerData->WeightmapData.TextureUsages.IsValidIndex(Allocation.WeightmapTextureIndex) ? LayerData->WeightmapData.TextureUsages[Allocation.WeightmapTextureIndex].Get() : nullptr)
 								{
-									ValidateWeightmapAllocationAndUsage(WeightmapTexture, Allocation, Usage, Layer.Guid);
+									ValidateWeightmapAllocationAndUsage(WeightmapTexture, Allocation, Usage, Layer.EditLayer->GetGuid());
 								}
 							}
 						}
@@ -4207,7 +4207,7 @@ bool ALandscape::PrepareLayersTextureResources(const TArray<FLandscapeLayer>& In
 		{
 			for (ULandscapeComponent* Component : Proxy->LandscapeComponents)
 			{
-				if (FLandscapeLayerComponentData* ComponentLayerData = Component->GetLayerData(Layer.Guid))
+				if (FLandscapeLayerComponentData* ComponentLayerData = Component->GetLayerData(Layer.EditLayer->GetGuid()))
 				{
 					check(ComponentLayerData->HeightmapData.Texture != nullptr);
 					bIsReady &= TextureStreamingManager->RequestTextureFullyStreamedInForever(ComponentLayerData->HeightmapData.Texture, bInWaitForStreaming);
@@ -5066,11 +5066,12 @@ void ALandscape::PrepareLayersHeightmapsLocalMergeRenderThreadData(const FUpdate
 	OutRenderThreadData.VisibleEditLayerInfos.Reserve(LandscapeEditLayers.Num());
 	for (const FLandscapeLayer& Layer : LandscapeEditLayers)
 	{
-		if (Layer.bVisible && !InMergeParams.bSkipBrush)
+		check(Layer.EditLayer != nullptr);
+		if (Layer.EditLayer->IsVisible() && !InMergeParams.bSkipBrush)
 		{
 			FEditLayerInfo& EditLayerInfo = OutRenderThreadData.VisibleEditLayerInfos.Emplace_GetRef();
-			EditLayerInfo.BlendMode = LandscapeBlendModeToEditLayerBlendMode(Layer.BlendMode);
-			EditLayerInfo.Alpha = Layer.HeightmapAlpha;
+			EditLayerInfo.BlendMode = LandscapeBlendModeToEditLayerBlendMode(Layer.EditLayer->GetBlendMode());
+			EditLayerInfo.Alpha = Layer.EditLayer->GetAlphaForTargetType(ELandscapeToolTargetType::Heightmap);
 		}
 	}
 
@@ -5106,9 +5107,9 @@ void ALandscape::PrepareLayersHeightmapsLocalMergeRenderThreadData(const FUpdate
 			NewComponentRenderInfo.VisibleLayerHeightmapTextures.Reserve(OutRenderThreadData.VisibleEditLayerInfos.Num());
 			for (const FLandscapeLayer& Layer : LandscapeEditLayers)
 			{
-				if (Layer.bVisible && !InMergeParams.bSkipBrush)
+				if (Layer.EditLayer->IsVisible() && !InMergeParams.bSkipBrush)
 				{
-					if (FLandscapeLayerComponentData* ComponentLayerData = Component->GetLayerData(Layer.Guid))
+					if (FLandscapeLayerComponentData* ComponentLayerData = Component->GetLayerData(Layer.EditLayer->GetGuid()))
 					{
 						if (UTexture2D* LayerHeightmap = ComponentLayerData->HeightmapData.Texture.Get())
 						{
@@ -5339,14 +5340,6 @@ TArray<UE::Landscape::EditLayers::FEditLayerRendererState> ALandscape::GetEditLa
 	for (FLandscapeLayer& Layer : LandscapeEditLayers)
 	{
 		TArray<FEditLayerRendererState> LayerRendererStates = Layer.GetEditLayerRendererStates(InMergeContext);
-		// Renderer states generated from a layer inherit the layer's state so start by computing the layer's target type mask: 
-		ELandscapeToolTargetTypeFlags LayerTargetTypeMask = Layer.GetEnabledTargetTypeMask();
-		// Then disable all types that are not in the layer's mask :
-		for (FEditLayerRendererState& LayerRendererState : LayerRendererStates)
-		{
-			LayerRendererState.DisableTargetTypeMask(~LayerTargetTypeMask);
-		}
-
 		RendererStates.Append(LayerRendererStates);
 	}
 	return RendererStates;
@@ -5374,29 +5367,16 @@ TArray<UE::Landscape::EditLayers::FEditLayerRendererState> FLandscapeLayer::GetE
 	{
 		RendererStates.Append(Brush.GetEditLayerRendererStates(InMergeContext));
 	}
-	return RendererStates;
-}
 
-ELandscapeToolTargetTypeFlags FLandscapeLayer::GetEnabledTargetTypeMask() const
+	// Renderer states generated from a layer inherit the layer's state so start by computing the layer's target type mask: 
+	ELandscapeToolTargetTypeFlags LayerTargetTypeMask = EditLayer->GetEnabledTargetTypeMask();
+	// Then disable all types that are not in the layer's mask :
+	for (FEditLayerRendererState& LayerRendererState : RendererStates)
 {
-	// Compute the default state of each target type : 
-	ELandscapeToolTargetTypeFlags EnabledTargetTypeMask = ELandscapeToolTargetTypeFlags::None;
-	if (bVisible)
-	{
-		// TODO [jonathan.bard] : Remove the splines type check when HeightmapAlpha is replaced by a virtual GetHeightmapAlpha() on ULandscapeEditLayer (alpha is fixed to 1.0 in this case, but 
-		//  HeightmapAlpha might still be set to a different value)
-		if (EditLayer->IsA<ULandscapeEditLayerSplines>() || (HeightmapAlpha != 0.0f))
-		{
-			EnabledTargetTypeMask |= ELandscapeToolTargetTypeFlags::Heightmap;
+		LayerRendererState.DisableTargetTypeMask(~LayerTargetTypeMask);
 		}
-		if (WeightmapAlpha > 0.0f)
-		{
-			EnabledTargetTypeMask |= ELandscapeToolTargetTypeFlags::Weightmap;
-		}
-		// If the layer is visible, visibility is always considered enabled since it does not depend on WeightmapAlpha 
-		EnabledTargetTypeMask |= ELandscapeToolTargetTypeFlags::Visibility;
-	}
-	return EnabledTargetTypeMask;
+
+	return RendererStates;
 }
 
 TArray<UE::Landscape::EditLayers::FEditLayerRendererState> FLandscapeLayerBrush::GetEditLayerRendererStates(const UE::Landscape::EditLayers::FMergeContext* InMergeContext)
@@ -5417,16 +5397,12 @@ void ULandscapeEditLayerPersistent::GetRendererStateInfo(const UE::Landscape::Ed
 {
 	using namespace UE::Landscape::EditLayers;
 
-	const FLandscapeLayer* LayerStruct = GetOwningLayer();
-	check(LayerStruct != nullptr);
-
 	// A layer can support all target types. 
 	// Add an entry for each weightmap and consider them supported because there's nothing that prevents a given edit layer to write on a given weightmap layer :
 	OutSupportedTargetTypeState = FEditLayerTargetTypeState(InMergeContext, ELandscapeToolTargetTypeFlags::All, InMergeContext->GetValidTargetLayerBitIndices());
 
-	// TODO [jonathan.bard] : for now, at this level enabled == supported (the enabled state will be corrected by FLandscapeLayer since it is the one that has the bVisibility / Heightmap/WeightmapAlpha members
-	//  but eventually, if those move to ULandscapeEditLayerBase, this is where the layer's enabled state will need to be set)
-	OutEnabledTargetTypeState = OutSupportedTargetTypeState;
+	// Compute the default state of each target type : 
+	OutEnabledTargetTypeState = FEditLayerTargetTypeState(InMergeContext, GetEnabledTargetTypeMask(), OutSupportedTargetTypeState.GetActiveWeightmapBitIndices());
 }
 
 UE::Landscape::EditLayers::ERenderFlags ULandscapeEditLayerPersistent::GetRenderFlags(const UE::Landscape::EditLayers::FMergeContext* InMergeContext) const
@@ -5439,9 +5415,6 @@ TArray<UE::Landscape::EditLayers::FEditLayerRenderItem> ULandscapeEditLayerPersi
 {
 	using namespace UE::Landscape::EditLayers;
 
-	const FLandscapeLayer* LayerStruct = GetOwningLayer();
-	check(LayerStruct != nullptr);
-
 	TArray<FEditLayerRenderItem> RenderItems;
 	ULandscapeInfo* Info = InMergeContext->GetLandscapeInfo();
 	RenderItems.Reserve(Info->XYtoComponentMap.Num() + 1);
@@ -5465,7 +5438,7 @@ TArray<UE::Landscape::EditLayers::FEditLayerRenderItem> ULandscapeEditLayerPersi
 		{
 			const FIntPoint ComponentKey = Component->GetComponentKey();
 
-			if (FLandscapeLayerComponentData* LayerData = Component->GetLayerData(LayerStruct->Guid))
+			if (FLandscapeLayerComponentData* LayerData = Component->GetLayerData(GetGuid()))
 			{
 				TArray<FName, TInlineAllocator<16>> ComponentWeightmaps;
 				FEditLayerTargetTypeState OutputTargetTypeState(InMergeContext);
@@ -5514,8 +5487,6 @@ bool ULandscapeEditLayerPersistent::RenderLayer(UE::Landscape::EditLayers::FRend
 	check(Info != nullptr);
 	ALandscape* Landscape = RenderParams.MergeRenderContext->GetLandscape();
 	FSceneInterface* SceneInterface = Landscape->GetWorld()->Scene;
-	const FLandscapeLayer* LayerStruct = GetOwningLayer();
-	check(LayerStruct != nullptr);
 
 	const int32 NumTargetLayersInGroup = RenderParams.TargetLayerGroupLayerNames.Num();
 
@@ -5557,13 +5528,13 @@ bool ULandscapeEditLayerPersistent::RenderLayer(UE::Landscape::EditLayers::FRend
 
 			if (RenderParams.MergeRenderContext->IsHeightmapMerge())
 			{
-				SourceTexture = ComponentMergeRenderInfo.Component->GetHeightmap(LayerStruct->Guid);
+				SourceTexture = ComponentMergeRenderInfo.Component->GetHeightmap(GetGuid());
 				SourceTextureBias = FVector2D(ComponentMergeRenderInfo.Component->HeightmapScaleBias.Z, ComponentMergeRenderInfo.Component->HeightmapScaleBias.W);
 			}
 			else
 			{
-				const TArray<UTexture2D*>& WeightmapTextures = ComponentMergeRenderInfo.Component->GetWeightmapTextures(LayerStruct->Guid);
-				const TArray<FWeightmapLayerAllocationInfo>& AllocInfos = ComponentMergeRenderInfo.Component->GetWeightmapLayerAllocations(LayerStruct->Guid);
+				const TArray<UTexture2D*>& WeightmapTextures = ComponentMergeRenderInfo.Component->GetWeightmapTextures(GetGuid());
+				const TArray<FWeightmapLayerAllocationInfo>& AllocInfos = ComponentMergeRenderInfo.Component->GetWeightmapLayerAllocations(GetGuid());
 				const FWeightmapLayerAllocationInfo* AllocInfo = AllocInfos.FindByPredicate([TargetLayerName](const FWeightmapLayerAllocationInfo& InAllocInfo) { return InAllocInfo.IsAllocated() && (InAllocInfo.GetLayerName() == TargetLayerName); });
 				if (AllocInfo != nullptr)
 				{
@@ -5806,8 +5777,6 @@ void ULandscapeEditLayerPersistent::BlendLayer(UE::Landscape::EditLayers::FRende
 
 	ULandscapeInfo* Info = RenderParams.MergeRenderContext->GetLandscapeInfo();
 	check(Info != nullptr);
-	const FLandscapeLayer* LayerStruct = GetOwningLayer();
-	check(LayerStruct != nullptr);
 
 	const int32 NumTargetLayersInGroup = RenderParams.TargetLayerGroupLayerNames.Num();
 
@@ -5815,10 +5784,8 @@ void ULandscapeEditLayerPersistent::BlendLayer(UE::Landscape::EditLayers::FRende
 	FBlendParams BlendParams;
 	if (RenderParams.MergeRenderContext->IsHeightmapMerge())
 	{
-		BlendParams.HeightmapBlendParams.BlendMode = LandscapeBlendModeToHeightmapBlendMode(LayerStruct->BlendMode);
-		// TODO [jonathan.bard] : Remove the splines type check when HeightmapAlpha is replaced by a virtual GetHeightmapAlpha() on ULandscapeEditLayer (alpha is fixed to 1.0 in this case, but 
-		//  HeightmapAlpha might still be set to a different value)
-		BlendParams.HeightmapBlendParams.Alpha = IsA<ULandscapeEditLayerSplines>() ? 1.0f : LayerStruct->HeightmapAlpha;
+		BlendParams.HeightmapBlendParams.BlendMode = LandscapeBlendModeToHeightmapBlendMode(GetBlendMode());
+		BlendParams.HeightmapBlendParams.Alpha = GetAlphaForTargetType(ELandscapeToolTargetType::Heightmap);
 	}
 	else
 	{
@@ -5835,14 +5802,14 @@ void ULandscapeEditLayerPersistent::BlendLayer(UE::Landscape::EditLayers::FRende
 			{
 				FWeightmapBlendParams& TargetLayerBlendParams = BlendParams.WeightmapBlendParams.Emplace(TargetLayerName, EWeightmapBlendMode::Additive);
 
-				if (const bool* bSubstractiveInLayer = LayerStruct->WeightmapLayerAllocationBlend.Find(LayerInfoObj); (bSubstractiveInLayer != nullptr) && *bSubstractiveInLayer)
+				if (const bool* bSubstractiveInLayer = GetWeightmapLayerAllocationBlend().Find(LayerInfoObj); (bSubstractiveInLayer != nullptr) && *bSubstractiveInLayer)
 				{
 					TargetLayerBlendParams.BlendMode = EWeightmapBlendMode::Subtractive;
 				}
 
 				if (TargetLayerName != UMaterialExpressionLandscapeVisibilityMask::ParameterName)
 				{
-					TargetLayerBlendParams.Alpha = LayerStruct->WeightmapAlpha;
+					TargetLayerBlendParams.Alpha = GetAlphaForTargetType(ELandscapeToolTargetType::Weightmap);
 				}
 			}
 		}
@@ -5854,9 +5821,7 @@ void ULandscapeEditLayerPersistent::BlendLayer(UE::Landscape::EditLayers::FRende
 
 FString ULandscapeEditLayerPersistent::GetEditLayerRendererDebugName() const
 {
-	const FLandscapeLayer* LayerStruct = GetOwningLayer();
-	check(LayerStruct != nullptr);
-	return LayerStruct->Name.ToString();
+	return GetName().ToString();
 }
 
 namespace UE::Landscape::EditLayers::Private
@@ -7385,29 +7350,32 @@ int32 ALandscape::PerformLayersHeightmapsGlobalMerge(const FUpdateLayersContentC
 
 	for (FLandscapeLayer& Layer : LandscapeEditLayers)
 	{
+		check(Layer.EditLayer != nullptr);
+		const ELandscapeBlendMode LayerBlendMode = Layer.EditLayer->GetBlendMode();
+
 		// Draw each Layer's heightmaps to a Combined RT Atlas in LandscapeScratchRT1
 		ShaderParams.ApplyLayerModifiers = false;
 		ShaderParams.SetAlphaOne = false;
-		ShaderParams.LayerVisible = Layer.bVisible;
+		ShaderParams.LayerVisible = Layer.EditLayer->IsVisible();
 		ShaderParams.GenerateNormals = false;
-		ShaderParams.LayerBlendMode = Layer.BlendMode;
+		ShaderParams.LayerBlendMode = LayerBlendMode;
 
-		if (Layer.BlendMode == LSBM_AlphaBlend)
+		if (LayerBlendMode == LSBM_AlphaBlend)
 		{
 			// For now, only Layer reserved for Landscape Splines will use the AlphaBlendMode
-			check((Layer.EditLayer != nullptr) && Layer.EditLayer->IsA<ULandscapeEditLayerSplines>());
+			check(Layer.EditLayer->IsA<ULandscapeEditLayerSplines>());
 			ShaderParams.LayerAlpha = 1.0f;
 		}
 		else
 		{
-			check(Layer.BlendMode == LSBM_AdditiveBlend);
-			ShaderParams.LayerAlpha = Layer.HeightmapAlpha;
+			check(LayerBlendMode == LSBM_AdditiveBlend);
+			ShaderParams.LayerAlpha = Layer.EditLayer->GetAlphaForTargetType(ELandscapeToolTargetType::Heightmap);
 		}
 
 		{
 			TArray<FLandscapeLayersCopyTextureParams> DeferredCopyTextures;
 			TArray<FHeightmapCopyInfo> LayerHeightmaps;
-			GetUniqueHeightmaps(InUpdateLayersContentContext.LandscapeComponentsHeightmapsToRender, LayerHeightmaps, LandscapeExtent.Min, Layer.Guid);
+			GetUniqueHeightmaps(InUpdateLayersContentContext.LandscapeComponentsHeightmapsToRender, LayerHeightmaps, LandscapeExtent.Min, Layer.EditLayer->GetGuid());
 			for (const FHeightmapCopyInfo& LayerHeightmap : LayerHeightmaps)
 			{
 				FLandscapeLayersCopyTextureParams& CopyTextureParams = DeferredCopyTextures.Add_GetRef(FLandscapeLayersCopyTextureParams(LayerHeightmap.Texture, LandscapeScratchRT1));
@@ -7421,18 +7389,18 @@ int32 ALandscape::PerformLayersHeightmapsGlobalMerge(const FUpdateLayersContentC
 
 		// Convert Atlas LandscapeScratchRT1 to the world-projected NonAtlas in LandscapeScratchRT2
 		// TODO: just use this format from the beginning above...
-		DrawHeightmapComponentsToRenderTarget(FString::Printf(TEXT("LS Height: %s %s += -> NonAtlas %s"), *Layer.Name.ToString(), *LandscapeScratchRT1->GetName(), *LandscapeScratchRT2->GetName()),
+		DrawHeightmapComponentsToRenderTarget(FString::Printf(TEXT("%s: %s -> NonAtlas %s"), *Layer.EditLayer->GetName().ToString(), *LandscapeScratchRT1->GetName(), *LandscapeScratchRT2->GetName()),
 			InUpdateLayersContentContext.LandscapeComponentsHeightmapsToRender, LandscapeExtent.Min, LandscapeScratchRT1, nullptr, LandscapeScratchRT2, ERTDrawingType::RTAtlasToNonAtlas, true, ShaderParams);
 
 		ShaderParams.ApplyLayerModifiers = true;
 
 		// Combine Current layer NonAtlas LandscapeScratchRT2 with current result in LandscapeScratchRT3, writing final result to CombinedHeightmapNonAtlasRT
-		DrawHeightmapComponentsToRenderTarget(FString::Printf(TEXT("LS Height: %s %s += -> CombinedNonAtlas %s"), *Layer.Name.ToString(), *LandscapeScratchRT2->GetName(), *CombinedHeightmapNonAtlasRT->GetName()),
+		DrawHeightmapComponentsToRenderTarget(FString::Printf(TEXT("%s: %s += -> CombinedNonAtlas %s"), *Layer.EditLayer->GetName().ToString(), *LandscapeScratchRT2->GetName(), *CombinedHeightmapNonAtlasRT->GetName()),
 			InUpdateLayersContentContext.LandscapeComponentsHeightmapsToRender, LandscapeExtent.Min, LandscapeScratchRT2, FirstLayer ? nullptr : LandscapeScratchRT3, CombinedHeightmapNonAtlasRT, ERTDrawingType::RTNonAtlas, FirstLayer, ShaderParams);
 
 		ShaderParams.ApplyLayerModifiers = false;
 
-		if (Layer.bVisible && !InMergeParams.bSkipBrush)
+		if (Layer.EditLayer->IsVisible() && !InMergeParams.bSkipBrush)
 		{
 			// Draw each brushes
 			for (int32 i = 0; i < Layer.Brushes.Num(); ++i)
@@ -7456,13 +7424,13 @@ int32 ALandscape::PerformLayersHeightmapsGlobalMerge(const FUpdateLayersContentC
 
 				INC_DWORD_STAT(STAT_LandscapeLayersRegenerateDrawCalls); // Brush Render
 
-				PrintLayersDebugRT(FString::Printf(TEXT("LS Height: %s %s -> BrushNonAtlas %s"), *Layer.Name.ToString(), *LandscapeBrush->GetName(), *BrushOutputNonAtlasRT->GetName()), BrushOutputNonAtlasRT);
+				PrintLayersDebugRT(FString::Printf(TEXT("LS Height: %s %s -> BrushNonAtlas %s"), *Layer.EditLayer->GetName().ToString(), *LandscapeBrush->GetName(), *BrushOutputNonAtlasRT->GetName()), BrushOutputNonAtlasRT);
 
 				// Resolve back to Combined heightmap (it's unlikely, but possible that the brush returns the same RT as input and output, if it did various operations on it, in which case the copy is useless) :
 				if (BrushOutputNonAtlasRT != CombinedHeightmapNonAtlasRT)
 				{
 					ExecuteCopyLayersTexture({ FLandscapeLayersCopyTextureParams(BrushOutputNonAtlasRT, CombinedHeightmapNonAtlasRT) });
-					PrintLayersDebugRT(FString::Printf(TEXT("LS Height: %s Component %s += -> CombinedNonAtlas %s"), *Layer.Name.ToString(), *BrushOutputNonAtlasRT->GetName(), *CombinedHeightmapNonAtlasRT->GetName()), CombinedHeightmapNonAtlasRT);
+					PrintLayersDebugRT(FString::Printf(TEXT("LS Height: %s Component %s += -> CombinedNonAtlas %s"), *Layer.EditLayer->GetName().ToString(), *BrushOutputNonAtlasRT->GetName(), *CombinedHeightmapNonAtlasRT->GetName()), CombinedHeightmapNonAtlasRT);
 				}
 			}
 
@@ -7482,7 +7450,7 @@ int32 ALandscape::PerformLayersHeightmapsGlobalMerge(const FUpdateLayersContentC
 					&& BrushOutputNonAtlasRT->SizeY == CombinedHeightmapNonAtlasRT->SizeY)
 				{
 					ExecuteCopyLayersTexture({ FLandscapeLayersCopyTextureParams(BrushOutputNonAtlasRT, CombinedHeightmapNonAtlasRT) });
-					PrintLayersDebugRT(FString::Printf(TEXT("LS Height: %s Component %s += -> CombinedNonAtlas %s"), *Layer.Name.ToString(), 
+					PrintLayersDebugRT(FString::Printf(TEXT("LS Height: %s Component %s += -> CombinedNonAtlas %s"), *Layer.EditLayer->GetName().ToString(), 
 						*BrushOutputNonAtlasRT->GetName(), *CombinedHeightmapNonAtlasRT->GetName()), CombinedHeightmapNonAtlasRT);
 				}
 			}
@@ -7491,7 +7459,7 @@ int32 ALandscape::PerformLayersHeightmapsGlobalMerge(const FUpdateLayersContentC
 		// copy CombinedHeightmapNonAtlasRT to LandscapeScratchRT3 (as a source for later layers... this is wasted on the last layer I think...)
 		// TODO: you can get the same effect for much cheaper by swapping these two pointers before the render above...
 		ExecuteCopyLayersTexture({ FLandscapeLayersCopyTextureParams(CombinedHeightmapNonAtlasRT, LandscapeScratchRT3) });
-		PrintLayersDebugRT(FString::Printf(TEXT("LS Height: %s Component %s += -> CombinedNonAtlas %s"), *Layer.Name.ToString(), *CombinedHeightmapNonAtlasRT->GetName(), *LandscapeScratchRT3->GetName()), LandscapeScratchRT3);
+		PrintLayersDebugRT(FString::Printf(TEXT("LS Height: %s Component %s += -> CombinedNonAtlas %s"), *Layer.EditLayer->GetName().ToString(), *CombinedHeightmapNonAtlasRT->GetName(), *LandscapeScratchRT3->GetName()), LandscapeScratchRT3);
 
 		FirstLayer = false;
 	}
@@ -8610,7 +8578,7 @@ void ALandscape::PrepareComponentDataToExtractMaterialLayersCS(const TArray<ULan
 	const int32 LocalComponentSizeVerts = (SubsectionSizeQuads + 1) * NumSubsections;
 	for (ULandscapeComponent* Component : InLandscapeComponents)
 	{
-		FLandscapeLayerComponentData* ComponentLayerData = Component->GetLayerData(InLayer.Guid);
+		FLandscapeLayerComponentData* ComponentLayerData = Component->GetLayerData(InLayer.EditLayer->GetGuid());
 
 		if (ComponentLayerData != nullptr)
 		{
@@ -8626,13 +8594,13 @@ void ALandscape::PrepareComponentDataToExtractMaterialLayersCS(const TArray<ULan
 				FVector2D SourcePositionOffset(FMath::RoundToInt((float)ComponentSectionBase.X / LocalComponentSizeQuad), FMath::RoundToInt((float)ComponentSectionBase.Y / LocalComponentSizeQuad));
 				FIntPoint SourceComponentVertexPosition = FIntPoint(static_cast<int32>(SourcePositionOffset.X * LocalComponentSizeVerts), static_cast<int32>(SourcePositionOffset.Y * LocalComponentSizeVerts));
 
-				FLandscapeLayersCopyTextureParams& CopyTextureParams = DeferredCopyTextures.Add_GetRef(FLandscapeLayersCopyTextureParams(*LayerWeightmap->GetName(), LayerWeightmap->GetResource(), FString::Printf(TEXT("%s WeightmapScratchTexture"), *InLayer.Name.ToString()), InOutTextureData));
+				FLandscapeLayersCopyTextureParams& CopyTextureParams = DeferredCopyTextures.Add_GetRef(FLandscapeLayersCopyTextureParams(*LayerWeightmap->GetName(), LayerWeightmap->GetResource(), FString::Printf(TEXT("%s WeightmapScratchTexture"), *InLayer.EditLayer->GetName().ToString()), InOutTextureData));
 				// Only copy the size that's actually needed : 
 				CopyTextureParams.CopySize.X = LayerWeightmap->GetResource()->GetSizeX();
 				CopyTextureParams.CopySize.Y = LayerWeightmap->GetResource()->GetSizeY();
 				// Copy from the top-left corner of the weightmap to the composited texture's position
 				CopyTextureParams.DestPosition = SourceComponentVertexPosition;
-				PrintLayersDebugTextureResource(FString::Printf(TEXT("LS Weight: %s WeightmapScratchTexture %s"), *InLayer.Name.ToString(), TEXT("WeightmapScratchTextureResource")), InOutTextureData, 0, false);
+				PrintLayersDebugTextureResource(FString::Printf(TEXT("LS Weight: %s WeightmapScratchTexture %s"), *InLayer.EditLayer->GetName().ToString(), TEXT("WeightmapScratchTextureResource")), InOutTextureData, 0, false);
 
 				for (const FWeightmapLayerAllocationInfo& WeightmapLayerAllocation : ComponentLayerData->WeightmapData.LayerAllocations)
 				{
@@ -8835,43 +8803,43 @@ void ALandscape::ReallocateLayersWeightmaps(FUpdateLayersContentContext& InUpdat
 	}
 	else
 	{
-		for (FLandscapeLayer& Layer : LandscapeEditLayers)
+	for (FLandscapeLayer& Layer : LandscapeEditLayers)
+	{
+		for (ULandscapeComponent* Component : InUpdateLayersContentContext.LandscapeComponentsWeightmapsToResolve)
 		{
-			for (ULandscapeComponent* Component : InUpdateLayersContentContext.LandscapeComponentsWeightmapsToResolve)
+			TArray<ULandscapeLayerInfoObject*>* ComponentLayerAlloc = LayerAllocsPerComponent.Find(Component);
+			if (ComponentLayerAlloc == nullptr)
 			{
-				TArray<ULandscapeLayerInfoObject*>* ComponentLayerAlloc = LayerAllocsPerComponent.Find(Component);
-				if (ComponentLayerAlloc == nullptr)
-				{
-					TArray<ULandscapeLayerInfoObject*> NewLayerAllocs;
-					ComponentLayerAlloc = &LayerAllocsPerComponent.Add(Component, NewLayerAllocs);
-				}
+				TArray<ULandscapeLayerInfoObject*> NewLayerAllocs;
+				ComponentLayerAlloc = &LayerAllocsPerComponent.Add(Component, NewLayerAllocs);
+			}
 
-				// No need for an allocation if the edit layer is invisible : 
-				if (Layer.bVisible)
+			// No need for an allocation if the edit layer is invisible : 
+			if (Layer.EditLayer->IsVisible())
+			{
+				if (FLandscapeLayerComponentData* LayerComponentData = Component->GetLayerData(Layer.EditLayer->GetGuid()))
 				{
-					if (FLandscapeLayerComponentData* LayerComponentData = Component->GetLayerData(Layer.Guid))
-					{
-					for (const FWeightmapLayerAllocationInfo& LayerWeightmapAllocation : LayerComponentData->WeightmapData.LayerAllocations)
+				for (const FWeightmapLayerAllocationInfo& LayerWeightmapAllocation : LayerComponentData->WeightmapData.LayerAllocations)
+				{
+					if (LayerWeightmapAllocation.LayerInfo != nullptr)
 					{
-						if (LayerWeightmapAllocation.LayerInfo != nullptr)
-						{
-							ComponentLayerAlloc->AddUnique(LayerWeightmapAllocation.LayerInfo);
-						}
+						ComponentLayerAlloc->AddUnique(LayerWeightmapAllocation.LayerInfo);
 					}
 				}
-				}
+			}
+			}
 
-				// Add the brush alloc also (only if !InMergeParams.bSkipBrush, but InBrushRequiredAllocations should be empty already if InMergeParams.bSkipBrush is true) :
-				for (ULandscapeLayerInfoObject* BrushLayerInfo : InBrushRequiredAllocations)
+			// Add the brush alloc also (only if !InMergeParams.bSkipBrush, but InBrushRequiredAllocations should be empty already if InMergeParams.bSkipBrush is true) :
+			for (ULandscapeLayerInfoObject* BrushLayerInfo : InBrushRequiredAllocations)
+			{
+				if (BrushLayerInfo != nullptr)
 				{
-					if (BrushLayerInfo != nullptr)
-					{
-						ComponentLayerAlloc->AddUnique(BrushLayerInfo);
-					}
+					ComponentLayerAlloc->AddUnique(BrushLayerInfo);
 				}
 			}
 		}
 	}
+	}
 
 	// Trim the components that don't need weightmaps anymore (e.g. all edit layers are made invisible : there were some
 	// components in LandscapeComponentsWeightmapsToResolve but there aren't anymore now).
@@ -9640,11 +9608,12 @@ void ALandscape::PrepareLayersWeightmapsLocalMergeRenderThreadData(const FUpdate
 			TMap<ULandscapeLayerInfoObject*, int32> PaintLayerToComponentPaintLayerRenderInfoIndex;
 			for (const FLandscapeLayer& Layer : LandscapeEditLayers)
 			{
-				if (Layer.bVisible && !InMergeParams.bSkipBrush)
+				check(Layer.EditLayer != nullptr);
+				if (Layer.EditLayer->IsVisible() && !InMergeParams.bSkipBrush)
 				{
 					// Retrieve the input textures/channels that are needed for this component and this edit layer :
-					const TArray<UTexture2D*>& ComponentEditLayerTextures = Component->GetWeightmapTextures(Layer.Guid);
-					const TArray<FWeightmapLayerAllocationInfo>& ComponentEditLayerAllocations = Component->GetWeightmapLayerAllocations(Layer.Guid);
+					const TArray<UTexture2D*>& ComponentEditLayerTextures = Component->GetWeightmapTextures(Layer.EditLayer->GetGuid());
+					const TArray<FWeightmapLayerAllocationInfo>& ComponentEditLayerAllocations = Component->GetWeightmapLayerAllocations(Layer.EditLayer->GetGuid());
 					int32 NumAllocations = ComponentEditLayerAllocations.Num();
 					for (int32 AllocIndex = 0; AllocIndex < NumAllocations; ++AllocIndex)
 					{
@@ -9697,12 +9666,14 @@ void ALandscape::PrepareLayersWeightmapsLocalMergeRenderThreadData(const FUpdate
 							UTexture2D* ComponentEditLayerTexture = ComponentEditLayerTextures[ComponentEditLayerAllocation.WeightmapTextureIndex];
 							int32 TextureIndexInVisibleEditLayerTextures = ComponentRenderInfo.VisibleEditLayerWeightmapTextures.AddUnique(ComponentEditLayerTexture->GetResource()->GetTexture2DResource());
 
+							const ELandscapeToolTargetType TargetType = UE::Landscape::IsVisibilityLayer(ComponentEditLayerAllocation.LayerInfo) ? ELandscapeToolTargetType::Visibility : ELandscapeToolTargetType::Weightmap;
+
 							// Add an entry for each edit layer that participates to this paint layer : 
 							FEditLayerInfo NewEditLayerInfo;
 							NewEditLayerInfo.SourceWeightmapTextureIndex = TextureIndexInVisibleEditLayerTextures;
 							NewEditLayerInfo.SourceWeightmapTextureChannel = ComponentEditLayerAllocation.WeightmapTextureChannel;
-							NewEditLayerInfo.Alpha = (ComponentEditLayerAllocation.LayerInfo == ALandscapeProxy::VisibilityLayer) ? 1.0f : Layer.WeightmapAlpha; // visibility can't affect or be affected by other paint layer weights
-							const bool* BlendSubstractive = Layer.WeightmapLayerAllocationBlend.Find(ComponentEditLayerAllocation.LayerInfo);
+							NewEditLayerInfo.Alpha = (ComponentEditLayerAllocation.LayerInfo == ALandscapeProxy::VisibilityLayer) ? 1.0f : Layer.EditLayer->GetAlphaForTargetType(ELandscapeToolTargetType::Weightmap); // visibility can't affect or be affected by other paint layer weights
+							const bool* BlendSubstractive = Layer.EditLayer->WeightmapLayerAllocationBlend.Find(ComponentEditLayerAllocation.LayerInfo);
 							NewEditLayerInfo.BlendMode = ((BlendSubstractive != nullptr) && *BlendSubstractive) ? ELandscapeEditLayerWeightmapBlendMode::Subtractive : ELandscapeEditLayerWeightmapBlendMode::Additive;
 
 							ComponentPaintLayerRenderInfo.VisibleEditLayerInfos.Add(NewEditLayerInfo);
@@ -9965,13 +9936,14 @@ int32 ALandscape::PerformLayersWeightmapsGlobalMerge(FUpdateLayersContentContext
 
 		for (FLandscapeLayer& Layer : LandscapeEditLayers)
 		{
+			check(Layer.EditLayer != nullptr);
 			int8 CurrentWeightmapToProcessIndex = 0;
 			bool HasFoundWeightmapToProcess = true; // try processing at least once
 
 			TMap<ULandscapeLayerInfoObject*, int32> LayerInfoObjects; // <LayerInfoObj, LayerIndex>
 
 			// Determine if some brush want to write to layer that we have currently no data on
-			if (Layer.bVisible && !InMergeParams.bSkipBrush)
+			if (Layer.EditLayer->IsVisible() && !InMergeParams.bSkipBrush)
 			{
 				for (int32 LayerInfoSettingsIndex = 0; LayerInfoSettingsIndex < Info->Layers.Num(); ++LayerInfoSettingsIndex)
 				{
@@ -10040,7 +10012,7 @@ int32 ALandscape::PerformLayersWeightmapsGlobalMerge(FUpdateLayersContentContext
 			while (HasFoundWeightmapToProcess)
 			{
 				SourceDebugName = EmptyRT->GetName();
-				DestDebugName = FString::Printf(TEXT("Weight: %s Clear WeightmapScratchExtractLayerTextureResource"), *Layer.Name.ToString());
+				DestDebugName = FString::Printf(TEXT("Weight: %s Clear WeightmapScratchExtractLayerTextureResource"), *Layer.EditLayer->GetName().ToString());
 
 				ExecuteCopyLayersTexture({ FLandscapeLayersCopyTextureParams(SourceDebugName, EmptyRT->GameThread_GetRenderTargetResource(), DestDebugName, WeightmapScratchExtractLayerTextureResource) });
 
@@ -10071,7 +10043,7 @@ int32 ALandscape::PerformLayersWeightmapsGlobalMerge(FUpdateLayersContentContext
 							ClearedLayerBit = true;
 
 							SourceDebugName = LandscapeScratchRT1->GetName();
-							DestDebugName = FString::Printf(TEXT("Weight: %s Clear CurrentProcLayerWeightmapAllLayersResource %d, "), *Layer.Name.ToString(), LayerIndex);
+							DestDebugName = FString::Printf(TEXT("Weight: %s Clear CurrentProcLayerWeightmapAllLayersResource %d, "), *Layer.EditLayer->GetName().ToString(), LayerIndex);
 
 							FLandscapeLayersCopyTextureParams& CopyTextureParams = DeferredCopyTextures.Add_GetRef(FLandscapeLayersCopyTextureParams(SourceDebugName, LandscapeScratchRT1->GameThread_GetRenderTargetResource(), DestDebugName, CurrentLayersWeightmapAllMaterialLayersResource));
 							CopyTextureParams.DestAccess = ERHIAccess::UAVMask;
@@ -10085,7 +10057,7 @@ int32 ALandscape::PerformLayersWeightmapsGlobalMerge(FUpdateLayersContentContext
 				// Perform the compute shader
 				if (ComponentsData.Num() > 0)
 				{
-					PrintLayersDebugTextureResource(FString::Printf(TEXT("LS Weight: %s WeightmapScratchTexture %s"), *Layer.Name.ToString(), TEXT("WeightmapScratchTextureResource")), WeightmapScratchExtractLayerTextureResource, 0, false);
+					PrintLayersDebugTextureResource(FString::Printf(TEXT("LS Weight: %s WeightmapScratchTexture %s"), *Layer.EditLayer->GetName().ToString(), TEXT("WeightmapScratchTextureResource")), WeightmapScratchExtractLayerTextureResource, 0, false);
 
 					FLandscapeLayerWeightmapExtractMaterialLayersComputeShaderParameters CSExtractLayersShaderParams;
 					CSExtractLayersShaderParams.AtlasWeightmapsPerLayer = CurrentLayersWeightmapAllMaterialLayersResource;
@@ -10117,10 +10089,10 @@ int32 ALandscape::PerformLayersWeightmapsGlobalMerge(FUpdateLayersContentContext
 					int32 LayerIndex = LayerInfoObject.Value;
 					ULandscapeLayerInfoObject* LayerInfoObj = LayerInfoObject.Key;
 
-					RHI_BREADCRUMB_EVENT_GAMETHREAD_F("Layer", "LS Weight: %s PaintLayer: %s", Layer.Name, LayerInfoObj->LayerName);
+					RHI_BREADCRUMB_EVENT_GAMETHREAD_F("Layer", "LS Weight: %s PaintLayer: %s", Layer.EditLayer->GetName(), LayerInfoObj->LayerName);
 
 					// Copy the layer we are working on
-					SourceDebugName = FString::Printf(TEXT("Weight: %s PaintLayer: %s, CurrentProcLayerWeightmapAllLayersResource"), *Layer.Name.ToString(), *LayerInfoObj->LayerName.ToString());
+					SourceDebugName = FString::Printf(TEXT("Weight: %s PaintLayer: %s, CurrentProcLayerWeightmapAllLayersResource"), *Layer.EditLayer->GetName().ToString(), *LayerInfoObj->LayerName.ToString());
 					DestDebugName = LandscapeScratchRT1->GetName();
 
 					{
@@ -10129,19 +10101,22 @@ int32 ALandscape::PerformLayersWeightmapsGlobalMerge(FUpdateLayersContentContext
 						CopyTextureParams.SourceArrayIndex = LayerIndex;
 						ExecuteCopyLayersTexture({ MoveTemp(CopyTextureParams) });
 					}
-					PrintLayersDebugRT(FString::Printf(TEXT("LS Weight: %s CurrentProcLayerWeightmapAllLayersResource -> Paint Layer RT %s"), *Layer.Name.ToString(), *LandscapeScratchRT1->GetName()), LandscapeScratchRT1, 0, false);
+					PrintLayersDebugRT(FString::Printf(TEXT("LS Weight: %s CurrentProcLayerWeightmapAllLayersResource -> Paint Layer RT %s"), *Layer.EditLayer->GetName().ToString(), *LandscapeScratchRT1->GetName()), LandscapeScratchRT1, 0, false);
+
+					const ELandscapeToolTargetType TargetType = UE::Landscape::IsVisibilityLayer(LayerInfoObj) ? ELandscapeToolTargetType::Visibility : ELandscapeToolTargetType::Weightmap;
 
 					PSShaderParams.ApplyLayerModifiers = true;
-					PSShaderParams.LayerVisible = Layer.bVisible;
-					PSShaderParams.LayerAlpha = LayerInfoObj == ALandscapeProxy::VisibilityLayer ? 1.0f : Layer.WeightmapAlpha; // visibility can't be affected by weight
+					PSShaderParams.LayerVisible = Layer.EditLayer->IsVisible();
+					PSShaderParams.LayerAlpha = Layer.EditLayer->GetAlphaForTargetType(TargetType);
+					checkf((TargetType != ELandscapeToolTargetType::Visibility) || (PSShaderParams.LayerAlpha == 1.0f), TEXT("Visibility can't affect or be affected by other paint layer weights"));
 
-					DrawWeightmapComponentsToRenderTarget(FString::Printf(TEXT("LS Weight: %s PaintLayer: %s, %s += -> %s"), *Layer.Name.ToString(), *LayerInfoObj->LayerName.ToString(), *LandscapeScratchRT1->GetName(), *LandscapeScratchRT2->GetName()),
+					DrawWeightmapComponentsToRenderTarget(FString::Printf(TEXT("LS Weight: %s PaintLayer: %s, %s += -> %s"), *Layer.EditLayer->GetName().ToString(), *LayerInfoObj->LayerName.ToString(), *LandscapeScratchRT1->GetName(), *LandscapeScratchRT2->GetName()),
 						InUpdateLayersContentContext.LandscapeComponentsWeightmapsToRender, LandscapeExtent.Min, LandscapeScratchRT1, nullptr, LandscapeScratchRT2, ERTDrawingType::RTAtlas, true, PSShaderParams, 0);
 
 					PSShaderParams.ApplyLayerModifiers = false;
 
 					// Combined Layer data with current stack
-					SourceDebugName = FString::Printf(TEXT("Weight: %s PaintLayer: %s CombinedProcLayerWeightmap"), *Layer.Name.ToString(), *LayerInfoObj->LayerName.ToString());
+					SourceDebugName = FString::Printf(TEXT("Weight: %s PaintLayer: %s CombinedProcLayerWeightmap"), *Layer.EditLayer->GetName().ToString(), *LayerInfoObj->LayerName.ToString());
 					DestDebugName = LandscapeScratchRT1->GetName();
 
 					{
@@ -10150,27 +10125,27 @@ int32 ALandscape::PerformLayersWeightmapsGlobalMerge(FUpdateLayersContentContext
 						CopyTextureParams.SourceArrayIndex = LayerIndex;
 						ExecuteCopyLayersTexture({ MoveTemp(CopyTextureParams) });
 					}
-					PrintLayersDebugRT(FString::Printf(TEXT("LS Weight: %s CombinedProcLayerWeightmap -> Paint Layer RT %s"), *Layer.Name.ToString(), *LandscapeScratchRT1->GetName()), LandscapeScratchRT1, 0, false);
+					PrintLayersDebugRT(FString::Printf(TEXT("LS Weight: %s CombinedProcLayerWeightmap -> Paint Layer RT %s"), *Layer.EditLayer->GetName().ToString(), *LandscapeScratchRT1->GetName()), LandscapeScratchRT1, 0, false);
 
 					// Combine with current status and copy back to the combined 2d resource array
 					PSShaderParams.OutputAsSubstractive = false;
 
 					if (!bFirstLayer)
 					{
-						const bool* BlendSubstractive = Layer.WeightmapLayerAllocationBlend.Find(LayerInfoObj);
+						const bool* BlendSubstractive = Layer.EditLayer->GetWeightmapLayerAllocationBlend().Find(LayerInfoObj);
 						PSShaderParams.OutputAsSubstractive = BlendSubstractive != nullptr ? *BlendSubstractive : false;
 					}
 
-					DrawWeightmapComponentsToRenderTarget(FString::Printf(TEXT("LS Weight: %s PaintLayer: %s, %s += -> Combined %s"), *Layer.Name.ToString(), *LayerInfoObj->LayerName.ToString(), *LandscapeScratchRT2->GetName(), *LandscapeScratchRT3->GetName()),
+					DrawWeightmapComponentsToRenderTarget(FString::Printf(TEXT("LS Weight: %s PaintLayer: %s, %s += -> Combined %s"), *Layer.EditLayer->GetName().ToString(), *LayerInfoObj->LayerName.ToString(), *LandscapeScratchRT2->GetName(), *LandscapeScratchRT3->GetName()),
 						InUpdateLayersContentContext.LandscapeComponentsWeightmapsToRender, LandscapeExtent.Min, LandscapeScratchRT2, bFirstLayer ? nullptr : LandscapeScratchRT1, LandscapeScratchRT3, ERTDrawingType::RTAtlasToNonAtlas, true, PSShaderParams, 0);
 
 					PSShaderParams.OutputAsSubstractive = false;
 
-					SourceDebugName = FString::Printf(TEXT("Weight: %s PaintLayer: %s %s"), *Layer.Name.ToString(), *LayerInfoObj->LayerName.ToString(), *LandscapeScratchRT3->GetName());
+					SourceDebugName = FString::Printf(TEXT("Weight: %s PaintLayer: %s %s"), *Layer.EditLayer->GetName().ToString(), *LayerInfoObj->LayerName.ToString(), *LandscapeScratchRT3->GetName());
 					DestDebugName = TEXT("CombinedProcLayerWeightmap");
 
 					// Handle brush blending
-					if (Layer.bVisible && !InMergeParams.bSkipBrush)
+					if (Layer.EditLayer->IsVisible() && !InMergeParams.bSkipBrush)
 					{
 						// Draw each brush				
 						for (int32 i = 0; i < Layer.Brushes.Num(); ++i)
@@ -10193,17 +10168,17 @@ int32 ALandscape::PerformLayersWeightmapsGlobalMerge(FUpdateLayersContentContext
 
 							INC_DWORD_STAT(STAT_LandscapeLayersRegenerateDrawCalls); // Brush RenderInitialize
 
-							PrintLayersDebugRT(FString::Printf(TEXT("LS Weight: %s %s -> Brush %s"), *Layer.Name.ToString(), *LandscapeBrush->GetName(), *BrushOutputRT->GetName()), BrushOutputRT);
+							PrintLayersDebugRT(FString::Printf(TEXT("LS Weight: %s %s -> Brush %s"), *Layer.EditLayer->GetName().ToString(), *LandscapeBrush->GetName(), *BrushOutputRT->GetName()), BrushOutputRT);
 
 							// Copy result back if brush did not edit things in place.
 
 							// Resolve back to Combined heightmap (it's unlikely, but possible that the brush returns the same RT as input and output, if it did various operations on it, in which case the copy is useless) :
 							if (BrushOutputRT != LandscapeScratchRT3)
 							{
-								SourceDebugName = FString::Printf(TEXT("Weight: %s PaintLayer: %s Brush: %s"), *Layer.Name.ToString(), *LayerInfoObj->LayerName.ToString(), *BrushOutputRT->GetName());
+								SourceDebugName = FString::Printf(TEXT("Weight: %s PaintLayer: %s Brush: %s"), *Layer.EditLayer->GetName().ToString(), *LayerInfoObj->LayerName.ToString(), *BrushOutputRT->GetName());
 								DestDebugName = LandscapeScratchRT3->GetName();
 								ExecuteCopyLayersTexture({ FLandscapeLayersCopyTextureParams(SourceDebugName, BrushOutputRT->GameThread_GetRenderTargetResource(), DestDebugName, LandscapeScratchRT3->GameThread_GetRenderTargetResource()) });
-								PrintLayersDebugRT(FString::Printf(TEXT("LS Weight: %s Component %s += -> Combined %s"), *Layer.Name.ToString(), *BrushOutputRT->GetName(), *LandscapeScratchRT3->GetName()), LandscapeScratchRT3);
+								PrintLayersDebugRT(FString::Printf(TEXT("LS Weight: %s Component %s += -> Combined %s"), *Layer.EditLayer->GetName().ToString(), *BrushOutputRT->GetName(), *LandscapeScratchRT3->GetName()), LandscapeScratchRT3);
 							}
 						}
 
@@ -10224,17 +10199,17 @@ int32 ALandscape::PerformLayersWeightmapsGlobalMerge(FUpdateLayersContentContext
 								// Same conditional copy as for blueprint brushes
 								if (BrushOutputRT != LandscapeScratchRT3)
 								{
-									SourceDebugName = FString::Printf(TEXT("Weight: %s PaintLayer: %s Brush: %s"), *Layer.Name.ToString(), *LayerInfoObj->LayerName.ToString(), *BrushOutputRT->GetName());
+									SourceDebugName = FString::Printf(TEXT("Weight: %s PaintLayer: %s Brush: %s"), *Layer.EditLayer->GetName().ToString(), *LayerInfoObj->LayerName.ToString(), *BrushOutputRT->GetName());
 									DestDebugName = LandscapeScratchRT3->GetName();
 									ExecuteCopyLayersTexture({ FLandscapeLayersCopyTextureParams(SourceDebugName, BrushOutputRT->GameThread_GetRenderTargetResource(), DestDebugName, LandscapeScratchRT3->GameThread_GetRenderTargetResource()) });
-									PrintLayersDebugRT(FString::Printf(TEXT("LS Weight: %s Component %s += -> Combined %s"), *Layer.Name.ToString(), *BrushOutputRT->GetName(), *LandscapeScratchRT3->GetName()), LandscapeScratchRT3);
+									PrintLayersDebugRT(FString::Printf(TEXT("LS Weight: %s Component %s += -> Combined %s"), *Layer.EditLayer->GetName().ToString(), *BrushOutputRT->GetName(), *LandscapeScratchRT3->GetName()), LandscapeScratchRT3);
 								}
 							}
 						}
 
-						PrintLayersDebugRT(FString::Printf(TEXT("LS Weight: %s CombinedPostBrushProcLayerWeightmap -> Paint Layer RT %s"), *Layer.Name.ToString(), *LandscapeScratchRT3->GetName()), LandscapeScratchRT3, 0, false);
+						PrintLayersDebugRT(FString::Printf(TEXT("LS Weight: %s CombinedPostBrushProcLayerWeightmap -> Paint Layer RT %s"), *Layer.EditLayer->GetName().ToString(), *LandscapeScratchRT3->GetName()), LandscapeScratchRT3, 0, false);
 
-						SourceDebugName = FString::Printf(TEXT("Weight: %s PaintLayer: %s %s"), *Layer.Name.ToString(), *LayerInfoObj->LayerName.ToString(), *LandscapeScratchRT3->GetName());
+						SourceDebugName = FString::Printf(TEXT("Weight: %s PaintLayer: %s %s"), *Layer.EditLayer->GetName().ToString(), *LayerInfoObj->LayerName.ToString(), *LandscapeScratchRT3->GetName());
 						DestDebugName = TEXT("CombinedProcLayerWeightmap");
 
 						FLandscapeLayersCopyTextureParams CopyTextureParams(SourceDebugName, LandscapeScratchRT3->GameThread_GetRenderTargetResource(), DestDebugName, CombinedLayersWeightmapAllMaterialLayersResource);
@@ -11161,7 +11136,7 @@ void ALandscape::RequestLayersInitialization(bool bInRequestContentUpdate)
 
 void ALandscape::RequestSplineLayerUpdate()
 {
-	if (HasLayersContent() && FindLayerOfType(ULandscapeEditLayerSplines::StaticClass()) != nullptr)
+	if (HasLayersContent() && FindLayerOfTypeConst(ULandscapeEditLayerSplines::StaticClass()) != nullptr)
 	{
 		bSplineLayerUpdateRequested = true;
 	}
@@ -11374,7 +11349,7 @@ void ALandscape::MonitorLandscapeEdModeChanges()
 	}
 
 	const FLandscapeLayer* SelectedLayer = LandscapeEdMode ? LandscapeEdMode->GetLandscapeSelectedLayer() : nullptr;
-	FGuid NewSelectedLayer = SelectedLayer && SelectedLayer->bVisible ? SelectedLayer->Guid : FGuid();
+	FGuid NewSelectedLayer = SelectedLayer && SelectedLayer->EditLayer->IsVisible() ? SelectedLayer->EditLayer->GetGuid() : FGuid();
 	if (LandscapeEdModeInfo.SelectedLayer != NewSelectedLayer)
 	{
 		LandscapeEdModeInfo.SelectedLayer = NewSelectedLayer;
@@ -11586,7 +11561,7 @@ void ALandscape::UpdateLayersContent(bool bInWaitForStreaming, bool bInSkipMonit
 
 	if (bSplineLayerUpdateRequested)
 	{
-		if (const FLandscapeLayer* SplinesLayer = FindLayerOfType(ULandscapeEditLayerSplines::StaticClass()))
+		if (const FLandscapeLayer* SplinesLayer = FindLayerOfTypeConst(ULandscapeEditLayerSplines::StaticClass()))
 		{
 			// We need the spline layer resources to all be ready before updating it:
 			if (!PrepareLayersTextureResources({ *SplinesLayer }, bInWaitForStreaming))
@@ -12212,7 +12187,7 @@ bool ALandscapeProxy::AddLayer(const FGuid& InLayerGuid)
 		if ((Component != nullptr) && !Component->GetLayerData(InLayerGuid))
 		{
 			const FLandscapeLayer* EditLayer = GetLandscapeActor() ? GetLandscapeActor()->GetLayerConst(InLayerGuid) : nullptr;
-			Component->AddLayerData(InLayerGuid, FLandscapeLayerComponentData(EditLayer ? EditLayer->Name : FName()));
+			Component->AddLayerData(InLayerGuid, FLandscapeLayerComponentData(EditLayer ? EditLayer->EditLayer->GetName() : FName()));
 			bModified = true;
 		}
 	}
@@ -12387,111 +12362,86 @@ bool ALandscape::IsUpToDate() const
 #if WITH_EDITOR
 bool ALandscape::IsLayerNameUnique(const FName& InName) const
 {
-	return Algo::CountIf(LandscapeEditLayers, [InName](const FLandscapeLayer& Layer) { return (Layer.Name == InName); }) == 0;
+	return Algo::CountIf(LandscapeEditLayers, [InName](const FLandscapeLayer& Layer) { return (Layer.EditLayer->GetName() == InName); }) == 0;
 }
 
-void ALandscape::SetLayerName(int32 InLayerIndex, const FName& InName)
+void ALandscape::OnEditLayerDataChanged(const FOnLandscapeEditLayerDataChangedParams& InParams)
 {
-	FLandscapeLayer* Layer = GetLayerInternal(InLayerIndex);
-	ULandscapeInfo* LandscapeInfo = GetLandscapeInfo();
-	if (!LandscapeInfo || !Layer || Layer->Name == InName)
+	const ULandscapeSettings* LandscapeSettings = GetDefault<ULandscapeSettings>();
+	check(LandscapeSettings != nullptr);
+
+	const bool bAllowLandscapeUpdate = (InParams.PropertyChangedEvent.ChangeType != EPropertyChangeType::Interactive) || LandscapeSettings->GetShouldUpdateEditLayersDuringInteractiveChanges();
+	if (InParams.bRequiresLandscapeUpdate && bAllowLandscapeUpdate)
 	{
-		return;
+		RequestLayersContentUpdateForceAll(ELandscapeLayerUpdateMode::Update_All, InParams.bUserTriggered);
 	}
+}
 
-	if (!IsLayerNameUnique(InName))
+// Deprecated
+void ALandscape::SetLayerName(int32 InLayerIndex, const FName& InName)
+{
+	const ULandscapeInfo* LandscapeInfo = GetLandscapeInfo();
+	if (FLandscapeLayer* Layer = GetLayerInternal(InLayerIndex); (Layer != nullptr) && (LandscapeInfo != nullptr))
 	{
-		return;
+		Layer->EditLayer->SetName(InName, /*bInModify = */true);;
 	}
-
-	Modify();
-	Layer->Name = InName;
 }
 
+// Deprecated
 float ALandscape::GetLayerAlpha(int32 InLayerIndex, bool bInHeightmap) const
 {
-	// TODO [jonathan.bard] : replace by proper getter virtual GetLayerAlpha implementation on ULandscapeEditLayerSplines
-	const FLandscapeLayer* SplinesReservedLayer = FindLayerOfType(ULandscapeEditLayerSplines::StaticClass());
-	const FLandscapeLayer* Layer = GetLayerConst(InLayerIndex);
-	if (Layer && SplinesReservedLayer != Layer)
+	if (const FLandscapeLayer* Layer = GetLayerConst(InLayerIndex))
 	{
-		return GetClampedLayerAlpha(bInHeightmap ? Layer->HeightmapAlpha : Layer->WeightmapAlpha, bInHeightmap);
+		check(Layer->EditLayer != nullptr);
+		return Layer->EditLayer->GetAlphaForTargetType(bInHeightmap ? ELandscapeToolTargetType::Heightmap : ELandscapeToolTargetType::Weightmap);
 	}
+
 	return 1.0f;
 }
 
+// Deprecated
 float ALandscape::GetClampedLayerAlpha(float InAlpha, bool bInHeightmap) const
 {
 	float AlphaClamped = FMath::Clamp<float>(InAlpha, bInHeightmap ? -1.f : 0.f, 1.f);
 	return AlphaClamped;
 }
 
+// Deprecated
 void ALandscape::SetLayerAlpha(int32 InLayerIndex, float InAlpha, bool bInHeightmap)
 {
-	FLandscapeLayer* Layer = GetLayerInternal(InLayerIndex);
-	ULandscapeInfo* LandscapeInfo = GetLandscapeInfo();
-	if (!LandscapeInfo || !Layer)
-	{
-		return;
-	}
-	const float InAlphaClamped = GetClampedLayerAlpha(InAlpha, bInHeightmap);
-	float& LayerAlpha = bInHeightmap ? Layer->HeightmapAlpha : Layer->WeightmapAlpha;
-	if (LayerAlpha == InAlphaClamped)
+	const ULandscapeInfo* LandscapeInfo = GetLandscapeInfo();
+	if (FLandscapeLayer* Layer = GetLayerInternal(InLayerIndex); (Layer != nullptr) && (LandscapeInfo != nullptr))
 	{
-		return;
+		Layer->EditLayer->SetAlphaForTargetType(bInHeightmap ? ELandscapeToolTargetType::Heightmap : ELandscapeToolTargetType::Weightmap, InAlpha, /*bInModify = */true, EPropertyChangeType::ValueSet);
 	}
-
-	Modify();
-	LayerAlpha = InAlphaClamped;
-	RequestLayersContentUpdateForceAll();
 }
 
+// Deprecated
 void ALandscape::SetLayerVisibility(int32 InLayerIndex, bool bInVisible, bool bInForIntermediateRender)
 {
-	FLandscapeLayer* Layer = GetLayerInternal(InLayerIndex);
-	ULandscapeInfo* LandscapeInfo = GetLandscapeInfo();
-	if (!LandscapeInfo || !Layer || Layer->bVisible == bInVisible)
-	{
-		return;
-	}
-
-	if (!bInForIntermediateRender)
-	{
-		Modify();
-	}
-
-	Layer->bVisible = bInVisible;
-
-	if (!bInForIntermediateRender)
+	const ULandscapeInfo* LandscapeInfo = GetLandscapeInfo();
+	if (FLandscapeLayer* Layer = GetLayerInternal(InLayerIndex); (Layer != nullptr) && (LandscapeInfo != nullptr))
 	{
-		RequestLayersContentUpdateForceAll();
+		Layer->EditLayer->SetVisible(bInVisible, /*bInModify = */bInForIntermediateRender);
 	}
 }
 
+// Deprecated 
 void ALandscape::SetLayerLocked(int32 InLayerIndex, bool bLocked)
 {
-	FLandscapeLayer* Layer = GetLayerInternal(InLayerIndex);
-	if (!Layer || Layer->bLocked == bLocked)
+	if (FLandscapeLayer* Layer = GetLayerInternal(InLayerIndex); (Layer != nullptr))
 	{
-		return;
+		Layer->EditLayer->SetLocked(bLocked, /*bInModify = */true);
 	}
-
-	Modify();
-	Layer->bLocked = bLocked;
 }
 
+// Deprecated
 void ALandscape::SetLayerBlendMode(int32 InLayerIndex, ELandscapeBlendMode InBlendMode)
 {
-	FLandscapeLayer* Layer = GetLayerInternal(InLayerIndex);
-	if (!Layer || Layer->BlendMode == InBlendMode)
-	{
-		return;
-	}
-
-	Modify();
-	Layer->BlendMode = InBlendMode;
+	// ULandscapeEditLayerBase no longer stores blend mode, override getter method on derived classes
 }
 
+// Deprecated
 uint8 ALandscape::GetLayerCount() const
 {
 	return static_cast<uint8>(LandscapeEditLayers.Num());
@@ -12524,6 +12474,11 @@ const FLandscapeLayer* ALandscape::GetLayer(const FName& InLayerName) const
 	return GetLayerConst(InLayerName);
 }
 
+TArrayView<const FLandscapeLayer> ALandscape::GetLayersConst() const
+{ 
+	return MakeArrayView(LandscapeEditLayers); 
+}
+
 const FLandscapeLayer* ALandscape::GetLayerConst(int32 InLayerIndex) const
 {
 	if (LandscapeEditLayers.IsValidIndex(InLayerIndex))
@@ -12535,22 +12490,111 @@ const FLandscapeLayer* ALandscape::GetLayerConst(int32 InLayerIndex) const
 
 int32 ALandscape::GetLayerIndex(const FGuid& InLayerGuid) const
 {
-	return LandscapeEditLayers.IndexOfByPredicate([&InLayerGuid](const FLandscapeLayer& Other) { return Other.Guid == InLayerGuid; });
+	return LandscapeEditLayers.IndexOfByPredicate([&InLayerGuid](const FLandscapeLayer& Other) { return Other.EditLayer->GetGuid() == InLayerGuid; });
 }
 
 const FLandscapeLayer* ALandscape::GetLayerConst(const FGuid& InLayerGuid) const
 {
-	return LandscapeEditLayers.FindByPredicate([&InLayerGuid](const FLandscapeLayer& Other) { return Other.Guid == InLayerGuid; });
+	return LandscapeEditLayers.FindByPredicate([&InLayerGuid](const FLandscapeLayer& Other) { return Other.EditLayer->GetGuid() == InLayerGuid; });
 }
 
 const FLandscapeLayer* ALandscape::GetLayerConst(const FName& InLayerName) const
 {
-	return LandscapeEditLayers.FindByPredicate([InLayerName](const FLandscapeLayer& Layer) { return Layer.Name == InLayerName; });
+	return LandscapeEditLayers.FindByPredicate([InLayerName](const FLandscapeLayer& Layer) { return Layer.EditLayer->GetName() == InLayerName; });
+}
+
+const ULandscapeEditLayerBase* ALandscape::GetEditLayerConst(int32 InLayerIndex) const
+{
+	if (const FLandscapeLayer* Layer = GetLayerConst(InLayerIndex))
+	{
+		check(Layer->EditLayer != nullptr);
+		return Layer->EditLayer;
+	}
+	return nullptr;
+}
+
+const ULandscapeEditLayerBase* ALandscape::GetEditLayerConst(const FGuid& InLayerGuid) const
+{
+	if (const FLandscapeLayer* Layer = GetLayerConst(InLayerGuid))
+	{
+		check(Layer->EditLayer != nullptr);
+		return Layer->EditLayer;
+	}
+	return nullptr;
+}
+
+const ULandscapeEditLayerBase* ALandscape::GetEditLayerConst(const FName& InLayerName) const
+{
+	if (const FLandscapeLayer* Layer = GetLayerConst(InLayerName))
+	{
+		check(Layer->EditLayer != nullptr);
+		return Layer->EditLayer;
+	}
+	return nullptr;
+}
+
+ULandscapeEditLayerBase* ALandscape::GetEditLayer(int32 InLayerIndex) const
+{
+	return const_cast<ULandscapeEditLayerBase*>(GetEditLayerConst(InLayerIndex));
+}
+
+ULandscapeEditLayerBase* ALandscape::GetEditLayer(const FGuid& InLayerGuid) const
+{
+	return const_cast<ULandscapeEditLayerBase*>(GetEditLayerConst(InLayerGuid));
+}
+
+ULandscapeEditLayerBase* ALandscape::GetEditLayer(const FName& InLayerName) const
+{
+	return const_cast<ULandscapeEditLayerBase*>(GetEditLayerConst(InLayerName));
+}
+
+const ULandscapeEditLayerBase* ALandscape::FindEditLayerOfTypeConst(const TSubclassOf<ULandscapeEditLayerBase>& InLayerClass) const
+{
+	if (const FLandscapeLayer* Layer = FindLayerOfTypeConst(InLayerClass))
+	{
+		check(Layer->EditLayer != nullptr);
+		return Layer->EditLayer;
+	}
+	return nullptr;
+}
+
+ULandscapeEditLayerBase* ALandscape::FindEditLayerOfType(const TSubclassOf<ULandscapeEditLayerBase>& InLayerClass) const
+{
+	if (const FLandscapeLayer* Layer = FindLayerOfTypeConst(InLayerClass))
+	{
+		check(Layer->EditLayer != nullptr);
+		return const_cast<FLandscapeLayer*>(Layer)->EditLayer;
+	}
+	return nullptr;
+}
+
+TArray<const ULandscapeEditLayerBase*> ALandscape::GetEditLayersOfTypeConst(const TSubclassOf<ULandscapeEditLayerBase>& InLayerClass) const
+{
+	TArray<const ULandscapeEditLayerBase*> EditLayers;
+	EditLayers.Reserve(LandscapeEditLayers.Num());
+	for (const FLandscapeLayer& InLayer : LandscapeEditLayers)
+	{
+		check(InLayer.EditLayer != nullptr);
+		EditLayers.Add(InLayer.EditLayer);
+	}
+	return EditLayers;
+}
+
+TArray<ULandscapeEditLayerBase*> ALandscape::GetEditLayersOfType(const TSubclassOf<ULandscapeEditLayerBase>& InLayerClass) const
+{
+	TArray<ULandscapeEditLayerBase*> EditLayers;
+	EditLayers.Reserve(LandscapeEditLayers.Num());
+	for (const FLandscapeLayer& InLayer : LandscapeEditLayers)
+	{
+		check(InLayer.EditLayer != nullptr);
+		EditLayers.Add(const_cast<ULandscapeEditLayerBase*>(InLayer.EditLayer.Get()));
+	}
+	return EditLayers;
 }
 
 int32 ALandscape::GetLayerIndex(FName InLayerName) const
 {
-	return LandscapeEditLayers.IndexOfByPredicate([InLayerName](const FLandscapeLayer& Layer) { return Layer.Name == InLayerName; });
+	return LandscapeEditLayers.IndexOfByPredicate([InLayerName](const FLandscapeLayer& Layer) { return Layer.EditLayer->GetName() == InLayerName; });
 }
 
 // Deprecated
@@ -12573,12 +12617,12 @@ void ALandscape::ForEachLayerConst(TFunctionRef<bool(const FLandscapeLayer&)> Fn
 	}
 }
 
-const FLandscapeLayer* ALandscape::FindLayerOfType(const TSubclassOf<ULandscapeEditLayerBase>& InLayerClass) const
+const FLandscapeLayer* ALandscape::FindLayerOfTypeConst(const TSubclassOf<ULandscapeEditLayerBase>& InLayerClass) const
 {
 	return LandscapeEditLayers.FindByPredicate([&InLayerClass](const FLandscapeLayer& InLayer) { check(InLayer.EditLayer != nullptr); return InLayer.EditLayer->GetClass()->IsChildOf(InLayerClass); });
 }
 
-TArray<const FLandscapeLayer*> ALandscape::GetLayersOfType(const TSubclassOf<ULandscapeEditLayerBase>& InLayerClass) const
+TArray<const FLandscapeLayer*> ALandscape::GetLayersOfTypeConst(const TSubclassOf<ULandscapeEditLayerBase>& InLayerClass) const
 {
 	TArray<const FLandscapeLayer*> Result;
 	Result.Reserve(LandscapeEditLayers.Num());
@@ -12615,7 +12659,7 @@ bool ALandscape::DeleteLayer(int32 InLayerIndex)
 	}
 
 	Modify();
-	FGuid LayerGuid = LayerStruct->Guid;
+	FGuid LayerGuid = LayerStruct->EditLayer->GetGuid();
 
 	// Clean up Weightmap usage in LandscapeProxies
 	if (ULandscapeInfo* LandscapeInfo = GetLandscapeInfo())
@@ -12629,16 +12673,17 @@ bool ALandscape::DeleteLayer(int32 InLayerIndex)
 
 	// We're about to remove the layer from our list, which will invalidate our LayerStruct pointer.
 	// We'll need to call OnLayerRemoved afterward, though, so keep pointer to the UObject.
-	ULandscapeEditLayerBase* Layer = LayerStruct->EditLayer;
+	ULandscapeEditLayerBase* EditLayer = LayerStruct->EditLayer;
+	check(EditLayer != nullptr);
 
 	// Remove layer from list
 	LandscapeEditLayers.RemoveAt(InLayerIndex);
 	LayerStruct = nullptr;
 
-	if (Layer)
-	{
-		Layer->OnLayerRemoved();
-	}
+	EditLayer->OnLayerRemoved();
+
+	// Unregister from data change events on the edit layer so that we can update the landscape accordingly : 
+	EditLayer->OnLayerDataChanged().RemoveAll(this);
 
 	// Request Update
 	RequestLayersContentUpdateForceAll();
@@ -12654,8 +12699,8 @@ void ALandscape::CollapseLayer(int32 InLayerIndex)
 	TArray<bool> BackupBrushVisibility;
 	for (int32 i = 0; i < LandscapeEditLayers.Num(); ++i)
 	{
-		BackupVisibility.Add(LandscapeEditLayers[i].bVisible);
-		LandscapeEditLayers[i].bVisible = i == InLayerIndex || i == InLayerIndex - 1;
+		BackupVisibility.Add(LandscapeEditLayers[i].EditLayer->IsVisible());
+		LandscapeEditLayers[i].EditLayer->SetVisible(i == InLayerIndex || i == InLayerIndex - 1, /*bInModify = */ true);
 	}
 
 	for (int32 i = 0; i < LandscapeEditLayers[InLayerIndex].Brushes.Num(); ++i)
@@ -12680,7 +12725,7 @@ void ALandscape::CollapseLayer(int32 InLayerIndex)
 		DataInterface.SetShouldDirtyPackage(true);
 
 		TSet<UTexture2D*> ProcessedHeightmaps;
-		FScopedSetLandscapeEditingLayer ScopeEditingLayer(this, LandscapeEditLayers[InLayerIndex - 1].Guid);
+		FScopedSetLandscapeEditingLayer ScopeEditingLayer(this, LandscapeEditLayers[InLayerIndex - 1].EditLayer->GetGuid());
 		GetLandscapeInfo()->ForAllLandscapeComponents([&](ULandscapeComponent* LandscapeComponent)
 		{
 			SlowTask.EnterProgressFrame(1.f);
@@ -12704,7 +12749,7 @@ void ALandscape::CollapseLayer(int32 InLayerIndex)
 
 	for (int32 i = 0; i < LandscapeEditLayers.Num(); ++i)
 	{
-		LandscapeEditLayers[i].bVisible = BackupVisibility[i];
+		LandscapeEditLayers[i].EditLayer->SetVisible(BackupVisibility[i], /*bInModify = */ true);
 	}
 
 	DeleteLayer(InLayerIndex);
@@ -12717,7 +12762,7 @@ void ALandscape::GetUsedPaintLayers(int32 InLayerIndex, TArray<ULandscapeLayerIn
 	const FLandscapeLayer* Layer = GetLayerConst(InLayerIndex);
 	if (Layer)
 	{
-		GetUsedPaintLayers(Layer->Guid, OutUsedLayerInfos);
+		GetUsedPaintLayers(Layer->EditLayer->GetGuid(), OutUsedLayerInfos);
 	}
 }
 
@@ -12737,7 +12782,7 @@ void ALandscape::ClearPaintLayer(int32 InLayerIndex, ULandscapeLayerInfoObject*
 	const FLandscapeLayer* Layer = GetLayerConst(InLayerIndex);
 	if (Layer)
 	{
-		ClearPaintLayer(Layer->Guid, InLayerInfo);
+		ClearPaintLayer(Layer->EditLayer->GetGuid(), InLayerInfo);
 	}
 }
 
@@ -12769,7 +12814,7 @@ void ALandscape::ClearLayer(int32 InLayerIndex, TSet<TObjectPtr<ULandscapeCompon
 	const FLandscapeLayer* Layer = GetLayerConst(InLayerIndex);
 	if (Layer)
 	{
-		ClearLayer(Layer->Guid, InComponents, InClearMode);
+		ClearLayer(Layer->EditLayer->GetGuid(), InComponents, InClearMode);
 	}
 }
 
@@ -12785,7 +12830,7 @@ void ALandscape::ClearLayer(const FGuid& InLayerGuid, TSet<TObjectPtr<ULandscape
 	}
 
 	Modify(bMarkPackageDirty);
-	FScopedSetLandscapeEditingLayer Scope(this, Layer->Guid, [this] { RequestLayersContentUpdate(ELandscapeLayerUpdateMode::Update_All); });
+	FScopedSetLandscapeEditingLayer Scope(this, Layer->EditLayer->GetGuid(), [this] { RequestLayersContentUpdate(ELandscapeLayerUpdateMode::Update_All); });
 
 	TArray<uint16> NewHeightData;
 	NewHeightData.AddZeroed(FMath::Square(ComponentSizeQuads + 1));
@@ -12800,7 +12845,7 @@ void ALandscape::ClearLayer(const FGuid& InLayerGuid, TSet<TObjectPtr<ULandscape
 
 	if (InClearMode & ELandscapeClearMode::Clear_Heightmap)
 	{
-		if (Layer->BlendMode == LSBM_AlphaBlend)
+		if (Layer->EditLayer->GetBlendMode() == LSBM_AlphaBlend)
 		{
 			NewHeightAlphaBlendData.Init(MAX_uint16, FMath::Square(ComponentSizeQuads + 1));
 			NewHeightFlagsData.AddZeroed(FMath::Square(ComponentSizeQuads + 1));
@@ -12868,24 +12913,14 @@ void ALandscape::ShowOnlySelectedLayer(int32 InLayerIndex)
 	const FLandscapeLayer* VisibleLayer = GetLayerConst(InLayerIndex);
 	if (VisibleLayer)
 	{
-		bool bModified = false;
 		for (FLandscapeLayer& Layer : LandscapeEditLayers)
 		{
 			bool bDesiredVisible = (&Layer == VisibleLayer);
-			if (Layer.bVisible != bDesiredVisible)
-			{
-				if (!bModified)
+			if (Layer.EditLayer->IsVisible() != bDesiredVisible)
 				{
-					Modify();
-					bModified = true;
-		}
-				Layer.bVisible = bDesiredVisible;
+				Layer.EditLayer->SetVisible(bDesiredVisible, /* bInModify = */ true);
 			}
 		}
-		if (bModified)
-		{
-		RequestLayersContentUpdateForceAll();
-	}
 }
 }
 
@@ -12893,23 +12928,13 @@ void ALandscape::ShowAllLayers()
 {
 	if (LandscapeEditLayers.Num() > 0)
 	{
-		bool bModified = false;
 		for (FLandscapeLayer& Layer : LandscapeEditLayers)
 		{
-			if (Layer.bVisible != true)
-			{
-				if (!bModified)
+			if (Layer.EditLayer->IsVisible() != true)
 				{
-					Modify();
-					bModified = true;
-				}
-			Layer.bVisible = true;
+				Layer.EditLayer->SetVisible(/*bInVisible = */ true, /*bInModify = */ true);
 		}
 		}
-		if (bModified)
-		{
-		RequestLayersContentUpdateForceAll();
-	}
 }
 }
 
@@ -12921,13 +12946,13 @@ void ALandscape::SetLandscapeSplinesReservedLayer(int32 InLayerIndex)
 // Deprecated
 const FLandscapeLayer* ALandscape::GetLandscapeSplinesReservedLayer() const
 {
-	return FindLayerOfType(ULandscapeEditLayerSplines::StaticClass());
+	return FindLayerOfTypeConst(ULandscapeEditLayerSplines::StaticClass());
 }
 
 // Deprecated
 FLandscapeLayer* ALandscape::GetLandscapeSplinesReservedLayer()
 {
-	return const_cast<FLandscapeLayer*>(FindLayerOfType(ULandscapeEditLayerSplines::StaticClass()));
+	return const_cast<FLandscapeLayer*>(FindLayerOfTypeConst(ULandscapeEditLayerSplines::StaticClass()));
 }
 
 LANDSCAPE_API extern bool GDisableUpdateLandscapeMaterialInstances;
@@ -12985,8 +13010,8 @@ void ALandscape::UpdateLandscapeSplines(const FGuid& InTargetLayer, bool bInUpda
 	TRACE_CPUPROFILER_EVENT_SCOPE(LandscapeLayers_UpdateLandscapeSplines);
 	check(CanHaveLayersContent());
 	ULandscapeInfo* LandscapeInfo = GetLandscapeInfo();
-	const FLandscapeLayer* SplinesLayer = FindLayerOfType(ULandscapeEditLayerSplines::StaticClass());
-	FGuid TargetLayerGuid = (SplinesLayer != nullptr) ? SplinesLayer->Guid : InTargetLayer;
+	const FLandscapeLayer* SplinesLayer = FindLayerOfTypeConst(ULandscapeEditLayerSplines::StaticClass());
+	FGuid TargetLayerGuid = (SplinesLayer != nullptr) ? SplinesLayer->EditLayer->GetGuid() : InTargetLayer;
 	const FLandscapeLayer* TargetLayer = GetLayerConst(TargetLayerGuid);
 	if (LandscapeInfo && TargetLayer)
 	{
@@ -13023,7 +13048,7 @@ void ALandscape::UpdateLandscapeSplines(const FGuid& InTargetLayer, bool bInUpda
 
 			// Clear layers without affecting weightmap allocations
 			const bool bMarkPackageDirty = false;
-			ClearLayer(SplinesLayer->Guid, (!bInForceUpdateAllCompoments && LandscapeSplinesAffectedComponents.Num()) ? &LandscapeSplinesAffectedComponents : nullptr, Clear_All, bMarkPackageDirty);
+			ClearLayer(SplinesLayer->EditLayer->GetGuid(), (!bInForceUpdateAllCompoments && LandscapeSplinesAffectedComponents.Num()) ? &LandscapeSplinesAffectedComponents : nullptr, Clear_All, bMarkPackageDirty);
 			LandscapeSplinesAffectedComponents.Empty();
 			ModifiedComponent = &LandscapeSplinesAffectedComponents;
 			// For now, in Landscape Layer System Mode with a reserved layer for splines, we always update all the splines since we clear the whole layer first
@@ -13139,14 +13164,14 @@ FLandscapeLayer* ALandscape::DuplicateLayerAndMoveBrushes(const FLandscapeLayer&
 
 	if (IsMaxLayersReached())
 	{
-		UE_LOG(LogLandscape, Warning, TEXT("Cannot duplicate layer : %s as the max number of layers (%i) has been reached"), *InOtherLayer.Name.ToString(), GetDefault<ULandscapeSettings>()->MaxNumberOfLayers);
+		UE_LOG(LogLandscape, Warning, TEXT("Cannot duplicate layer : %s as the max number of layers (%i) has been reached"), *InOtherLayer.EditLayer->GetName().ToString(), GetDefault<ULandscapeSettings>()->MaxNumberOfLayers);
 		return nullptr;
 	}
 
 	Modify();
 
 	FLandscapeLayer NewLayer(InOtherLayer);
-	NewLayer.Guid = FGuid::NewGuid();
+	NewLayer.EditLayer->SetGuid(FGuid::NewGuid(), /*bInModify =*/true);
 
 	// Duplicate the internal edit layer object by hand : 
 	check(InOtherLayer.EditLayer != nullptr);
@@ -13184,8 +13209,10 @@ int32 ALandscape::CreateLayer(FName InName, const TSubclassOf<ULandscapeEditLaye
 	const UClass* EditLayerClass = (InEditLayerClass.Get() != nullptr) ? InEditLayerClass.Get() : ULandscapeEditLayer::StaticClass();
 	int32 LayerIndex = LandscapeEditLayers.Emplace();
 	FLandscapeLayer& NewLayer = LandscapeEditLayers[LayerIndex];
-	NewLayer.Name = GenerateUniqueLayerName(InName);
-	NewLayer.EditLayer = NewObject<ULandscapeEditLayerBase>(this, EditLayerClass, MakeUniqueObjectName(this, EditLayerClass));
+	NewLayer.EditLayer = NewObject<ULandscapeEditLayerBase>(this, EditLayerClass, MakeUniqueObjectName(this, EditLayerClass), RF_Transactional);
+	NewLayer.EditLayer->SetBackPointer(this);
+	NewLayer.EditLayer->SetName(GenerateUniqueLayerName(InName), /*bInModify = */true);
+
 	OnLayerCreatedInternal(NewLayer);
 	return LayerIndex;
 }
@@ -13199,7 +13226,7 @@ int32 ALandscape::CreateLayerFrom(const FLandscapeLayer& InLayer)
 	}
 
 	// Cannot create a layer whose Guid is already taken : 
-	if (GetLayerConst(InLayer.Guid) != nullptr)
+	if (GetLayerConst(InLayer.EditLayer->GetGuid()) != nullptr)
 	{
 		return INDEX_NONE;
 	}
@@ -13219,11 +13246,14 @@ void ALandscape::OnLayerCreatedInternal(FLandscapeLayer& Layer)
 	// and has access to its guid.
 	Layer.EditLayer->SetBackPointer(this);
 	Layer.EditLayer->SetFlags(RF_Transactional);
-	Layer.EditLayer->OnLayerCreated(Layer);
+
+	// Register to data change events on the edit layer so that we can update the landscape accordingly : 
+	check(!Layer.EditLayer->OnLayerDataChanged().IsBoundToObject(this));
+	Layer.EditLayer->OnLayerDataChanged().AddUObject(this, &ALandscape::OnEditLayerDataChanged);
 
 	// Add to self first instead of from the loop.  Avoids potential load-order problems if this is called at load-time,
 	// before the LandscapeInfo is created.  The other proxies can fix themselves if they load later.
-	AddLayer(Layer.Guid);
+	AddLayer(Layer.EditLayer->GetGuid());
 
 	// Create associated layer data in each landscape proxy
 	if (ULandscapeInfo* LandscapeInfo = GetLandscapeInfo())
@@ -13232,7 +13262,7 @@ void ALandscape::OnLayerCreatedInternal(FLandscapeLayer& Layer)
 		{
 			if (Proxy != this)
 			{
-				Proxy->AddLayer(Layer.Guid);
+			Proxy->AddLayer(Layer.EditLayer->GetGuid());
 			}
 			return true;
 		});
@@ -13257,7 +13287,7 @@ void ALandscape::AddLayersToProxy(ALandscapeProxy* InProxy)
 
 	ForEachLayerConst([&](const FLandscapeLayer& Layer)
 	{
-		InProxy->AddLayer(Layer.Guid);
+		InProxy->AddLayer(Layer.EditLayer->GetGuid());
 		return true;
 	});
 
@@ -13284,7 +13314,7 @@ bool ALandscape::ReorderLayer(int32 InStartingLayerIndex, int32 InDestinationLay
 FName ALandscape::GenerateUniqueLayerName(FName InName) const
 {
 	// If we are receiving a unique name, use it.
-	if (InName != NAME_None && !LandscapeEditLayers.ContainsByPredicate([InName](const FLandscapeLayer& Layer) { return Layer.Name == InName; }))
+	if (InName != NAME_None && !LandscapeEditLayers.ContainsByPredicate([InName](const FLandscapeLayer& Layer) { return Layer.EditLayer->GetName() == InName; }))
 	{
 		return InName;
 	}
@@ -13296,7 +13326,7 @@ FName ALandscape::GenerateUniqueLayerName(FName InName) const
 	{
 		++LayerIndex;
 		NewName = FName(*FString::Printf(TEXT("%s%d"), *BaseName, LayerIndex));
-	} while (LandscapeEditLayers.ContainsByPredicate([NewName](const FLandscapeLayer& Layer) { return Layer.Name == NewName; }));
+	} while (LandscapeEditLayers.ContainsByPredicate([NewName](const FLandscapeLayer& Layer) { return Layer.EditLayer->GetName() == NewName; }));
 
 	return NewName;
 }
@@ -13310,7 +13340,7 @@ bool ALandscape::IsLayerBlendSubstractive(int32 InLayerIndex, const TWeakObjectP
 		return false;
 	}
 
-	const bool* AllocationBlend = Layer->WeightmapLayerAllocationBlend.Find(InLayerInfoObj.Get());
+	const bool* AllocationBlend = Layer->EditLayer->GetWeightmapLayerAllocationBlend().Find(InLayerInfoObj.Get());
 
 	if (AllocationBlend != nullptr)
 	{
@@ -13322,38 +13352,20 @@ bool ALandscape::IsLayerBlendSubstractive(int32 InLayerIndex, const TWeakObjectP
 
 void ALandscape::SetLayerSubstractiveBlendStatus(int32 InLayerIndex, bool InStatus, const TWeakObjectPtr<ULandscapeLayerInfoObject>& InLayerInfoObj)
 {
-	FLandscapeLayer* Layer = GetLayerInternal(InLayerIndex);
-
-	if (Layer == nullptr)
-	{
-		return;
-	}
-
-	Modify();
-	bool* AllocationBlend = Layer->WeightmapLayerAllocationBlend.Find(InLayerInfoObj.Get());
-
-	if (AllocationBlend == nullptr)
-	{
-		Layer->WeightmapLayerAllocationBlend.Add(InLayerInfoObj.Get(), InStatus);
-	}
-	else
+	if (FLandscapeLayer* Layer = GetLayerInternal(InLayerIndex))
 	{
-		*AllocationBlend = InStatus;
+		Layer->EditLayer->AddOrUpdateWeightmapAllocationLayerBlend(InLayerInfoObj.Get(), InStatus, /*bInModify =*/ true);
 	}
-
-	RequestLayersContentUpdateForceAll(ELandscapeLayerUpdateMode::Update_Weightmap_All);
 }
 
 void ALandscape::ReplaceLayerSubstractiveBlendStatus(ULandscapeLayerInfoObject* InFromLayerInfo, ULandscapeLayerInfoObject* InToLayerInfo, bool bInShouldDirtyPackage)
 {
-	Modify(bInShouldDirtyPackage);
-
 	for (FLandscapeLayer& Layer : LandscapeEditLayers)
 	{
 		bool OutValue;
-		if (Layer.WeightmapLayerAllocationBlend.RemoveAndCopyValue(InFromLayerInfo, OutValue))
+		if (Layer.EditLayer->RemoveAndCopyWeightmapAllocationLayerBlend(InFromLayerInfo, OutValue, /*bInModify = */ bInShouldDirtyPackage))
 		{
-			Layer.WeightmapLayerAllocationBlend.Add(InToLayerInfo, OutValue);
+			Layer.EditLayer->AddOrUpdateWeightmapAllocationLayerBlend(InToLayerInfo, OutValue, /*bInModify = */ bInShouldDirtyPackage);
 		}
 	}
 }
diff --git a/Engine/Source/Runtime/Landscape/Private/LandscapeSplineRaster.cpp b/Engine/Source/Runtime/Landscape/Private/LandscapeSplineRaster.cpp
--- a/Engine/Source/Runtime/Landscape/Private/LandscapeSplineRaster.cpp
+++ b/Engine/Source/Runtime/Landscape/Private/LandscapeSplineRaster.cpp
@@ -293,9 +293,9 @@ void RasterizeHeight(int32& MinX, int32& MinY, int32& MaxX, int32& MaxY, FLandsc
 
 	const FLandscapeLayer* EditingLayer = Landscape->GetLayerConst(Landscape->GetEditingLayer());
 	check((EditingLayer == nullptr) || (EditingLayer->EditLayer != nullptr));
-	const FLandscapeLayer* SplinesLayer = Landscape->FindLayerOfType(ULandscapeEditLayerSplines::StaticClass());
+	const FLandscapeLayer* SplinesLayer = Landscape->FindLayerOfTypeConst(ULandscapeEditLayerSplines::StaticClass());
 	bool bIsEditingLayerReservedForSplines = (EditingLayer != nullptr) && (SplinesLayer == EditingLayer);
-	check(!bIsEditingLayerReservedForSplines || (SplinesLayer->BlendMode == LSBM_AlphaBlend));
+	check(!bIsEditingLayerReservedForSplines || (SplinesLayer->EditLayer->GetBlendMode() == LSBM_AlphaBlend));
 
 	TArray<uint16> HeightAlphaBlendData;
 	TArray<uint8> HeightFlagsData;
@@ -577,8 +577,8 @@ bool ULandscapeInfo::ApplySplines(bool bOnlySelected, TSet<TObjectPtr<ULandscape
 	bool bResult = false;
 
 	ALandscape* Landscape = LandscapeActor.Get();
-	const FLandscapeLayer* Layer = Landscape ? Landscape->FindLayerOfType(ULandscapeEditLayerSplines::StaticClass()) : nullptr;
-	FGuid SplinesTargetLayerGuid = Layer ? Layer->Guid : Landscape ? Landscape->GetEditingLayer() : FGuid();
+	const FLandscapeLayer* Layer = Landscape ? Landscape->FindLayerOfTypeConst(ULandscapeEditLayerSplines::StaticClass()) : nullptr;
+	FGuid SplinesTargetLayerGuid = Layer ? Layer->EditLayer->GetGuid() : Landscape ? Landscape->GetEditingLayer() : FGuid();
 	FScopedSetLandscapeEditingLayer Scope(Landscape, SplinesTargetLayerGuid, [=] { Landscape->RequestLayersContentUpdate(ELandscapeLayerUpdateMode::Update_All); });
 
 	TMap<ULandscapeLayerInfoObject*, TSharedPtr<FModulateAlpha>> ModulatePerLayerInfo;
diff --git a/Engine/Source/Runtime/Landscape/Public/LandscapeEdit.h b/Engine/Source/Runtime/Landscape/Public/LandscapeEdit.h
--- a/Engine/Source/Runtime/Landscape/Public/LandscapeEdit.h
+++ b/Engine/Source/Runtime/Landscape/Public/LandscapeEdit.h
@@ -30,10 +30,10 @@ class ULandscapeComponent;
 class ULandscapeInfo;
 class ULandscapeLayerInfoObject;
 
-#define MAX_LANDSCAPE_LOD_DISTANCE_FACTOR 10.f
 
-#if WITH_EDITOR
+#define MAX_LANDSCAPE_LOD_DISTANCE_FACTOR UE_DEPRECATED_MACRO(5.6, "MAX_LANDSCAPE_LOD_DISTANCE_FACTOR has been deprecated because it is no longer necessary") 10.f
 
+#if WITH_EDITOR
 
 /** Landscape EdMode Interface (used by ALandscape in Editor mode to access EdMode properties) */
 class ILandscapeEdModeInterface
diff --git a/Engine/Source/Runtime/Landscape/Public/LandscapeSettings.h b/Engine/Source/Runtime/Landscape/Public/LandscapeSettings.h
--- a/Engine/Source/Runtime/Landscape/Public/LandscapeSettings.h
+++ b/Engine/Source/Runtime/Landscape/Public/LandscapeSettings.h
@@ -1,15 +1,15 @@
 // Copyright Epic Games, Inc. All Rights Reserved.
-
+ 
 #pragma once
-
+ 
 #include "CoreMinimal.h"
 #include "Engine/DeveloperSettings.h"
 #include "Delegates/Delegate.h"
 #include "LandscapeSettings.generated.h"
-
+ 
 class UMaterialInterface;
 class ULandscapeLayerInfoObject;
-
+ 
 UENUM()
 enum class ELandscapeDirtyingMode : uint8
 {
@@ -28,113 +28,119 @@ enum class ELandscapeDirtyingMode : uint8
 	 This mode is recommended for team collaboration as it provides the best features of the other two modes while ensuring that modified landscape actors are still saved and properly submitted to source control. */
 	InLandscapeModeAndUserTriggeredChanges
 };
-
+ 
 UCLASS(config = Engine, defaultconfig, meta = (DisplayName = "Landscape"), MinimalAPI)
 class ULandscapeSettings : public UDeveloperSettings
 {
 	GENERATED_BODY()
-
+ 
 public:
 	/** Returns true if landscape resolution should be constrained. */
 	bool IsLandscapeResolutionRestricted() const { return InRestrictiveMode(); }
-
+ 
 	/** Returns true if blueprint landscape tools usage is allowed */
 	bool AreBlueprintToolsAllowed() const { return !InRestrictiveMode(); }
 	
 	/** Returns the current landscape resolution limit. */
 	int32 GetTotalResolutionLimit() const { return SideResolutionLimit * SideResolutionLimit; }
-
+ 
 	bool InRestrictiveMode() const { return bRestrictiveMode; }
 	void SetRestrictiveMode(bool bEnabled) { bRestrictiveMode = bEnabled; }
-
+ 
 	int32 GetSideResolutionLimit() const { return SideResolutionLimit; }
-
+ 
 	float GetBrushSizeUIMax() const { return BrushSizeUIMax; }
 	float GetBrushSizeClampMax() const { return BrushSizeClampMax; }
-
+ 
 	int32 GetHLODMaxTextureSize() const { return HLODMaxTextureSize; }
-
+ 
 	float GetSplineIconWorldZOffset() const { return SplineIconWorldZOffset; }
 	float GetSplineIconScale() const { return SplineIconScale;  }
-
+ 
 	/** Returns the default landscape material that should be used when creating a new landscape. */
 	TSoftObjectPtr<UMaterialInterface> GetDefaultLandscapeMaterial() const { return DefaultLandscapeMaterial; }
-
+ 
 	/** Returns the default landscape layer info object that will be assigned to unset layers when creating a new landscape. */
 	TSoftObjectPtr<ULandscapeLayerInfoObject> GetDefaultLayerInfoObject() const { return DefaultLayerInfoObject; }
-
+ 
+	LANDSCAPE_API bool GetShouldUpdateEditLayersDuringInteractiveChanges() const { return bShouldUpdateEditLayersDuringInteractiveChanges; }
+ 
 	bool ShouldDisplayTargetLayerThumbnails() const { return bDisplayTargetLayerThumbnails; }
-
+ 
 	bool ShouldDisableTemporalAntiAliasingInLandscapeMode() const { return bDisableTemporalAntiAliasingInLandscapeMode; }
-
+ 
 #if WITH_EDITOR
 	//~ Begin UObject Interface.
 	virtual void PostEditChangeProperty(FPropertyChangedEvent& PropertyChangedEvent) override;
 	virtual void PreEditUndo() override;
 	virtual void PostEditUndo() override;
 #endif // WITH_EDITOR
-
+ 
 public:
-	UPROPERTY(config, EditAnywhere, Category = "Layers", meta=(UIMin = "1", UIMax = "32", ClampMin = "1", ClampMax = "32", ToolTip = "This option controls the maximum editing layers that can be added to a Landscape"))
+	UPROPERTY(config, EditAnywhere, Category = "Edit Layers", meta=(UIMin = "1", UIMax = "32", ClampMin = "1", ClampMax = "32", ToolTip = "This option controls the maximum editing layers that can be added to a Landscape"))
 	int32 MaxNumberOfLayers = 8;
-
+ 
 	UPROPERTY(EditAnywhere, config, Category = "Layers", meta = (ToolTip = 
 		"When true, automatic edit layer creation pops up a dialog where the new layer can be reordered relative to other layers."))
 	bool bShowDialogForAutomaticLayerCreation = true;
-
+ 
 	UPROPERTY(config, EditAnywhere, Category = "Configuration", meta=(ToolTip = "Maximum Dimension of Landscape in Components"))
 	int32 MaxComponents = 256;
-
+ 
 	UPROPERTY(config, EditAnywhere, Category = "Configuration", meta = (UIMin = "1", UIMax = "1024", ClampMin = "1", ClampMax = "1024", ToolTip = "Maximum Size of Import Image Cache in MB"))
 	uint32 MaxImageImportCacheSizeMegaBytes = 256;
-
+ 
 	UPROPERTY(config, EditAnywhere, Category = "Configuration", meta = (UIMin = "0.0", UIMax = "10.0", ClampMin = "0.0", ClampMax = "10.0", ToolTip = "Exponent for the Paint Tool Strength"))
 	float PaintStrengthGamma = 2.2f;
-
+ 
 	UPROPERTY(config, EditAnywhere, Category = "Configuration", meta = (ToolTip = "Disable Painting Startup Slowdown"))
 	bool bDisablePaintingStartupSlowdown = true;
 	
 	/** Defines when the engine requires the landscape actors to be resaved */
 	UPROPERTY(Config, Category = "Configuration", EditAnywhere)
 	ELandscapeDirtyingMode LandscapeDirtyingMode = ELandscapeDirtyingMode::InLandscapeModeAndUserTriggeredChanges;
-
+ 
 protected:
 	UPROPERTY(config)
 	int32 SideResolutionLimit = 2048;
-
+ 
 	UPROPERTY(EditAnywhere, config, Category = "Materials", meta = (ToolTip = "Default Landscape Material will be prefilled when creating a new landscape."))
 	TSoftObjectPtr<UMaterialInterface> DefaultLandscapeMaterial;
 
-	UPROPERTY(EditAnywhere, config, Category = "Layers", meta = (ToolTip = "Default Layer Info Object"))
+	UPROPERTY(EditAnywhere, config, Category = "Target Layers", meta = (ToolTip = "Default Layer Info Object"))
 	TSoftObjectPtr<ULandscapeLayerInfoObject> DefaultLayerInfoObject;
-
+ 
 	UPROPERTY(EditAnywhere, config, Category = "Configuration", meta = (ToolTip = "Maximum size that can be set via the slider for the landscape sculpt/paint brushes"))
 	float BrushSizeUIMax = 8192;
-
+ 
 	UPROPERTY(EditAnywhere, config, Category = "Configuration", meta = (ToolTip = "Maximum size that can be set manually for the landscape sculpt/paint brushes"))
 	float BrushSizeClampMax = 65536;
-
+ 
 	UPROPERTY(EditAnywhere, config, Category = "HLOD", meta = (DisplayName = "HLOD Max Texture Size", ClampMin = "64", ClampMax = "8192", ToolTip = "Maximum size of the textures generated for landscape HLODs"))
 	int32 HLODMaxTextureSize = 1024;
-
+ 
+	/** For landscape layers-affecting changes, allows the landscape to be updated when performing an interactive change (e.g. when changing an edit layer's alpha). Set to false if the performance when editing gets too bad (the landscape will be properly updated when the dragging operation is done). */
+	UPROPERTY(EditAnywhere, config, Category = "Edit Layers")
+	bool bShouldUpdateEditLayersDuringInteractiveChanges = false;
+ 
 	UPROPERTY(transient)
 	bool bRestrictiveMode = false;
-
+ 
 	/** Offset in Z for the landscape spline icon in world-space. */
 	UPROPERTY(EditAnywhere, config, Category = "Spline", meta = (ClampMin = "0", UIMax = "4096"))
 	float SplineIconWorldZOffset = 75.0f;
-
+ 
 	UPROPERTY(EditAnywhere, config, Category = "Spline", meta = (ClampMin = "32", UIMax = "2048"))
 	float SplineIconScale = 50.0f;
-
+ 
 	UPROPERTY(EditAnywhere, config, Category = "Configuration", meta = (ToolTip = 
 		"When true, each target layer will have a representative thumbnail in landscape mode. However, the thumbnail system can be slow. Plus, certain landscape materials (e.g. using RVT) are not able to properly render thumbnails. Thus, setting this to false will skip needlessly rendering landscape layer thumbnails, which can improve the editing experience"))
 	bool bDisplayTargetLayerThumbnails = true;
-
+ 
 	UPROPERTY(EditAnywhere, config, Category = "Configuration", meta = (ToolTip = 
 		"When true, temporal anti-aliasing will be inactive while in landscape mode. This avoids the ghosting effect on the landscape brush but can lead to aliasing or shimmering on other parts of the image."))
 	bool bDisableTemporalAntiAliasingInLandscapeMode = true;
-
+ 
 private:
 	TOptional<bool> DisplayTargetLayerThumbnailsBeforeUndo;
 };
