Fixed assert in landscape on undo : undo should not be considered a user-triggered change as far as requesting a landscape update goes

diff --git a/Engine/Source/Runtime/Landscape/Private/LandscapeEditLayer.cpp b/Engine/Source/Runtime/Landscape/Private/LandscapeEditLayer.cpp
--- a/Engine/Source/Runtime/Landscape/Private/LandscapeEditLayer.cpp
+++ b/Engine/Source/Runtime/Landscape/Private/LandscapeEditLayer.cpp
@@ -287,7 +287,7 @@ void ULandscapeEditLayerBase::PostEditUndo()
 {
 	Super::PostEditUndo();
 
-	BroadcastOnLayerDataChanged(/*InPropertyName = */NAME_None, /*bInUserTriggered = */true, /*bInRequiresLandscapeUpdate =*/ true, /*bInHasValueChanged = */true, /*InChangeType = */EPropertyChangeType::ValueSet);
+	BroadcastOnLayerDataChanged(/*InPropertyName = */NAME_None, /*bInUserTriggered = */false, /*bInRequiresLandscapeUpdate =*/ true, /*bInHasValueChanged = */true, /*InChangeType = */EPropertyChangeType::ValueSet);
 }
 
 bool ULandscapeEditLayerBase::CanEditChange(const FProperty* InProperty) const
