Fixed landscape smooth/erosion/flatten tools performance hit

diff --git a/Engine/Source/Editor/LandscapeEditor/Private/LandscapeEdModeTools.h b/Engine/Source/Editor/LandscapeEditor/Private/LandscapeEdModeTools.h
--- a/Engine/Source/Editor/LandscapeEditor/Private/LandscapeEdModeTools.h
+++ b/Engine/Source/Editor/LandscapeEditor/Private/LandscapeEdModeTools.h
@@ -1045,9 +1045,9 @@ private:
 			{
 				if (Layer->EditLayer->IsVisible() != InLayerVisibility[i])
 				{
-					// Pass bInModify = false here, because we don't want to call Modify on the landscape, we don't actually want to truly change 
+					// Pass bInModify and bInBroadcastDataChange = false here, because we don't want to call Modify on the landscape, we don't actually want to truly change 
 					//  the layer visibility, only temporarily for the purpose of the intermediate render : 
-					Layer->EditLayer->SetVisible(InLayerVisibility[i], /*bInModify = */ false);
+					Layer->EditLayer->SetVisible(InLayerVisibility[i], /*bInModify = */ false, /*bInBroadcastDataChange = */false);
 					bVisibilityChanged = true;
 				}
 			}
diff --git a/Engine/Source/Runtime/Landscape/Classes/LandscapeEditLayer.h b/Engine/Source/Runtime/Landscape/Classes/LandscapeEditLayer.h
--- a/Engine/Source/Runtime/Landscape/Classes/LandscapeEditLayer.h
+++ b/Engine/Source/Runtime/Landscape/Classes/LandscapeEditLayer.h
@@ -206,6 +206,15 @@ public:
 	 * @param bInModify : true when Modify() needs to be called
 	 */
 	LANDSCAPE_API virtual void SetVisible(bool bInVisible, bool bInModify);
+	
+	/**
+	 * Sets the layer's visibility value : reserved for internal use : we only need this version for intermediate renders, which 
+	 * should disappear in a later version
+	 * @param bInVisible : true to set visible, false to set invisible
+	 * @param bInModify : true when Modify() needs to be called
+	 * @param bInBroadcastDataChange : true when BroadcastOnLayerDataChanged() needs to be called
+	 */
+	UE_INTERNAL LANDSCAPE_API virtual void SetVisible(bool bInVisible, bool bInModify, bool bInBroadcastDataChange);
 
 	LANDSCAPE_API virtual bool IsVisible() const;
 
diff --git a/Engine/Source/Runtime/Landscape/Private/LandscapeEditLayer.cpp b/Engine/Source/Runtime/Landscape/Private/LandscapeEditLayer.cpp
--- a/Engine/Source/Runtime/Landscape/Private/LandscapeEditLayer.cpp
+++ b/Engine/Source/Runtime/Landscape/Private/LandscapeEditLayer.cpp
@@ -149,6 +149,11 @@ FName ULandscapeEditLayerBase::GetName() const
 }
 
 void ULandscapeEditLayerBase::SetVisible(bool bInVisible, bool bInModify)
+{
+	SetVisible(bInVisible, bInModify, /*bInBroadcastDataChange = */true);
+}
+
+void ULandscapeEditLayerBase::SetVisible(bool bInVisible, bool bInModify, bool bInBroadcastDataChange)
 {
 	const bool bHasValueChanged = (bInVisible != bVisible);
 	if (bHasValueChanged)
@@ -161,7 +166,10 @@ void ULandscapeEditLayerBase::SetVisible(bool bInVisible, bool bInModify)
 		bVisible = bInVisible;
 	}
 
-	BroadcastOnLayerDataChanged(GET_MEMBER_NAME_CHECKED(ULandscapeEditLayerBase, bVisible), /*bInUserTriggered = */true, /*bInRequiresLandscapeUpdate =*/ true, bHasValueChanged, /*InChangeType = */EPropertyChangeType::ValueSet);
+	if (bInBroadcastDataChange)
+	{
+		BroadcastOnLayerDataChanged(GET_MEMBER_NAME_CHECKED(ULandscapeEditLayerBase, bVisible), /*bInUserTriggered = */true, /*bInRequiresLandscapeUpdate =*/ true, bHasValueChanged, /*InChangeType = */EPropertyChangeType::ValueSet);
+	}
 }
 
 bool ULandscapeEditLayerBase::IsVisible() const
