Fixed crash when loading an edit layer whose class is unknown

diff --git a/Engine/Source/Runtime/Landscape/Private/Landscape.cpp b/Engine/Source/Runtime/Landscape/Private/Landscape.cpp
--- a/Engine/Source/Runtime/Landscape/Private/Landscape.cpp
+++ b/Engine/Source/Runtime/Landscape/Private/Landscape.cpp
@@ -2999,7 +2999,9 @@ PRAGMA_ENABLE_DEPRECATION_WARNINGS
 		}
 		else
 		{
-			UE_LOG(LogLandscape, Error, TEXT("Couldn't load edit layer object associated with layer %s for landscape %s. This may happen when the edit layer class cannot be found (for example, when a plugin is removed from the project). The layer will be deleted."), *Layer.EditLayer->GetName().ToString(), *GetFullName());
+			UE_LOG(LogLandscape, Error, TEXT("Couldn't load edit layer object associated with layer at index %i for landscape %s. This may happen when the edit layer class cannot be found "
+				"(for example, when a plugin is removed from the project). The layer will be deleted."), 
+				LayerIndex, *GetFullName());
 			ensure(DeleteLayer(LayerIndex));
 		}
 	}
diff --git a/Engine/Source/Runtime/Landscape/Private/LandscapeEditLayers.cpp b/Engine/Source/Runtime/Landscape/Private/LandscapeEditLayers.cpp
--- a/Engine/Source/Runtime/Landscape/Private/LandscapeEditLayers.cpp
+++ b/Engine/Source/Runtime/Landscape/Private/LandscapeEditLayers.cpp
@@ -12362,7 +12362,7 @@ bool ALandscape::IsUpToDate() const
 #if WITH_EDITOR
 bool ALandscape::IsLayerNameUnique(const FName& InName) const
 {
-	return Algo::CountIf(LandscapeEditLayers, [InName](const FLandscapeLayer& Layer) { return (Layer.EditLayer->GetName() == InName); }) == 0;
+	return Algo::CountIf(LandscapeEditLayers, [InName](const FLandscapeLayer& Layer) { return (Layer.EditLayer != nullptr) && (Layer.EditLayer->GetName() == InName); }) == 0;
 }
 
 void ALandscape::OnEditLayerDataChanged(const FOnLandscapeEditLayerDataChangedParams& InParams)
@@ -12659,31 +12659,38 @@ bool ALandscape::DeleteLayer(int32 InLayerIndex)
 	}
 
 	Modify();
-	FGuid LayerGuid = LayerStruct->EditLayer->GetGuid();
-
-	// Clean up Weightmap usage in LandscapeProxies
-	if (ULandscapeInfo* LandscapeInfo = GetLandscapeInfo())
-	{
-		LandscapeInfo->ForEachLandscapeProxy([&LayerGuid](ALandscapeProxy* Proxy)
-		{
-			Proxy->DeleteLayer(LayerGuid);
-			return true;
-		});
-	}
 
 	// We're about to remove the layer from our list, which will invalidate our LayerStruct pointer.
 	// We'll need to call OnLayerRemoved afterward, though, so keep pointer to the UObject.
 	ULandscapeEditLayerBase* EditLayer = LayerStruct->EditLayer;
-	check(EditLayer != nullptr);
+	// It's possible the edit layer UObject is missing (e.g. when we're trying to load an invalid UObject layer class), 
+	//  so we do our best to cleanup the associated data if we have access to it, but otherwise, it should be deleted on load : 
+	if (EditLayer != nullptr)
+	{
+		FGuid LayerGuid = LayerStruct->EditLayer->GetGuid();
+
+		// Clean up Weightmap usage in LandscapeProxies
+		if (ULandscapeInfo* LandscapeInfo = GetLandscapeInfo())
+		{
+			LandscapeInfo->ForEachLandscapeProxy([&LayerGuid](ALandscapeProxy* Proxy)
+			{
+				Proxy->DeleteLayer(LayerGuid);
+				return true;
+			});
+		}
+	}
 
 	// Remove layer from list
 	LandscapeEditLayers.RemoveAt(InLayerIndex);
 	LayerStruct = nullptr;
 
-	EditLayer->OnLayerRemoved();
+	if (EditLayer != nullptr)
+	{
+		EditLayer->OnLayerRemoved();
 
-	// Unregister from data change events on the edit layer so that we can update the landscape accordingly : 
-	EditLayer->OnLayerDataChanged().RemoveAll(this);
+		// Unregister from data change events on the edit layer so that we can update the landscape accordingly : 
+		EditLayer->OnLayerDataChanged().RemoveAll(this);
+	}
 
 	// Request Update
 	RequestLayersContentUpdateForceAll();
