Landscape Paint - Show unique thumbnail for target layers not referenced by any landscape proxy material
- Deprecate bValid field on FLandscapeInfoLayerSettings
- Deprecate LandscapeProxy OnMaterialChanged event for new event in the Landscape subsytem
- Refactor LandscapeProxy PostEditChange for material related properties

diff --git a/Engine/Source/Editor/LandscapeEditor/Private/LandscapeEdMode.cpp b/Engine/Source/Editor/LandscapeEditor/Private/LandscapeEdMode.cpp
--- a/Engine/Source/Editor/LandscapeEditor/Private/LandscapeEdMode.cpp
+++ b/Engine/Source/Editor/LandscapeEditor/Private/LandscapeEdMode.cpp
@@ -117,6 +117,7 @@ FEdModeLandscape::FEdModeLandscape()
 	, InteractorPainting(nullptr)
 	, bNeedsUpdateLayerUsageInformation(false)
 	, bUpdatingLandscapeInfo(false)
+	, bNeedsUpdateTargetLayerList(false)
 {
 	using namespace UE::Landscape::Editor;
 	
@@ -418,7 +419,9 @@ void FEdModeLandscape::Enter()
 	// Call parent implementation
 	FEdMode::Enter();
 
-	if (UWorld* World = GetWorld())
+	UWorld* World = GetWorld();
+
+	if (World != nullptr)
 	{
 		for (auto It = ULandscapeInfoMap::GetLandscapeInfoMap(World).Map.CreateIterator(); It; ++It)
 		{
@@ -482,10 +485,12 @@ void FEdModeLandscape::Enter()
 	OnLevelsChangedDelegateHandle				= GetWorld()->OnLevelsChanged().AddRaw(this, &FEdModeLandscape::HandleLevelsChanged);
 	OnMaterialCompilationFinishedDelegateHandle = UMaterial::OnMaterialCompilationFinished().AddRaw(this, &FEdModeLandscape::OnMaterialCompilationFinished);
 
-	if (CurrentToolTarget.LandscapeInfo.IsValid())
+	if (CurrentToolTarget.LandscapeInfo.IsValid() && World != nullptr)
 	{
-		ALandscapeProxy* LandscapeProxy = CurrentToolTarget.LandscapeInfo->GetLandscapeProxy();
-		LandscapeProxy->OnMaterialChangedDelegate().AddRaw(this, &FEdModeLandscape::OnLandscapeMaterialChangedDelegate);
+		if (ULandscapeSubsystem* LandscapeSubsystem = World->GetSubsystem<ULandscapeSubsystem>())
+		{
+			LandscapeSubsystem->OnLandscapeProxyMaterialChanged().AddRaw(this, &FEdModeLandscape::OnLandscapeMaterialChangedDelegate);
+		}
 
 		if (ALandscape* Landscape = GetLandscape())
 		{
@@ -647,7 +652,9 @@ void FEdModeLandscape::Exit()
 	GIsGizmoDragging = false;
 	GIsAdjustingBrush  = false;
 
-	if (UWorld* World = GetWorld())
+	UWorld* World = GetWorld();
+
+	if (World != nullptr)
 	{
 		for (auto It = ULandscapeInfoMap::GetLandscapeInfoMap(World).Map.CreateIterator(); It; ++It)
 		{
@@ -674,11 +681,11 @@ void FEdModeLandscape::Exit()
 	GetWorld()->OnLevelsChanged().Remove(OnLevelsChangedDelegateHandle);
 	UMaterial::OnMaterialCompilationFinished().Remove(OnMaterialCompilationFinishedDelegateHandle);
 
-	if (CurrentToolTarget.LandscapeInfo.IsValid())
+	if (CurrentToolTarget.LandscapeInfo.IsValid() && World != nullptr)
 	{
-		if (ALandscapeProxy* LandscapeProxy = CurrentToolTarget.LandscapeInfo->GetLandscapeProxy())
+		if (ULandscapeSubsystem* LandscapeSubsystem = World->GetSubsystem<ULandscapeSubsystem>())
 		{
-			LandscapeProxy->OnMaterialChangedDelegate().RemoveAll(this);
+			LandscapeSubsystem->OnLandscapeProxyMaterialChanged().RemoveAll(this);
 		}
 
 		if (ALandscape* Landscape = GetLandscape())
@@ -858,6 +865,13 @@ void FEdModeLandscape::Tick(FEditorViewportClient* ViewportClient, float DeltaTi
 		}
 		LastLandscapeSplineFalloffModulationValue = CVarLandscapeSplineFalloffModulation.GetValueOnAnyThread();
 	}
+
+	if (bNeedsUpdateTargetLayerList)
+	{
+		// Updating target list because a landscape material has changed, regenerate lists that may hold on to a ThumbnailMIC
+		UpdateTargetList(/*bRegenerateThumbnails =*/ true);
+		bNeedsUpdateTargetLayerList = false;
+	}
 }
 
 /** FEdMode: Called when the mouse enters the viewport area */
@@ -2507,8 +2521,6 @@ void FEdModeLandscape::SetTargetLandscape(const TWeakObjectPtr<ULandscapeInfo>&
 	// Unregister from old one
 	if (CurrentToolTarget.LandscapeInfo.IsValid())
 	{
-		ALandscapeProxy* LandscapeProxy = CurrentToolTarget.LandscapeInfo->GetLandscapeProxy();
-		LandscapeProxy->OnMaterialChangedDelegate().RemoveAll(this);
 		if (ALandscape* Landscape = GetLandscape())
 		{
 			Landscape->OnBlueprintBrushChangedDelegate().RemoveAll(this);
@@ -2526,8 +2538,6 @@ void FEdModeLandscape::SetTargetLandscape(const TWeakObjectPtr<ULandscapeInfo>&
 	// register to new one
 	if (CurrentToolTarget.LandscapeInfo.IsValid())
 	{
-		ALandscapeProxy* LandscapeProxy = CurrentToolTarget.LandscapeInfo->GetLandscapeProxy();
-		LandscapeProxy->OnMaterialChangedDelegate().AddRaw(this, &FEdModeLandscape::OnLandscapeMaterialChangedDelegate);
 		if (ALandscape* Landscape = GetLandscape())
 		{
 			Landscape->OnBlueprintBrushChangedDelegate().AddRaw(this, &FEdModeLandscape::RefreshDetailPanel);
@@ -2592,7 +2602,7 @@ void FEdModeLandscape::UpdateTargetList(bool bRegenerateThumbnails)
 			bool bFoundSelected = false;
 
 			// Add heightmap
-			LandscapeTargetList.Add(MakeShareable(new FLandscapeTargetListInfo(LOCTEXT("Heightmap", "Heightmap"), ELandscapeToolTargetType::Heightmap, CurrentToolTarget.LandscapeInfo.Get(), UISettings->CurrentLayerIndex)));
+			LandscapeTargetList.Add(MakeShareable(new FLandscapeTargetListInfo(LOCTEXT("Heightmap", "Heightmap"), ELandscapeToolTargetType::Heightmap, CurrentToolTarget.LandscapeInfo.Get(), UISettings->CurrentLayerIndex, /* bIsLayerReferencedByMaterial=*/ false)));
 
 			if (CurrentToolTarget.TargetType == ELandscapeToolTargetType::Heightmap)
 			{
@@ -2601,7 +2611,7 @@ void FEdModeLandscape::UpdateTargetList(bool bRegenerateThumbnails)
 
 			// Add visibility
 			FLandscapeInfoLayerSettings VisibilitySettings(ALandscapeProxy::VisibilityLayer, Landscape);
-			LandscapeTargetList.Add(MakeShareable(new FLandscapeTargetListInfo(LOCTEXT("Visibility", "Visibility"), ELandscapeToolTargetType::Visibility, VisibilitySettings, UISettings->CurrentLayerIndex)));
+			LandscapeTargetList.Add(MakeShareable(new FLandscapeTargetListInfo(LOCTEXT("Visibility", "Visibility"), ELandscapeToolTargetType::Visibility, VisibilitySettings, UISettings->CurrentLayerIndex, /* bIsLayerReferencedByMaterial=*/ false)));
 
 			if (CurrentToolTarget.TargetType == ELandscapeToolTargetType::Visibility)
 			{
@@ -2610,6 +2620,12 @@ void FEdModeLandscape::UpdateTargetList(bool bRegenerateThumbnails)
 
 			// Add layers
 			TargetLayerStartingIndex = LandscapeTargetList.Num();
+			TSet<FName> LayerNames;
+			CurrentToolTarget.LandscapeInfo->ForEachLandscapeProxy([&](ALandscapeProxy* Proxy)
+				{
+					LayerNames.Append(Proxy->RetrieveTargetLayerNamesFromMaterials());
+					return true;
+				});
 
 			// Don't recreate the render state of everything, only update the materials context
 			FMaterialUpdateContext MaterialUpdateContext(FMaterialUpdateContext::EOptions::Default & ~FMaterialUpdateContext::EOptions::RecreateRenderStates);
@@ -2640,7 +2656,7 @@ void FEdModeLandscape::UpdateTargetList(bool bRegenerateThumbnails)
 				}
 
 				// Add the layer
-				LandscapeTargetList.Add(MakeShareable(new FLandscapeTargetListInfo(FText::FromName(LayerName), ELandscapeToolTargetType::Weightmap, LayerSettings, UISettings->CurrentLayerIndex)));
+				LandscapeTargetList.Add(MakeShareable(new FLandscapeTargetListInfo(FText::FromName(LayerName), ELandscapeToolTargetType::Weightmap, LayerSettings, UISettings->CurrentLayerIndex, /* bIsLayerReferencedByMaterial=*/ LayerNames.Contains(LayerName))));
 			}
 
 			if (!bFoundSelected)
@@ -2768,9 +2784,10 @@ void FEdModeLandscape::UpdateTargetLayerDisplayOrder(ELandscapeLayerDisplayMode
 	}
 }
 
-void FEdModeLandscape::OnLandscapeMaterialChangedDelegate()
+void FEdModeLandscape::OnLandscapeMaterialChangedDelegate(ALandscapeProxy* InProxyChanged, const FOnLandscapeProxyMaterialChangedParams& InParams)
 {
-	UpdateTargetList();
+	// call UpdateTargetList next frame to avoid being called by all proxies
+	bNeedsUpdateTargetLayerList = true;
 }
 
 void FEdModeLandscape::RequestUpdateLayerUsageInformation()
diff --git a/Engine/Source/Editor/LandscapeEditor/Private/LandscapeEdMode.h b/Engine/Source/Editor/LandscapeEditor/Private/LandscapeEdMode.h
--- a/Engine/Source/Editor/LandscapeEditor/Private/LandscapeEdMode.h
+++ b/Engine/Source/Editor/LandscapeEditor/Private/LandscapeEdMode.h
@@ -69,18 +69,18 @@ struct FLandscapeToolMode
 struct FLandscapeTargetListInfo
 {
 	FText TargetLayerDisplayName;									// UI Display Name
-	ELandscapeToolTargetType TargetType;
+	ELandscapeToolTargetType TargetType = ELandscapeToolTargetType::Invalid;
 	TWeakObjectPtr<ULandscapeInfo> LandscapeInfo;
 
 	TWeakObjectPtr<ULandscapeLayerInfoObject> LayerInfoObj;			// ignored for heightmap
 	FName LayerName;												// ignored for heightmap
 	TWeakObjectPtr<class ALandscapeProxy> Owner;					// ignored for heightmap
 	TWeakObjectPtr<class UMaterialInstanceConstant> ThumbnailMIC;	// ignored for heightmap
-	int32 DebugColorChannel;										// ignored for heightmap
-	uint32 bValid : 1;												// ignored for heightmap
-	int32 LayerIndex;
+	int32 DebugColorChannel = -1;									// ignored for heightmap
+	int32 LayerIndex = INDEX_NONE;
+	bool bIsLayerReferencedByMaterial = false;
 
-	FLandscapeTargetListInfo(FText InTargetLayerDisplayName, ELandscapeToolTargetType InTargetType, const FLandscapeInfoLayerSettings& InLayerSettings, int32 InLayerIndex)
+	FLandscapeTargetListInfo(FText InTargetLayerDisplayName, ELandscapeToolTargetType InTargetType, const FLandscapeInfoLayerSettings& InLayerSettings, int32 InLayerIndex, bool bInIsLayerReferencedByMaterial)
 		: TargetLayerDisplayName(InTargetLayerDisplayName)
 		, TargetType(InTargetType)
 		, LandscapeInfo(InLayerSettings.Owner->GetLandscapeInfo())
@@ -89,12 +89,12 @@ struct FLandscapeTargetListInfo
 		, Owner(InLayerSettings.Owner)
 		, ThumbnailMIC(InLayerSettings.ThumbnailMIC)
 		, DebugColorChannel(InLayerSettings.DebugColorChannel)
-		, bValid(InLayerSettings.bValid)
 		, LayerIndex(InLayerIndex)
+		, bIsLayerReferencedByMaterial(bInIsLayerReferencedByMaterial)
 	{
 	}
 
-	FLandscapeTargetListInfo(FText InTargetLayerDisplayName, ELandscapeToolTargetType InTargetType, ULandscapeInfo* InLandscapeInfo, int32 InLayerIndex)
+	FLandscapeTargetListInfo(FText InTargetLayerDisplayName, ELandscapeToolTargetType InTargetType, ULandscapeInfo* InLandscapeInfo, int32 InLayerIndex, bool bInIsLayerReferencedByMaterial)
 		: TargetLayerDisplayName(InTargetLayerDisplayName)
 		, TargetType(InTargetType)
 		, LandscapeInfo(InLandscapeInfo)
@@ -102,8 +102,8 @@ struct FLandscapeTargetListInfo
 		, LayerName(NAME_None)
 		, Owner(nullptr)
 		, ThumbnailMIC(nullptr)
-		, bValid(true)
 		, LayerIndex(InLayerIndex)
+		, bIsLayerReferencedByMaterial(bInIsLayerReferencedByMaterial)
 	{
 	}
 
@@ -574,7 +574,7 @@ public:
 	void RequestUpdateLayerUsageInformation();
 	bool ShouldShowLayer(TSharedRef<FLandscapeTargetListInfo> Target) const;
 	void UpdateLayerUsageInformation(TWeakObjectPtr<ULandscapeLayerInfoObject>* LayerInfoObjectThatChanged = nullptr);
-	void OnLandscapeMaterialChangedDelegate();
+	void OnLandscapeMaterialChangedDelegate(ALandscapeProxy* InProxyChanged, const FOnLandscapeProxyMaterialChangedParams& InParams);
 	void RefreshDetailPanel();
 	void RefreshInspectedObjectsDetailPanel();
 	void RegenerateLayerThumbnails();
@@ -703,6 +703,7 @@ private:
 	/** Delayed refresh */
 	bool bNeedsUpdateLayerUsageInformation;
 	bool bUpdatingLandscapeInfo;
+	bool bNeedsUpdateTargetLayerList;
 
 	/** The landscape mode panel has the ability to display the properties of certain UObjects if needed (to bypass the fact that the standard details 
 	 panel is not able to display anything else than actors and components) */
diff --git a/Engine/Source/Editor/LandscapeEditor/Private/LandscapeEdModeComponentTools.cpp b/Engine/Source/Editor/LandscapeEditor/Private/LandscapeEdModeComponentTools.cpp
--- a/Engine/Source/Editor/LandscapeEditor/Private/LandscapeEdModeComponentTools.cpp
+++ b/Engine/Source/Editor/LandscapeEditor/Private/LandscapeEdModeComponentTools.cpp
@@ -536,104 +536,6 @@ public:
 // FLandscapeToolAddComponent
 //
 
-namespace
-{
-	void AddComponents(ULandscapeInfo* LandscapeInfo, ULandscapeSubsystem* LandscapeSubsystem, const TArray<FIntPoint>& ComponentCoordinates)
-	{
-		TArray<ULandscapeComponent*> NewComponents;
-		LandscapeInfo->Modify();
-		for (const auto& ComponentCoordinate : ComponentCoordinates)
-		{
-			ULandscapeComponent* LandscapeComponent = LandscapeInfo->XYtoComponentMap.FindRef(ComponentCoordinate);
-			if (LandscapeComponent)
-			{
-				continue;
-			}
-
-			// Add New component...
-			FIntPoint ComponentBase = ComponentCoordinate * LandscapeInfo->ComponentSizeQuads;
-
-			ALandscapeProxy* LandscapeProxy = LandscapeSubsystem->FindOrAddLandscapeProxy(LandscapeInfo, ComponentBase);
-			if (!LandscapeProxy)
-			{
-				continue;
-			}
-
-			LandscapeComponent = NewObject<ULandscapeComponent>(LandscapeProxy, NAME_None, RF_Transactional);
-			NewComponents.Add(LandscapeComponent);
-			LandscapeComponent->Init(
-				ComponentBase.X, ComponentBase.Y,
-				LandscapeProxy->ComponentSizeQuads,
-				LandscapeProxy->NumSubsections,
-				LandscapeProxy->SubsectionSizeQuads
-			);
-
-			TArray<FColor> HeightData;
-			const int32 ComponentVerts = (LandscapeComponent->SubsectionSizeQuads + 1) * LandscapeComponent->NumSubsections;
-			HeightData.Empty(FMath::Square(ComponentVerts));
-			HeightData.AddZeroed(FMath::Square(ComponentVerts));
-			LandscapeComponent->InitHeightmapData(HeightData, true);
-			LandscapeComponent->UpdateMaterialInstances();
-
-			LandscapeInfo->XYtoComponentMap.Add(ComponentCoordinate, LandscapeComponent);
-			LandscapeInfo->XYtoAddCollisionMap.Remove(ComponentCoordinate);
-		}
-
-		// Need to register to use general height/xyoffset data update
-		for (int32 Idx = 0; Idx < NewComponents.Num(); Idx++)
-		{
-			NewComponents[Idx]->RegisterComponent();
-		}
-
-		const bool bHasXYOffset = false;
-		ALandscape* Landscape = LandscapeInfo->LandscapeActor.Get();
-
-		bool bHasLandscapeLayersContent = Landscape && Landscape->HasLayersContent();
-
-		if (bHasLandscapeLayersContent)
-		{
-			Landscape->RequestLayersInitialization();
-		}
-
-		for (ULandscapeComponent* NewComponent : NewComponents)
-		{
-			if (bHasLandscapeLayersContent)
-			{
-				TArray<ULandscapeComponent*> ComponentsUsingHeightmap;
-				ComponentsUsingHeightmap.Add(NewComponent);
-
-				for (const FLandscapeLayer& Layer : Landscape->GetLayersConst())
-				{
-					// Since we do not share heightmap when adding new component, we will provided the required array, but they will only be used for 1 component
-					TMap<UTexture2D*, UTexture2D*> CreatedHeightmapTextures;
-					NewComponent->AddDefaultLayerData(Layer.EditLayer->GetGuid(), ComponentsUsingHeightmap, CreatedHeightmapTextures);
-				}
-			}
-
-			// Update Collision
-			NewComponent->UpdateCachedBounds();
-			NewComponent->UpdateBounds();
-			NewComponent->MarkRenderStateDirty();
-
-			if (!bHasLandscapeLayersContent)
-			{
-				ULandscapeHeightfieldCollisionComponent* CollisionComp = NewComponent->GetCollisionComponent();
-				if (CollisionComp && !bHasXYOffset)
-				{
-					CollisionComp->MarkRenderStateDirty();
-					CollisionComp->RecreateCollision();
-				}
-			}
-		}
-
-		if (Landscape)
-		{
-			GEngine->BroadcastOnActorMoved(Landscape);
-		}
-
-	}
-}
-
 class FLandscapeToolStrokeAddComponent : public FLandscapeToolStrokeBase
 {
 	using Super = FLandscapeToolStrokeBase;
@@ -705,8 +607,7 @@ public:
 						// Add New component...
 						FIntPoint ComponentBase = FIntPoint(ComponentIndexX, ComponentIndexY)* LandscapeInfo->ComponentSizeQuads;
 
-						ALandscapeProxy* LandscapeProxy = LandscapeSubsystem->FindOrAddLandscapeProxy(LandscapeInfo, ComponentBase);
-						if (LandscapeProxy)
+						if (ALandscapeProxy* LandscapeProxy = LandscapeSubsystem->FindOrAddLandscapeProxy(LandscapeInfo, ComponentBase))
 						{
 							LandscapeComponent = NewObject<ULandscapeComponent>(LandscapeProxy, NAME_None, RF_Transactional);
 							NewComponents.Add(LandscapeComponent);
@@ -1501,7 +1402,6 @@ public:
 						{
 							LayerSettings.Owner = LandscapeInfo->GetLandscapeProxy(); // this isn't strictly accurate, but close enough
 							LayerSettings.LayerInfoObj = LayerInfo;
-							LayerSettings.bValid = true;
 						}
 					}
 				}
diff --git a/Engine/Source/Editor/LandscapeEditor/Private/LandscapeEditorDetailCustomization_NewLandscape.cpp b/Engine/Source/Editor/LandscapeEditor/Private/LandscapeEditorDetailCustomization_NewLandscape.cpp
--- a/Engine/Source/Editor/LandscapeEditor/Private/LandscapeEditorDetailCustomization_NewLandscape.cpp
+++ b/Engine/Source/Editor/LandscapeEditor/Private/LandscapeEditorDetailCustomization_NewLandscape.cpp
@@ -1304,12 +1304,6 @@ FReply FLandscapeEditorDetailCustomization_NewLandscape::OnCreateButtonClicked()
 
 	World->GetSubsystem<ULandscapeSubsystem>()->ChangeGridSize(LandscapeInfo, UISettings->WorldPartitionGridSize);
 
-	if (LandscapeEdMode->CurrentToolTarget.LandscapeInfo.IsValid())
-	{
-		ALandscapeProxy* LandscapeProxy = LandscapeEdMode->CurrentToolTarget.LandscapeInfo->GetLandscapeProxy();
-		LandscapeProxy->OnMaterialChangedDelegate().AddRaw(LandscapeEdMode, &FEdModeLandscape::OnLandscapeMaterialChangedDelegate);
-	}
-
 	ULandscapeSubsystem* LandscapeSubsystem = World->GetSubsystem<ULandscapeSubsystem>();
 	ALandscapeProxy* LandscapeProxy = LandscapeEdMode->CurrentToolTarget.LandscapeInfo->GetLandscapeProxy();
 	
diff --git a/Engine/Source/Editor/LandscapeEditor/Private/LandscapeEditorDetailCustomization_TargetLayers.cpp b/Engine/Source/Editor/LandscapeEditor/Private/LandscapeEditorDetailCustomization_TargetLayers.cpp
--- a/Engine/Source/Editor/LandscapeEditor/Private/LandscapeEditorDetailCustomization_TargetLayers.cpp
+++ b/Engine/Source/Editor/LandscapeEditor/Private/LandscapeEditorDetailCustomization_TargetLayers.cpp
@@ -538,7 +538,7 @@ const TArray<const TSharedRef<FLandscapeTargetListInfo>> FLandscapeEditorCustomN
 			if (!LandscapeActor->GetTargetLayers().Find(LayerName))
 			{
 				const FLandscapeInfoLayerSettings LayerSettings = FLandscapeInfoLayerSettings(LayerName, LandscapeActor.Get());
-				const TSharedRef<FLandscapeTargetListInfo> Target = MakeShareable(new FLandscapeTargetListInfo(FText::FromName(LayerName), ELandscapeToolTargetType::Weightmap, LayerSettings, LandscapeEdMode->UISettings->CurrentLayerIndex));
+				const TSharedRef<FLandscapeTargetListInfo> Target = MakeShareable(new FLandscapeTargetListInfo(FText::FromName(LayerName), ELandscapeToolTargetType::Weightmap, LayerSettings, LandscapeEdMode->UISettings->CurrentLayerIndex, /* bIsLayerReferencedByMaterial=*/ true));
 
 				TargetLayerList.Add(Target);
 			}
@@ -1141,7 +1141,7 @@ TSharedPtr<SWidget> FLandscapeEditorCustomNodeBuilder_TargetLayers::GenerateRow(
 					.VAlign(VAlign_Center)
 					.Padding(FMargin(2))
 					[
-						(Target->bValid)
+						(Target->bIsLayerReferencedByMaterial)
 							? (TSharedRef<SWidget>)(
 								SNew(SLandscapeAssetThumbnail, Target->ThumbnailMIC.Get(), ThumbnailPool)
 								.Visibility_Static(&FLandscapeEditorCustomNodeBuilder_TargetLayers::GetDebugModeLayerUsageVisibility_Invert, Target)
@@ -1161,6 +1161,7 @@ TSharedPtr<SWidget> FLandscapeEditorCustomNodeBuilder_TargetLayers::GenerateRow(
 								SNew(SImage)
 								.Visibility_Static(&FLandscapeEditorCustomNodeBuilder_TargetLayers::GetDebugModeLayerUsageVisibility_Invert, Target)
 								.Image(FAppStyle::GetBrush(TEXT("LandscapeEditor.Target_Invalid")))
+								.ToolTipText(LOCTEXT("LandscapeTargetLayer_NotReferencedByMaterial", "This target layer is not referenced by any assigned landscape material"))
 								)
 					]
 					+ SHorizontalBox::Slot()
@@ -1258,7 +1259,6 @@ TSharedPtr<SWidget> FLandscapeEditorCustomNodeBuilder_TargetLayers::GenerateRow(
 									.VAlign(VAlign_Center)
 									[
 										SNew(SObjectPropertyEntryBox)
-											.IsEnabled((bool)Target->bValid)
 											.ObjectPath(Target->LayerInfoObj != nullptr ? Target->LayerInfoObj->GetPathName() : FString())
 											.AllowedClass(ULandscapeLayerInfoObject::StaticClass())
 											.OnObjectChanged_Static(&FLandscapeEditorCustomNodeBuilder_TargetLayers::OnTargetLayerSetObject, Target)
@@ -1850,16 +1850,6 @@ bool FLandscapeEditorCustomNodeBuilder_TargetLayers::GetTargetLayerCreateEnabled
 	return false;
 }
 
-EVisibility FLandscapeEditorCustomNodeBuilder_TargetLayers::GetTargetLayerDeleteVisibility(const TSharedRef<FLandscapeTargetListInfo> Target)
-{
-	if (!Target->bValid)
-	{
-		return EVisibility::Visible;
-	}
-
-	return EVisibility::Collapsed;
-}
-
 TSharedRef<SWidget> FLandscapeEditorCustomNodeBuilder_TargetLayers::OnGetTargetLayerCreateMenu(const TSharedRef<FLandscapeTargetListInfo> Target)
 {
 	FMenuBuilder MenuBuilder(true, nullptr);
diff --git a/Engine/Source/Editor/LandscapeEditor/Private/LandscapeEditorDetailCustomization_TargetLayers.h b/Engine/Source/Editor/LandscapeEditor/Private/LandscapeEditorDetailCustomization_TargetLayers.h
--- a/Engine/Source/Editor/LandscapeEditor/Private/LandscapeEditorDetailCustomization_TargetLayers.h
+++ b/Engine/Source/Editor/LandscapeEditor/Private/LandscapeEditorDetailCustomization_TargetLayers.h
@@ -105,7 +105,6 @@ protected:
 	static void OnTargetLayerSetObject(const FAssetData& AssetData, const TSharedRef<FLandscapeTargetListInfo> Target); 
 	static EVisibility GetTargetLayerInfoSelectorVisibility(const TSharedRef<FLandscapeTargetListInfo> Target);
 	static bool GetTargetLayerCreateEnabled(const TSharedRef<FLandscapeTargetListInfo> Target);
-	static EVisibility GetTargetLayerDeleteVisibility(const TSharedRef<FLandscapeTargetListInfo> Target);
 	static TSharedRef<SWidget> OnGetTargetLayerCreateMenu(const TSharedRef<FLandscapeTargetListInfo> Target);
 	static void OnTargetLayerCreateClicked(const TSharedRef<FLandscapeTargetListInfo> Target, bool bNoWeightBlend);
 	static void CreateTargetLayerInfoAsset(const TSharedRef<FLandscapeTargetListInfo> Target, const FString& PackageName, const FName& LayerObjectName, const bool bNoWeightBlend);
diff --git a/Engine/Source/Runtime/Landscape/Classes/Landscape.h b/Engine/Source/Runtime/Landscape/Classes/Landscape.h
--- a/Engine/Source/Runtime/Landscape/Classes/Landscape.h
+++ b/Engine/Source/Runtime/Landscape/Classes/Landscape.h
@@ -286,7 +286,6 @@ public:
 	//~ End APartitionActor Interface
 
 	//~ Begin UObject Interface.
-	virtual void PreEditChange(FProperty* PropertyThatWillChange) override;
 	virtual void PostEditChangeProperty(FPropertyChangedEvent& PropertyChangedEvent) override;
 	virtual void PostEditMove(bool bFinished) override;
 	virtual void PostEditUndo() override;
diff --git a/Engine/Source/Runtime/Landscape/Classes/LandscapeInfo.h b/Engine/Source/Runtime/Landscape/Classes/LandscapeInfo.h
--- a/Engine/Source/Runtime/Landscape/Classes/LandscapeInfo.h
+++ b/Engine/Source/Runtime/Landscape/Classes/LandscapeInfo.h
@@ -81,9 +81,6 @@ struct FLandscapeInfoLayerSettings
 
 	UPROPERTY(transient)
 	int32 DebugColorChannel;
-
-	UPROPERTY(transient)
-	uint32 bValid:1;
 #endif // WITH_EDITORONLY_DATA
 
 	FLandscapeInfoLayerSettings()
@@ -93,7 +90,6 @@ struct FLandscapeInfoLayerSettings
 		, ThumbnailMIC(nullptr)
 		, Owner(nullptr)
 		, DebugColorChannel(0)
-		, bValid(false)
 #endif // WITH_EDITORONLY_DATA
 	{
 	}
diff --git a/Engine/Source/Runtime/Landscape/Classes/LandscapeProxy.h b/Engine/Source/Runtime/Landscape/Classes/LandscapeProxy.h
--- a/Engine/Source/Runtime/Landscape/Classes/LandscapeProxy.h
+++ b/Engine/Source/Runtime/Landscape/Classes/LandscapeProxy.h
@@ -67,6 +67,11 @@ extern int32 GGrassMapUseRuntimeGeneration;
 
 DECLARE_MULTICAST_DELEGATE_TwoParams(FOnLandscapeProxyComponentDataChanged, ALandscapeProxy*, const FLandscapeProxyComponentDataChangedParams&);
 
+struct FOnLandscapeProxyMaterialChangedParams
+{
+};
+DECLARE_MULTICAST_DELEGATE_TwoParams(FOnLandscapeProxyMaterialChanged, ALandscapeProxy*, const FOnLandscapeProxyMaterialChangedParams&);
+
 struct FOnLandscapeProxyFixupSharedDataParams
 {
 	/** Parent landscape actor for the landscape proxy on which OnLandscapeProxyFixupSharedData is triggered */
@@ -452,8 +457,9 @@ protected:
 	TArray<FLandscapePerLODMaterialOverride> PerLODOverrideMaterials;
 
 #if WITH_EDITORONLY_DATA
+	UE_DEPRECATED(5.6, "PreEditPerLODOverrideMaterials has been deprecated. PerLODOverrideMaterial property changes handled in PostEditChange.")
 	UPROPERTY(Transient)
-	TArray<FLandscapePerLODMaterialOverride> PreEditPerLODOverrideMaterials;
+	TArray<FLandscapePerLODMaterialOverride> PreEditPerLODOverrideMaterials_DEPRECATED;
 
 	/** LOD level of the landscape when generating the Nanite mesh. Mostly there for debug reasons, since Nanite is meant to allow high density meshes, we want to use 0 most of the times. */
 	UPROPERTY(EditAnywhere, Category = Nanite, AdvancedDisplay, meta = (EditCondition = "bEnableNanite", LandscapeInherited))
@@ -570,20 +576,23 @@ public:
 
 #if WITH_EDITORONLY_DATA
 
-	PRAGMA_DISABLE_DEPRECATION_WARNINGS
+PRAGMA_DISABLE_DEPRECATION_WARNINGS
 	UE_DEPRECATED(all, "LandscapeComponentMaterialOverride has been deprecated, use PerLODOverrideMaterials instead.")
 	UPROPERTY()
 	TArray<FLandscapeProxyMaterialOverride> LandscapeMaterialsOverride_DEPRECATED;
-	PRAGMA_ENABLE_DEPRECATION_WARNINGS
+PRAGMA_ENABLE_DEPRECATION_WARNINGS
 
+	UE_DEPRECATED(5.6, "PreEditLandscapeMaterial has been deprecated. LandscapeMaterial Property changes handled in PostEditChange.")
 	UPROPERTY(Transient)
-	TObjectPtr<UMaterialInterface> PreEditLandscapeMaterial;
+	TObjectPtr<UMaterialInterface> PreEditLandscapeMaterial_DEPRECATED;
 
+	UE_DEPRECATED(5.6, "PreEditLandscapeHoleMaterial has been deprecated. LandscapeHoleMaterial Property changes handled in PostEditChange.")
 	UPROPERTY(Transient)
-	TObjectPtr<UMaterialInterface> PreEditLandscapeHoleMaterial;
+	TObjectPtr<UMaterialInterface> PreEditLandscapeHoleMaterial_DEPRECATED;
 
+	UE_DEPRECATED(5.6, "bIsPerformingInteractiveActionOnLandscapeMaterialOverride has been deprecated. PostEditChange checks interaction type.")
 	UPROPERTY(Transient)
-	bool bIsPerformingInteractiveActionOnLandscapeMaterialOverride;
+	bool bIsPerformingInteractiveActionOnLandscapeMaterialOverride_DEPRECATED = false;
 #endif // WITH_EDITORONLY_DATA
 
 	/**
@@ -1564,7 +1573,13 @@ PRAGMA_ENABLE_DEPRECATION_WARNINGS
 	bool LandscapeExportWeightmapToRenderTarget(UTextureRenderTarget2D* InRenderTarget, FName InLayerName);
 
 	DECLARE_EVENT(ALandscape, FLandscapeMaterialChangedDelegate);
-	FLandscapeMaterialChangedDelegate& OnMaterialChangedDelegate() { return LandscapeMaterialChangedDelegate; }
+	UE_DEPRECATED(5.6, "Use ULandscapeSubsystem::OnLandscapeProxyMaterialChanged")
+	FLandscapeMaterialChangedDelegate& OnMaterialChangedDelegate() 
+	{ 
+PRAGMA_DISABLE_DEPRECATION_WARNINGS
+		return LandscapeMaterialChangedDelegate; 
+PRAGMA_ENABLE_DEPRECATION_WARNINGS
+	}
 
 	/** Will tell if the landscape proxy as some content related to the layer system */
 	LANDSCAPE_API virtual bool HasLayersContent() const;
@@ -1639,6 +1654,7 @@ protected:
 	LANDSCAPE_API virtual TArray<FName> SynchronizeUnmarkedSharedProperties(ALandscapeProxy* InLandscape);
 
 protected:
+	UE_DEPRECATED(5.6, "Use ULandscapeSubsystem::OnLandscapeProxyComponentDataChanged")
 	FLandscapeMaterialChangedDelegate LandscapeMaterialChangedDelegate;
 
 	// Used to know if the deprecation of shared properties modified before the enforcement system introduction has been performed on load.
diff --git a/Engine/Source/Runtime/Landscape/Private/Landscape.cpp b/Engine/Source/Runtime/Landscape/Private/Landscape.cpp
--- a/Engine/Source/Runtime/Landscape/Private/Landscape.cpp
+++ b/Engine/Source/Runtime/Landscape/Private/Landscape.cpp
@@ -4119,7 +4119,6 @@ bool ULandscapeInfo::UpdateLayerInfoMapInternal(ALandscapeProxy* Proxy)
 	for (const TTuple<FName, FLandscapeTargetLayerSettings>& TargetLayer : LandscapeActor->GetTargetLayers())
 	{
 		FLandscapeInfoLayerSettings InfoLayerSettings (TargetLayer.Key, LandscapeActor.Get());
-		InfoLayerSettings.bValid = true;
 		InfoLayerSettings.LayerInfoObj = TargetLayer.Value.LayerInfoObj;
 		if (const FLandscapeInfoLayerSettings* PreviousLayer = Algo::FindBy(PreviousLayers, TargetLayer.Key, &FLandscapeInfoLayerSettings::LayerName))
 		{
diff --git a/Engine/Source/Runtime/Landscape/Private/LandscapeEdit.cpp b/Engine/Source/Runtime/Landscape/Private/LandscapeEdit.cpp
--- a/Engine/Source/Runtime/Landscape/Private/LandscapeEdit.cpp
+++ b/Engine/Source/Runtime/Landscape/Private/LandscapeEdit.cpp
@@ -6269,6 +6269,7 @@ void ALandscapeProxy::PostEditChangeProperty(FPropertyChangedEvent& PropertyChan
 	const FName SubPropertyName = PropertyChangedEvent.Property ? PropertyChangedEvent.Property->GetFName() : NAME_None;
 
 	bool bChangedPhysMaterial = false;
+	bool bChangedMaterial = false;
 
 	if (MemberPropertyName == FName(TEXT("RelativeScale3D")))
 	{
@@ -6455,6 +6456,58 @@ void ALandscapeProxy::PostEditChangeProperty(FPropertyChangedEvent& PropertyChan
 		InvalidateGeneratedComponentData(/* bInvalidateLightingCache = */false);
 		MarkComponentsRenderStateDirty();
 	}
+	else if ((MemberPropertyName == GET_MEMBER_NAME_CHECKED(ALandscapeProxy, LandscapeMaterial))
+		|| (MemberPropertyName == GET_MEMBER_NAME_CHECKED(ALandscapeProxy, LandscapeHoleMaterial))
+		|| (MemberPropertyName == GET_MEMBER_NAME_CHECKED(ALandscapeProxy, PerLODOverrideMaterials))
+		|| (SubPropertyName == GET_MEMBER_NAME_CHECKED(FLandscapePerLODMaterialOverride, Material))
+		|| (SubPropertyName == GET_MEMBER_NAME_CHECKED(FLandscapePerLODMaterialOverride, LODIndex)))
+	{
+		bChangedMaterial = true;
+
+		// Ignore changes when LODIndex slider is active or adding new PerLODMaterialOverride (Material is initially null, does not require rebuild)
+		if ((MemberPropertyName == GET_MEMBER_NAME_CHECKED(ALandscapeStreamingProxy, PerLODOverrideMaterials) && PropertyChangedEvent.ChangeType == EPropertyChangeType::ArrayAdd)
+			|| PropertyChangedEvent.ChangeType == EPropertyChangeType::Interactive)
+		{
+			bChangedMaterial = false;
+		}
+
+		if (ULandscapeInfo* Info = GetLandscapeInfo(); Info != nullptr && bChangedMaterial)
+		{
+			FMaterialUpdateContext MaterialUpdateContext;
+			Info->UpdateLayerInfoMap(/*this*/);
+
+			// Clear the parents out of combination material instances
+			for (const auto& MICPair : MaterialInstanceConstantMap)
+			{
+				UMaterialInstanceConstant* MaterialInstance = MICPair.Value;
+				MaterialInstance->BasePropertyOverrides.bOverride_BlendMode = false;
+				MaterialInstance->SetParentEditorOnly(nullptr);
+				MaterialUpdateContext.AddMaterialInstance(MaterialInstance);
+			}
+
+			// Remove our references to any material instances
+			MaterialInstanceConstantMap.Empty();
+
+			UpdateAllComponentMaterialInstances();
+
+			if (UWorld* World = GetWorld())
+			{
+				if (World->GetFeatureLevel() == ERHIFeatureLevel::ES3_1)
+				{
+					for (ULandscapeComponent* Component : LandscapeComponents)
+					{
+						if (Component != nullptr)
+						{
+							Component->CheckGenerateMobilePlatformData(/*bIsCooking = */ false, /*TargetPlatform = */ nullptr);
+						}
+					}
+				}
+
+				InvalidateOrUpdateNaniteRepresentation(/* bInCheckContentId = */true, /*InTargetPlatform = */nullptr);
+				UpdateRenderingMethod();
+			}
+		}
+	}
 
 	// Remove any null landscape components
 	LandscapeComponents.RemoveAll([](const ULandscapeComponent* Component) { return Component == nullptr; });
@@ -6468,6 +6521,15 @@ void ALandscapeProxy::PostEditChangeProperty(FPropertyChangedEvent& PropertyChan
 	{
 		ChangedPhysMaterial();
 	}
+
+	if (ULandscapeSubsystem* Subsystem = GetWorld()->GetSubsystem<ULandscapeSubsystem>(); Subsystem != nullptr && bChangedMaterial)
+	{
+		Subsystem->GetDelegateAccess().OnLandscapeProxyMaterialChangedDelegate.Broadcast(this, FOnLandscapeProxyMaterialChangedParams());
+
+PRAGMA_DISABLE_DEPRECATION_WARNINGS
+		LandscapeMaterialChangedDelegate.Broadcast();
+PRAGMA_ENABLE_DEPRECATION_WARNINGS
+	}
 }
 
 bool ALandscapeStreamingProxy::CanEditChange(const FProperty* InProperty) const
@@ -6510,62 +6572,6 @@ void ALandscapeStreamingProxy::PostEditChangeProperty(FPropertyChangedEvent& Pro
 			LandscapeActorRef = nullptr;
 		}
 	}
-	else if ((MemberPropertyName == GET_MEMBER_NAME_CHECKED(ALandscapeProxy, LandscapeMaterial))
-		|| (MemberPropertyName == GET_MEMBER_NAME_CHECKED(ALandscapeProxy, LandscapeHoleMaterial))
-		|| (MemberPropertyName == GET_MEMBER_NAME_CHECKED(ALandscapeStreamingProxy, PerLODOverrideMaterials)))
-	{
-		bool RecreateMaterialInstances = true;
-
-		if ((MemberPropertyName == GET_MEMBER_NAME_CHECKED(ALandscapeStreamingProxy, PerLODOverrideMaterials)) && (PropertyChangedEvent.ChangeType == EPropertyChangeType::ArrayAdd))
-		{
-			RecreateMaterialInstances = false;
-		}
-
-		if (RecreateMaterialInstances)
-		{
-			{
-				FMaterialUpdateContext MaterialUpdateContext;
-				if (ULandscapeInfo* LandscapeInfo = GetLandscapeInfo())
-				{
-					LandscapeInfo->UpdateLayerInfoMap(/*this*/);
-				}
-
-				// Clear the parents out of combination material instances
-				for (const auto& MICPair : MaterialInstanceConstantMap)
-				{
-					UMaterialInstanceConstant* MaterialInstance = MICPair.Value;
-					MaterialInstance->BasePropertyOverrides.bOverride_BlendMode = false;
-					MaterialInstance->SetParentEditorOnly(nullptr);
-					MaterialUpdateContext.AddMaterialInstance(MaterialInstance);
-				}
-
-				// Remove our references to any material instances
-				MaterialInstanceConstantMap.Empty();
-			}
-
-			UpdateAllComponentMaterialInstances();
-
-			UWorld* World = GetWorld();
-
-			if (World != nullptr)
-			{
-				if (World->GetFeatureLevel() == ERHIFeatureLevel::ES3_1)
-				{
-					for (ULandscapeComponent* Component : LandscapeComponents)
-					{
-						if (Component != nullptr)
-						{
-							Component->CheckGenerateMobilePlatformData(/*bIsCooking = */ false, /*TargetPlatform = */ nullptr);
-						}
-					}
-				}
-
-				InvalidateOrUpdateNaniteRepresentation(/* bInCheckContentId = */true, /*InTargetPlatform = */nullptr);
-
-				UpdateRenderingMethod();
-			}
-		}
-	}
 
 	// Must do this *after* clamping values
 	Super::PostEditChangeProperty(PropertyChangedEvent);
@@ -6837,15 +6843,6 @@ bool ULandscapeInfo::CanDeleteLandscape(FText& OutReason) const
 	return true;
 }
 
-void ALandscape::PreEditChange(FProperty* PropertyThatWillChange)
-{
-	PreEditLandscapeMaterial = LandscapeMaterial;
-	PreEditLandscapeHoleMaterial = LandscapeHoleMaterial;
-	PreEditPerLODOverrideMaterials = PerLODOverrideMaterials;
-
-	Super::PreEditChange(PropertyThatWillChange);
-}
-
 void ALandscape::PostEditChangeProperty(FPropertyChangedEvent& PropertyChangedEvent)
 {
 	TRACE_CPUPROFILER_EVENT_SCOPE(ALandscape::PostEditChangeProperty);
@@ -6866,6 +6863,7 @@ void ALandscape::PostEditChangeProperty(FPropertyChangedEvent& PropertyChangedEv
 	};
 
 	const FName MemberPropertyName = PropertyChangedEvent.MemberProperty ? PropertyChangedEvent.MemberProperty->GetFName() : NAME_None;
+	const FName SubPropertyName = PropertyChangedEvent.Property ? PropertyChangedEvent.Property->GetFName() : NAME_None;
 
 	bool bMaterialChanged = false;
 	bool bNeedsRecalcBoundingBox = false;
@@ -6876,60 +6874,19 @@ void ALandscape::PostEditChangeProperty(FPropertyChangedEvent& PropertyChangedEv
 
 	ULandscapeInfo* Info = GetLandscapeInfo();
 
-	if ((MemberPropertyName == GET_MEMBER_NAME_CHECKED(ALandscapeProxy, LandscapeMaterial)) 
-		|| (MemberPropertyName == GET_MEMBER_NAME_CHECKED(ALandscapeProxy, LandscapeHoleMaterial)) 
-		|| ((MemberPropertyName == GET_MEMBER_NAME_CHECKED(ALandscapeProxy, PerLODOverrideMaterials)) && (PropertyChangedEvent.ChangeType != EPropertyChangeType::ArrayAdd)))
+	if ((MemberPropertyName == GET_MEMBER_NAME_CHECKED(ALandscapeProxy, LandscapeMaterial))
+		|| (MemberPropertyName == GET_MEMBER_NAME_CHECKED(ALandscapeProxy, LandscapeHoleMaterial))
+		|| (MemberPropertyName == GET_MEMBER_NAME_CHECKED(ALandscapeProxy, PerLODOverrideMaterials))
+		|| (SubPropertyName == GET_MEMBER_NAME_CHECKED(FLandscapePerLODMaterialOverride, Material))
+		|| (SubPropertyName == GET_MEMBER_NAME_CHECKED(FLandscapePerLODMaterialOverride, LODIndex)))
 	{
-		bool bHasMaterialChanged = false;
-
-		if (PropertyChangedEvent.ChangeType != EPropertyChangeType::Interactive)
-		{
-			if (PreEditLandscapeMaterial != LandscapeMaterial || PreEditLandscapeHoleMaterial != LandscapeHoleMaterial || PreEditPerLODOverrideMaterials.Num() != PerLODOverrideMaterials.Num() || bIsPerformingInteractiveActionOnLandscapeMaterialOverride)
-			{
-				bHasMaterialChanged = true;
-			}
+		bMaterialChanged = true;
 
-			if (!bHasMaterialChanged)
-			{
-				for (int32 i = 0; i < PerLODOverrideMaterials.Num(); ++i)
-				{
-					const FLandscapePerLODMaterialOverride& NewMaterialOverride = PerLODOverrideMaterials[i];
-					const FLandscapePerLODMaterialOverride& PreEditMaterialOverride = PreEditPerLODOverrideMaterials[i];
-
-					if (!(PreEditMaterialOverride == NewMaterialOverride))
-					{
-						bHasMaterialChanged = true;
-						break;
-					}
-				}
-			}
-
-			bIsPerformingInteractiveActionOnLandscapeMaterialOverride = false;
-		}
-		else
+		// Ignore changes when LODIndex slider is active or adding new PerLODMaterialOverride (Material is initially null, does not require rebuild)
+		if ((MemberPropertyName == GET_MEMBER_NAME_CHECKED(ALandscapeStreamingProxy, PerLODOverrideMaterials) && PropertyChangedEvent.ChangeType == EPropertyChangeType::ArrayAdd)
+			|| PropertyChangedEvent.ChangeType == EPropertyChangeType::Interactive)
 		{
-			// We are probably using a slider or something similar in PerLODOverrideMaterials
-			bIsPerformingInteractiveActionOnLandscapeMaterialOverride = (MemberPropertyName == GET_MEMBER_NAME_CHECKED(ALandscapeProxy, PerLODOverrideMaterials));
-		}
-
-		if (Info != nullptr && bHasMaterialChanged)
-		{
-			FMaterialUpdateContext MaterialUpdateContext;
-			Info->UpdateLayerInfoMap(/*this*/);
-
-			bMaterialChanged = true;
-
-			// Clear the parents out of combination material instances
-			for (const auto& MICPair : MaterialInstanceConstantMap)
-			{
-				UMaterialInstanceConstant* MaterialInstance = MICPair.Value;
-				MaterialInstance->BasePropertyOverrides.bOverride_BlendMode = false;
-				MaterialInstance->SetParentEditorOnly(nullptr);
-				MaterialUpdateContext.AddMaterialInstance(MaterialInstance);
-			}
-
-			// Remove our references to any material instances
-			MaterialInstanceConstantMap.Empty();
+			bMaterialChanged = false;
 		}
 	}
 	else if (MemberPropertyName == FName(TEXT("RelativeScale3D")) ||
@@ -7110,17 +7067,7 @@ void ALandscape::PostEditChangeProperty(FPropertyChangedEvent& PropertyChangedEv
 				Gizmo->MarkComponentsRenderStateDirty();
 			}
 		}
-
-		// Must be done after the AActor::PostEditChange as we depend on the relinking of the landscapeInfo->LandscapeActor
-		if (bMaterialChanged)
-		{
-			LandscapeMaterialChangedDelegate.Broadcast();
-		}
 	}
-
-	PreEditLandscapeMaterial = nullptr;
-	PreEditLandscapeHoleMaterial = nullptr;
-	PreEditPerLODOverrideMaterials.Empty();
 }
 
 void ALandscapeProxy::ChangedPhysMaterial()
@@ -7171,7 +7118,7 @@ void ULandscapeComponent::PostEditChangeProperty(FPropertyChangedEvent& Property
 	{
 		bool RecreateMaterialInstances = true;
 
-		if ((MemberPropertyName == GET_MEMBER_NAME_CHECKED(ULandscapeComponent, PerLODOverrideMaterials)) && (PropertyChangedEvent.ChangeType == EPropertyChangeType::ArrayAdd))
+		if (MemberPropertyName == GET_MEMBER_NAME_CHECKED(ULandscapeComponent, PerLODOverrideMaterials) && PropertyChangedEvent.ChangeType == EPropertyChangeType::ArrayAdd)
 		{
 			RecreateMaterialInstances = false;
 		}
diff --git a/Engine/Source/Runtime/Landscape/Private/LandscapeInfo.cpp b/Engine/Source/Runtime/Landscape/Private/LandscapeInfo.cpp
--- a/Engine/Source/Runtime/Landscape/Private/LandscapeInfo.cpp
+++ b/Engine/Source/Runtime/Landscape/Private/LandscapeInfo.cpp
@@ -13,7 +13,6 @@ LANDSCAPE_API FLandscapeInfoLayerSettings::FLandscapeInfoLayerSettings(ULandscap
 	, ThumbnailMIC(nullptr)
 	, Owner(InProxy)
 	, DebugColorChannel(0)
-	, bValid(false)
 #endif
 {}
 
@@ -24,7 +23,6 @@ LANDSCAPE_API FLandscapeInfoLayerSettings::FLandscapeInfoLayerSettings(FName InP
 		, ThumbnailMIC(nullptr)
 		, Owner(InProxy)
 		, DebugColorChannel(0)
-		, bValid(false)
 #endif
 	{
 	}
\ No newline at end of file
diff --git a/Engine/Source/Runtime/Landscape/Private/LandscapeSubsystem.cpp b/Engine/Source/Runtime/Landscape/Private/LandscapeSubsystem.cpp
--- a/Engine/Source/Runtime/Landscape/Private/LandscapeSubsystem.cpp
+++ b/Engine/Source/Runtime/Landscape/Private/LandscapeSubsystem.cpp
@@ -1641,14 +1641,15 @@ void ULandscapeSubsystem::WaitLaunchNaniteBuild()
 	check((LandscapeMultithreadNaniteBuild != 0) || (NaniteStaticMeshesInFlight <= 1));
 }
 
-ULandscapeSubsystem::FDelegateAccess::FDelegateAccess(FOnHeightmapStreamedDelegate& InOnHeightmapStreamed, FOnLandscapeProxyComponentDataChanged& InOnLandscapeProxyComponentDataChanged) 
+ULandscapeSubsystem::FDelegateAccess::FDelegateAccess(FOnHeightmapStreamedDelegate& InOnHeightmapStreamed, FOnLandscapeProxyComponentDataChanged& InOnLandscapeProxyComponentDataChanged, FOnLandscapeProxyMaterialChanged& InOnLandscapeProxyMaterialChanged) 
 	: OnHeightmapStreamedDelegate(InOnHeightmapStreamed)
 	, OnLandscapeProxyComponentDataChangedDelegate(InOnLandscapeProxyComponentDataChanged)
+	, OnLandscapeProxyMaterialChangedDelegate(InOnLandscapeProxyMaterialChanged)
 {}
 
 ULandscapeSubsystem::FDelegateAccess ULandscapeSubsystem::GetDelegateAccess() const
 {
-	return FDelegateAccess(OnHeightmapStreamedDelegate, OnLandscapeProxyComponentDataChangedDelegate);
+	return FDelegateAccess(OnHeightmapStreamedDelegate, OnLandscapeProxyComponentDataChangedDelegate, OnLandscapeProxyMaterialChangedDelegate);
 }
 
 #endif // WITH_EDITOR
diff --git a/Engine/Source/Runtime/Landscape/Public/LandscapeSubsystem.h b/Engine/Source/Runtime/Landscape/Public/LandscapeSubsystem.h
--- a/Engine/Source/Runtime/Landscape/Public/LandscapeSubsystem.h
+++ b/Engine/Source/Runtime/Landscape/Public/LandscapeSubsystem.h
@@ -208,6 +208,11 @@ public:
 		return OnLandscapeProxyComponentDataChangedDelegate;
 	}
 
+	FOnLandscapeProxyMaterialChanged::RegistrationType& OnLandscapeProxyMaterialChanged() const
+	{
+		return OnLandscapeProxyMaterialChangedDelegate;
+	}
+
 	bool AnyViewShowCollisions() const { return bAnyViewShowCollisions; }  //! Returns true if any view has view collisions enabled.
 	FDateTime GetAppCurrentDateTime();
 
@@ -246,10 +251,11 @@ public:
 		friend ALandscape;
 		friend ALandscapeProxy;
 
-		FDelegateAccess(FOnHeightmapStreamedDelegate& InOnHeightmapStreamed, FOnLandscapeProxyComponentDataChanged& InOnLandscapeProxyComponentDataChanged);
+		FDelegateAccess(FOnHeightmapStreamedDelegate& InOnHeightmapStreamed, FOnLandscapeProxyComponentDataChanged& InOnLandscapeProxyComponentDataChanged, FOnLandscapeProxyMaterialChanged& InOnLandscapeProxyMaterialChanged);
 
 		FOnHeightmapStreamedDelegate& OnHeightmapStreamedDelegate;
 		FOnLandscapeProxyComponentDataChanged& OnLandscapeProxyComponentDataChangedDelegate;
+		FOnLandscapeProxyMaterialChanged& OnLandscapeProxyMaterialChangedDelegate;
 	};
 	FDelegateAccess GetDelegateAccess() const;
 #endif // WITH_EDITOR
@@ -304,6 +310,7 @@ private:
 #if WITH_EDITORONLY_DATA
 	mutable FOnHeightmapStreamedDelegate OnHeightmapStreamedDelegate;
 	mutable FOnLandscapeProxyComponentDataChanged OnLandscapeProxyComponentDataChangedDelegate;
+	mutable FOnLandscapeProxyMaterialChanged OnLandscapeProxyMaterialChangedDelegate;
 
 	class FLandscapePhysicalMaterialBuilder* PhysicalMaterialBuilder = nullptr;
 	FLandscapeNotificationManager* NotificationManager = nullptr;
