Landscape Editor - Replaced FLandscapeLayer with ULandscapeEditLayerBase accessors

diff --git a/Engine/Plugins/Experimental/LandscapePatch/Source/LandscapePatch/Private/LandscapePatchComponent.cpp b/Engine/Plugins/Experimental/LandscapePatch/Source/LandscapePatch/Private/LandscapePatchComponent.cpp
--- a/Engine/Plugins/Experimental/LandscapePatch/Source/LandscapePatch/Private/LandscapePatchComponent.cpp
+++ b/Engine/Plugins/Experimental/LandscapePatch/Source/LandscapePatch/Private/LandscapePatchComponent.cpp
@@ -211,8 +211,8 @@ void ULandscapePatchComponent::CheckForErrors()
 
 	if (EditLayerGuid.IsValid())
 	{
-		const FLandscapeLayer* Layer = nullptr;
-		if (!Landscape.IsValid() || !Landscape->CanHaveLayersContent() || (Layer = Landscape->GetLayerConst(EditLayerGuid)) == nullptr)
+		const ULandscapeEditLayerBase* LocalEditLayer = nullptr;
+		if (!Landscape.IsValid() || !Landscape->CanHaveLayersContent() || (LocalEditLayer = Landscape->GetEditLayerConst(EditLayerGuid)) == nullptr)
 		{
 			FMessageLog("MapCheck").Warning()
 			->AddToken(FUObjectToken::Create(this))
@@ -226,7 +226,7 @@ void ULandscapePatchComponent::CheckForErrors()
 				FixBindings();
 			})));
 		}
-		else if (Cast<ULandscapePatchEditLayer>(Layer->EditLayer) == nullptr)
+		else if (Cast<ULandscapePatchEditLayer>(LocalEditLayer) == nullptr)
 		{
 			FMessageLog("MapCheck").Warning()
 			->AddToken(FUObjectToken::Create(this))
@@ -817,31 +817,31 @@ bool ULandscapePatchComponent::BindToEditLayer(const FGuid& Guid)
 	}
 
 	Modify(false);
-	ULandscapePatchEditLayer* PreviousLayer = EditLayer.Get();
+	ULandscapePatchEditLayer* PreviousEditLayer = EditLayer.Get();
 
-	const FLandscapeLayer* LayerStruct = nullptr;
+	ULandscapeEditLayerBase* NewEditLayer = nullptr;
 
 	// See if the layer is in our current landscape
 	if (Landscape.Get())
 	{
-		LayerStruct = Landscape->GetLayerConst(Guid);
-		if (!LayerStruct && !bGaveMismatchedLandscapeWarning)
+		NewEditLayer = Landscape->GetEditLayer(Guid);
+		if (!NewEditLayer && !bGaveMismatchedLandscapeWarning)
 		{
 			UE_LOG(LogLandscapePatch, Warning, TEXT("Mismatch between landscape and layer Guid in patch."));
 			bGaveMismatchedLandscapeWarning = true;
 		}
 	}
 	
-	// If the layer wasn't in the current landscape, see if it's in some other landsape
-	if (!LayerStruct)
+	// If the layer wasn't in the current landscape, see if it's in some other landscape
+	if (!NewEditLayer)
 	{
 		UWorld* World = GetWorld();
 		if (ensure(World))
 		{
 			for (TActorIterator<ALandscape> LandscapeIterator(World); LandscapeIterator; ++LandscapeIterator)
 			{
-				LayerStruct = LandscapeIterator->GetLayerConst(Guid);
-				if (LayerStruct)
+				NewEditLayer = LandscapeIterator->GetEditLayer(Guid);
+				if (NewEditLayer)
 				{
 					// Found!
 					// Note that the layer we found might be of an incorrect type, so it's arguable whether landscape
@@ -858,7 +858,7 @@ bool ULandscapePatchComponent::BindToEditLayer(const FGuid& Guid)
 		Modify();  // Changing the guid, so mark the package dirty.
 	}
 
-	EditLayer = Cast<ULandscapePatchEditLayer>(LayerStruct ? LayerStruct->EditLayer : nullptr);
+	EditLayer = Cast<ULandscapePatchEditLayer>(NewEditLayer);
 	EditLayerGuid = Guid;
 	DetailPanelLayerGuid = Guid.ToString();
 	DetailPanelLayerName =
@@ -866,10 +866,10 @@ bool ULandscapePatchComponent::BindToEditLayer(const FGuid& Guid)
 		? EditLayer->GetName().ToString()
 		: TEXT("-Layer Not Found-");
 
-	if (PreviousLayer && EditLayer != PreviousLayer)
+	if (PreviousEditLayer && EditLayer != PreviousEditLayer)
 	{
-		PreviousLayer->NotifyOfPatchRemoval(this);
-		PreviousLayer->RequestLandscapeUpdate();
+		PreviousEditLayer->NotifyOfPatchRemoval(this);
+		PreviousEditLayer->RequestLandscapeUpdate();
 	}
 
 	if (EditLayer.IsValid() && IsPatchInWorld())
@@ -933,23 +933,24 @@ bool ULandscapePatchComponent::BindToLandscape(ALandscape* LandscapeIn)
 	}
 	ResetPatchManager();
 
-	const FLandscapeLayer* LayerStruct = nullptr;
-	ULandscapePatchEditLayer* Layer = nullptr;
+	const ULandscapeEditLayerBase* LocalEditLayer = nullptr;
+	const ULandscapePatchEditLayer* PatchEditLayer = nullptr;
+
 	if (EditLayerGuid.IsValid())
 	{
 		// See if this layer is here and of the appropriate type
-		LayerStruct = LandscapeIn->GetLayerConst(EditLayerGuid);
-		Layer = LayerStruct ? Cast<ULandscapePatchEditLayer>(LayerStruct->EditLayer) : nullptr;
+		LocalEditLayer = LandscapeIn->GetEditLayerConst(EditLayerGuid);
+		PatchEditLayer = LocalEditLayer ? Cast<ULandscapePatchEditLayer>(LocalEditLayer) : nullptr;
 	}
-	if (!Layer)
+	if (!PatchEditLayer)
 	{
 		// See if we have any layer of the appropriate type
-		LayerStruct = LandscapeIn->FindLayerOfTypeConst(ULandscapePatchEditLayer::StaticClass());
-		Layer = LayerStruct ? Cast<ULandscapePatchEditLayer>(LayerStruct->EditLayer) : nullptr;
+		LocalEditLayer = LandscapeIn->FindEditLayerOfTypeConst(ULandscapePatchEditLayer::StaticClass());
+		PatchEditLayer = LocalEditLayer ? Cast<ULandscapePatchEditLayer>(LocalEditLayer) : nullptr;
 	}
 
 	// If we couldn't find an appropriate layer, make a new one
-	if (!Layer && IsPatchInWorld() 
+	if (!PatchEditLayer && IsPatchInWorld() 
 		// We don't want to create a layer if this is the preview, because that will make the layer insertion
 		// not transactable, and we're not going to delete the layer if we don't complete the drop.
 		&& !IsPatchPreview())
@@ -959,12 +960,12 @@ bool ULandscapePatchComponent::BindToLandscape(ALandscape* LandscapeIn)
 
 		ILandscapeModule& LandscapeModule = FModuleManager::GetModuleChecked<ILandscapeModule>("Landscape");
 		int32 LayerIndex = LandscapeModule.GetLandscapeEditorServices()->GetOrCreateEditLayer(PatchLayerName, LandscapeIn, ULandscapePatchEditLayer::StaticClass());
-		LayerStruct = LandscapeIn->GetLayerConst(LayerIndex);
+		LocalEditLayer = LandscapeIn->GetEditLayerConst(LayerIndex);
 
-		Layer = LayerStruct ? Cast<ULandscapePatchEditLayer>(LayerStruct->EditLayer) : nullptr;
+		PatchEditLayer = LocalEditLayer ? Cast<ULandscapePatchEditLayer>(LocalEditLayer) : nullptr;
 	}
 
-	if (!Layer)
+	if (!PatchEditLayer)
 	{
 		// This happens if we failed to find a layer and either the patch was a preview, or IsPatchInWorld was false
 		// (the latter has happened while doing a SaveAs on a level).
@@ -972,7 +973,7 @@ bool ULandscapePatchComponent::BindToLandscape(ALandscape* LandscapeIn)
 		return false;
 	}
 
-	return BindToEditLayer(LayerStruct->EditLayer->GetGuid());
+	return BindToEditLayer(LocalEditLayer->GetGuid());
 }
 
 bool ULandscapePatchComponent::BindToAnyLandscape()
@@ -1012,8 +1013,8 @@ void ULandscapePatchComponent::UpdateEditLayerFromDetailPanelLayerName()
 		return;
 	}
 
-	const FLandscapeLayer* Layer = Landscape->GetLayerConst(*DetailPanelLayerName);
-	SetEditLayerGuid(Layer ? Layer->EditLayer->GetGuid() : FGuid());
+	const ULandscapeEditLayerBase* CurrentEditLayer = Landscape->GetEditLayerConst(*DetailPanelLayerName);
+	SetEditLayerGuid(CurrentEditLayer ? CurrentEditLayer->GetGuid() : FGuid());
 }
 
 void ULandscapePatchComponent::FixBindings()
@@ -1151,11 +1152,11 @@ const TArray<FString> ULandscapePatchComponent::GetLayerOptions()
 	// Get the layer names from our landscape
 	if (Landscape.IsValid())
 	{
-		for (const FLandscapeLayer& Layer : Landscape->GetLayersConst())
+		for (const ULandscapeEditLayerBase* CurrentEditLayer : Landscape->GetEditLayersConst())
 		{
-			if (Layer.EditLayer->IsA<ULandscapePatchEditLayer>())
+			if (CurrentEditLayer->IsA<ULandscapePatchEditLayer>())
 			{
-				Options.Add(Layer.EditLayer->GetName().ToString());
+				Options.Add(CurrentEditLayer->GetName().ToString());
 			}
 		}
 	}
diff --git a/Engine/Plugins/Experimental/LandscapePatch/Source/LandscapePatch/Private/LandscapePatchManager.cpp b/Engine/Plugins/Experimental/LandscapePatch/Source/LandscapePatch/Private/LandscapePatchManager.cpp
--- a/Engine/Plugins/Experimental/LandscapePatch/Source/LandscapePatch/Private/LandscapePatchManager.cpp
+++ b/Engine/Plugins/Experimental/LandscapePatch/Source/LandscapePatch/Private/LandscapePatchManager.cpp
@@ -532,8 +532,7 @@ void ALandscapePatchManager::MigrateToPrioritySystemAndDeleteInternal(bool bAllo
 		{
 			// Create the ULandscapePatchEditLayer in advance, if needed.  This prevents BindToLandscape from triggering the
 			// modal UI window for choosing the new layer index.
-			const FLandscapeLayer* LayerStruct = OwningLandscape->FindLayerOfTypeConst(ULandscapePatchEditLayer::StaticClass());
-			ULandscapePatchEditLayer* Layer = LayerStruct ? Cast<ULandscapePatchEditLayer>(LayerStruct->EditLayer) : nullptr;
+			const ULandscapeEditLayerBase* Layer = Cast<ULandscapePatchEditLayer>(OwningLandscape->FindEditLayerOfTypeConst(ULandscapePatchEditLayer::StaticClass()));
 
 			if (!Layer)
 			{
diff --git a/Engine/Plugins/Experimental/LandscapePatch/Source/LandscapePatch/Public/LandscapePatchComponent.h b/Engine/Plugins/Experimental/LandscapePatch/Source/LandscapePatch/Public/LandscapePatchComponent.h
--- a/Engine/Plugins/Experimental/LandscapePatch/Source/LandscapePatch/Public/LandscapePatchComponent.h
+++ b/Engine/Plugins/Experimental/LandscapePatch/Source/LandscapePatch/Public/LandscapePatchComponent.h
@@ -17,7 +17,6 @@ enum class ETeleportType : uint8;
 
 class ALandscape;
 class ALandscapePatchManager;
-struct FLandscapeLayer;
 class UTextureRenderTarget2D;
 
 UENUM()
diff --git a/Engine/Source/Editor/LandscapeEditor/Private/LandscapeEdMode.cpp b/Engine/Source/Editor/LandscapeEditor/Private/LandscapeEdMode.cpp
--- a/Engine/Source/Editor/LandscapeEditor/Private/LandscapeEdMode.cpp
+++ b/Engine/Source/Editor/LandscapeEditor/Private/LandscapeEdMode.cpp
@@ -320,9 +320,9 @@ ELandscapeToolTargetType FEdModeLandscape::GetLandscapeToolTargetType() const
 	return ELandscapeToolTargetType::Invalid;
 }
 
-const FLandscapeLayer* FEdModeLandscape::GetLandscapeSelectedLayer() const
+const ULandscapeEditLayerBase* FEdModeLandscape::GetLandscapeSelectedLayer() const
 {
-	return GetCurrentLayerConst();
+	return GetCurrentEditLayerConst();
 }
 
 ULandscapeLayerInfoObject* FEdModeLandscape::GetSelectedLandscapeLayerInfo() const
@@ -1665,15 +1665,15 @@ bool FEdModeLandscape::ProcessEditPaste()
 
 	bool Result = false;
 
-	
-	const FLandscapeLayer* SplinesLayer = nullptr;
+	const ULandscapeEditLayerBase* SplinesEditLayer = nullptr;
+
 	if (CurrentTool == (FLandscapeTool*)SplinesTool)
 	{
 		ALandscape* Landscape = GetLandscape();
-		SplinesLayer = Landscape ? Landscape->FindLayerOfTypeConst(ULandscapeEditLayerSplines::StaticClass()) : nullptr;
+		SplinesEditLayer = Landscape ? Landscape->FindEditLayerOfTypeConst(ULandscapeEditLayerSplines::StaticClass()) : nullptr;
 	}
 	FText Reason;
-	if (!CanEditLayer(&Reason, SplinesLayer))
+	if (!CanEditLayer(&Reason, SplinesEditLayer))
 	{
 		FMessageDialog::Open(EAppMsgType::Ok, Reason);
 		return Result;
@@ -2513,7 +2513,7 @@ bool FEdModeLandscape::IsGridBased() const
 
 bool FEdModeLandscape::HasValidLandscapeEditLayerSelection() const
 {
-	return !CanHaveLandscapeLayersContent() || GetCurrentLayerConst();
+	return !CanHaveLandscapeLayersContent() || GetCurrentEditLayerConst();
 }
 
 void FEdModeLandscape::SetTargetLandscape(const TWeakObjectPtr<ULandscapeInfo>& InLandscapeInfo)
@@ -3704,9 +3704,9 @@ ALandscape* FEdModeLandscape::ChangeComponentSetting(int32 NumComponentsX, int32
 				{
 					int32 HeightCount = 0;
 
-					for (const FLandscapeLayer& OldLayer : OldLandscape->GetLayersConst())
+					for (const ULandscapeEditLayerBase* OldEditLayer : OldLandscape->GetEditLayersConst())
 					{
-						FScopedSetLandscapeEditingLayer Scope(OldLandscape, OldLayer.EditLayer->GetGuid());
+						FScopedSetLandscapeEditingLayer Scope(OldLandscape, OldEditLayer->GetGuid());
 
 						TArray<uint16> HeightData;
 						TArray<FLandscapeImportLayerInfo> ImportMaterialLayerInfos;
@@ -3714,8 +3714,8 @@ ALandscape* FEdModeLandscape::ChangeComponentSetting(int32 NumComponentsX, int32
 						ExtractHeightmapWeightmapContent(HeightData, ImportMaterialLayerInfos);
 
 						HeightCount = FMath::Max(HeightCount, HeightData.Num());
-						HeightDataPerLayers.Add(OldLayer.EditLayer->GetGuid(), MoveTemp(HeightData));
-						ImportMaterialLayerInfosPerLayers.Add(OldLayer.EditLayer->GetGuid(), MoveTemp(ImportMaterialLayerInfos));
+						HeightDataPerLayers.Add(OldEditLayer->GetGuid(), MoveTemp(HeightData));
+						ImportMaterialLayerInfosPerLayers.Add(OldEditLayer->GetGuid(), MoveTemp(ImportMaterialLayerInfos));
 					}
 
 					TArray<uint16> DefaultHeightData;
@@ -3980,7 +3980,7 @@ bool FEdModeLandscape::HasLandscapeLayersContent() const
 int32 FEdModeLandscape::GetLayerCount() const
 {
 	ALandscape* Landscape = GetLandscape();
-	return Landscape ? Landscape->GetLayersConst().Num() : 0;
+	return Landscape ? Landscape->GetEditLayersConst().Num() : 0;
 }
 
 void FEdModeLandscape::SetCurrentLayer(int32 InLayerIndex)
@@ -3990,10 +3990,11 @@ void FEdModeLandscape::SetCurrentLayer(int32 InLayerIndex)
 
 	if (ALandscape* Landscape = GetLandscape())
 	{
-		const FLandscapeLayer* Layer = Landscape->GetLayerConst(InLayerIndex);
+		const ULandscapeEditLayerBase* EditLayer = Landscape->GetEditLayerConst(InLayerIndex);
+
 		if (Landscape->HasLayersContent() 
-			&& ensure(Layer && Layer->EditLayer)
-			&& Layer->EditLayer->IsA<ULandscapeEditLayerSplines>())
+			&& ensure(EditLayer)
+			&& EditLayer->IsA<ULandscapeEditLayerSplines>())
 		{
 			SetCurrentToolMode("ToolMode_Manage", false);
 			SetCurrentTool(FName("Splines"));
@@ -4014,12 +4015,6 @@ ALandscape* FEdModeLandscape::GetLandscape() const
 	return CurrentToolTarget.LandscapeInfo.IsValid() ? CurrentToolTarget.LandscapeInfo->LandscapeActor.Get() : nullptr;
 }
 
-const FLandscapeLayer* FEdModeLandscape::GetLayerConst(int32 InLayerIndex) const
-{
-	const ALandscape* Landscape = GetLandscape();
-	return Landscape ? Landscape->GetLayerConst(InLayerIndex) : nullptr;
-}
-
 bool FEdModeLandscape::CanRenameLayerTo(int32 InLayerIndex, const FName& InNewName) const
 {
 	if (ALandscape* Landscape = GetLandscape(); Landscape && !Landscape->IsLayerNameUnique(InNewName))
@@ -4127,9 +4122,9 @@ void FEdModeLandscape::DuplicateBrush(class ALandscapeBlueprintBrushBase* InBrus
 	GEditor->edactDuplicateSelected(InBrush->GetLevel(), false);
 }
 
-const FLandscapeLayer* FEdModeLandscape::GetCurrentLayerConst() const
+const ULandscapeEditLayerBase* FEdModeLandscape::GetCurrentEditLayerConst() const
 {
-	return GetLayerConst(GetCurrentLayerIndex());
+	return GetEditLayerConst(GetCurrentLayerIndex());
 }
 
 void FEdModeLandscape::AutoUpdateDirtyLandscapeSplines()
@@ -4146,12 +4141,12 @@ void FEdModeLandscape::AutoUpdateDirtyLandscapeSplines()
 	}
 }
 
-bool FEdModeLandscape::CanEditLayer(FText* Reason /*=nullptr*/, const FLandscapeLayer* InLayer /*= nullptr*/)
+bool FEdModeLandscape::CanEditLayer(FText* Reason /*=nullptr*/, const ULandscapeEditLayerBase* InLayer /*= nullptr*/)
 {
 	if (CanHaveLandscapeLayersContent())
 	{
 		ALandscape* Landscape = GetLandscape();
-		const FLandscapeLayer* TargetLayer = InLayer ? InLayer : GetCurrentLayerConst();
+		const ULandscapeEditLayerBase* TargetLayer = InLayer ? InLayer : GetCurrentEditLayerConst();
 		if (!TargetLayer)
 		{
 			if (Reason)
@@ -4161,16 +4156,16 @@ bool FEdModeLandscape::CanEditLayer(FText* Reason /*=nullptr*/, const FLandscape
 			return false;
 		}
 
-		check(TargetLayer->EditLayer != nullptr);
-		if (!TargetLayer->EditLayer->IsVisible())
+		check(TargetLayer != nullptr);
+		if (!TargetLayer->IsVisible())
 		{
 			if (Reason)
 			{
-				*Reason = NSLOCTEXT("UnrealEd", "LandscapeLayerHidden", "Painting or sculping in a hidden layer is not allowed.");
+				*Reason = NSLOCTEXT("UnrealEd", "LandscapeLayerHidden", "Painting or sculpting in a hidden layer is not allowed.");
 			}
 			return false;
 		}
-		else if (TargetLayer->EditLayer->IsLocked())
+		else if (TargetLayer->IsLocked())
 		{
 			if (Reason)
 			{
@@ -4183,7 +4178,7 @@ bool FEdModeLandscape::CanEditLayer(FText* Reason /*=nullptr*/, const FLandscape
 			// Special case for the splines tool, which is supported on the splines layer : 
 			if (CurrentTool != (FLandscapeTool*)SplinesTool)
 			{
-				if (Landscape && (TargetLayer == Landscape->FindLayerOfTypeConst(ULandscapeEditLayerSplines::StaticClass())))
+				if (Landscape && (TargetLayer == Landscape->FindEditLayerOfTypeConst(ULandscapeEditLayerSplines::StaticClass())))
 				{
 					if (Reason)
 					{
@@ -4192,11 +4187,11 @@ bool FEdModeLandscape::CanEditLayer(FText* Reason /*=nullptr*/, const FLandscape
 					return false;
 				}
 
-				if (!TargetLayer->EditLayer->SupportsEditingTools())
+				if (!TargetLayer->SupportsEditingTools())
 				{
 					if (Reason)
 					{
-						*Reason = FText::Format(NSLOCTEXT("UnrealEd", "LandscapeLayerEditLayerDoesntSupportEditing", "This layer's type ({0}) doesn't support direct editing."), TargetLayer->EditLayer->GetClass()->GetDisplayNameText());
+						*Reason = FText::Format(NSLOCTEXT("UnrealEd", "LandscapeLayerEditLayerDoesntSupportEditing", "This layer's type ({0}) doesn't support direct editing."), TargetLayer->GetClass()->GetDisplayNameText());
 					}
 					return false;
 				}
@@ -4249,8 +4244,8 @@ void FEdModeLandscape::UpdateLandscapeSplines(bool bUpdateOnlySelected /* = fals
 
 FGuid FEdModeLandscape::GetCurrentLayerGuid() const
 {
-	const FLandscapeLayer* CurrentLayer = GetCurrentLayerConst();
-	return CurrentLayer ? CurrentLayer->EditLayer->GetGuid() : FGuid();
+	const ULandscapeEditLayerBase* EditLayer = GetCurrentEditLayerConst();
+	return EditLayer ? EditLayer->GetGuid() : FGuid();
 }
 
 void FEdModeLandscape::UpdateBrushList()
diff --git a/Engine/Source/Editor/LandscapeEditor/Private/LandscapeEdMode.h b/Engine/Source/Editor/LandscapeEditor/Private/LandscapeEdMode.h
--- a/Engine/Source/Editor/LandscapeEditor/Private/LandscapeEdMode.h
+++ b/Engine/Source/Editor/LandscapeEditor/Private/LandscapeEdMode.h
@@ -394,7 +394,7 @@ public:
 	/** ILandscapeEdModeInterface */
 	virtual void PostUpdateLayerContent() override;
 	virtual ELandscapeToolTargetType GetLandscapeToolTargetType() const override;
-	virtual const FLandscapeLayer* GetLandscapeSelectedLayer() const override;
+	virtual const ULandscapeEditLayerBase* GetLandscapeSelectedLayer() const override;
 	virtual ULandscapeLayerInfoObject* GetSelectedLandscapeLayerInfo() const override;
 	virtual void OnCanHaveLayersContentChanged() override;
 
@@ -589,16 +589,15 @@ public:
 	void SetCurrentLayer(int32 InLayerIndex);
 	int32 GetCurrentLayerIndex() const;
 	ALandscape* GetLandscape() const;
-	const struct FLandscapeLayer* GetLayerConst(int32 InLayerIndex) const;
 	bool CanRenameLayerTo(int32 InLayerIndex, const FName& InNewName) const;
 	ULandscapeEditLayerBase* GetEditLayer(int32 InLayerIndex) const;
 	const ULandscapeEditLayerBase* GetEditLayerConst(int32 InLayerIndex) const;
 	bool IsLayerAlphaVisible(int32 InLayerIndex) const;
-	const struct FLandscapeLayer* GetCurrentLayerConst() const;
+	const ULandscapeEditLayerBase* GetCurrentEditLayerConst() const;
 	FGuid GetCurrentLayerGuid() const;
 	void UpdateLandscapeSplines(bool bUpdateOnlySelected = false);
 	void AutoUpdateDirtyLandscapeSplines();
-	bool CanEditLayer(FText* Reason = nullptr, const FLandscapeLayer* InLayer = nullptr);
+	bool CanEditLayer(FText* Reason = nullptr, const ULandscapeEditLayerBase* InLayer = nullptr);
 
 	void AddBrushToCurrentLayer(class ALandscapeBlueprintBrushBase* InBrush);
 	void RemoveBrushFromCurrentLayer(int32 InBrushIndex);
diff --git a/Engine/Source/Editor/LandscapeEditor/Private/LandscapeEdModeBPCustomTools.cpp b/Engine/Source/Editor/LandscapeEditor/Private/LandscapeEdModeBPCustomTools.cpp
--- a/Engine/Source/Editor/LandscapeEditor/Private/LandscapeEdModeBPCustomTools.cpp
+++ b/Engine/Source/Editor/LandscapeEditor/Private/LandscapeEdModeBPCustomTools.cpp
@@ -82,9 +82,9 @@ public:
 		}
 
 		// Only allow placing brushes if our landscape can be edited and the current edit layer supports BP brushes
-		const FLandscapeLayer* Layer = EdMode->GetCurrentLayerConst();
+		const ULandscapeEditLayerBase* EditLayer = EdMode->GetCurrentEditLayerConst();
 
-		if ((Layer && Layer->EditLayer->SupportsBlueprintBrushes()) && (DefaultObject->CanAffectHeightmap() || DefaultObject->CanAffectWeightmap() || DefaultObject->CanAffectVisibilityLayer()))
+		if ((EditLayer && EditLayer->SupportsBlueprintBrushes()) && (DefaultObject->CanAffectHeightmap() || DefaultObject->CanAffectWeightmap() || DefaultObject->CanAffectVisibilityLayer()))
 		{
 			ULandscapeInfo* Info = EdMode->CurrentToolTarget.LandscapeInfo.Get();
 			check(Info);
diff --git a/Engine/Source/Editor/LandscapeEditor/Private/LandscapeEdModeComponentTools.cpp b/Engine/Source/Editor/LandscapeEditor/Private/LandscapeEdModeComponentTools.cpp
--- a/Engine/Source/Editor/LandscapeEditor/Private/LandscapeEdModeComponentTools.cpp
+++ b/Engine/Source/Editor/LandscapeEditor/Private/LandscapeEdModeComponentTools.cpp
@@ -667,11 +667,11 @@ public:
 					TArray<ULandscapeComponent*> ComponentsUsingHeightmap;
 					ComponentsUsingHeightmap.Add(NewComponent);
 
-					for (const FLandscapeLayer& Layer : Landscape->GetLayersConst())
+					for (const ULandscapeEditLayerBase* EditLayer : Landscape->GetEditLayersConst())
 					{
 						// Since we do not share heightmap when adding new component, we will provide the required array, but they will only be used for 1 component
 						TMap<UTexture2D*, UTexture2D*> CreatedHeightmapTextures;
-						NewComponent->AddDefaultLayerData(Layer.EditLayer->GetGuid(), ComponentsUsingHeightmap, CreatedHeightmapTextures);
+						NewComponent->AddDefaultLayerData(EditLayer->GetGuid(), ComponentsUsingHeightmap, CreatedHeightmapTextures);
 					}
 				}
 
diff --git a/Engine/Source/Editor/LandscapeEditor/Private/LandscapeEdModeTools.h b/Engine/Source/Editor/LandscapeEditor/Private/LandscapeEdModeTools.h
--- a/Engine/Source/Editor/LandscapeEditor/Private/LandscapeEdModeTools.h
+++ b/Engine/Source/Editor/LandscapeEditor/Private/LandscapeEdModeTools.h
@@ -899,16 +899,16 @@ public:
 			if (bCombinedLayerOperation)
 			{
 				int32 I = 0;
-				for (const FLandscapeLayer& CurrentLayer : Landscape->GetLayersConst())
+				for (const ULandscapeEditLayerBase* EditLayer : Landscape->GetEditLayersConst())
 				{
-					BackupLayerVisibility.Add(CurrentLayer.EditLayer->IsVisible());
-					if (CurrentLayer.EditLayer->GetGuid() == EditingLayerGuid.GetValue())
+					BackupLayerVisibility.Add(EditLayer->IsVisible());
+					if (EditLayer->GetGuid() == EditingLayerGuid.GetValue())
 					{
 						EditingLayerIndex = I;
 					}
 					++I;
 				}
-				check(Landscape->GetLayersConst().IsValidIndex(EditingLayerIndex));
+				check(Landscape->GetEditLayersConst().IsValidIndex(EditingLayerIndex));
 			}
 			bIsInitialized = true;
 		}
@@ -922,9 +922,9 @@ public:
 		{
 			TArray<bool> NewLayerVisibility;
 			int32 I = 0;
-			for (const FLandscapeLayer& CurrentLayer : Landscape->GetLayersConst())
+			for (const ULandscapeEditLayerBase* EditLayer : Landscape->GetEditLayersConst())
 			{
-				NewLayerVisibility.Add((I > EditingLayerIndex) ? false : CurrentLayer.EditLayer->IsVisible());
+				NewLayerVisibility.Add((I > EditingLayerIndex) ? false : EditLayer->IsVisible());
 				++I;
 			}
 
@@ -980,9 +980,7 @@ public:
 		check(bIsInitialized);
 		if (bCombinedLayerOperation)
 		{
-			const FLandscapeLayer* Layer = Landscape->GetLayerConst(EditingLayerIndex);
-			check(Layer != nullptr);
-			const ULandscapeEditLayerBase* EditLayer = Layer->EditLayer;
+			const ULandscapeEditLayerBase* EditLayer = Landscape->GetEditLayerConst(EditingLayerIndex);
 			check(EditLayer != nullptr);
 			check(bTargetIsHeightmap);
 			const float Alpha = EditLayer->GetAlphaForTargetType(ELandscapeToolTargetType::Heightmap);
@@ -999,8 +997,8 @@ public:
 
 			FGuid CacheAccessorLayerGuid = CacheUpToEditingLayer.DataAccess.GetEditLayer();
  			checkf(EditingLayerGuid.GetValue() == CacheAccessorLayerGuid, TEXT("Editing Layer has changed between Initialize and Write. Was: %s (%s). Is now: %s (%s)"),
- 				Landscape->GetLayerConst(*EditingLayerGuid) ? *(Landscape->GetLayerConst(*EditingLayerGuid)->EditLayer->GetName().ToString()) : TEXT("<unknown>"), *EditingLayerGuid->ToString(),
- 				Landscape->GetLayerConst(CacheAccessorLayerGuid) ? *(Landscape->GetLayerConst(CacheAccessorLayerGuid)->EditLayer->GetName().ToString()) : TEXT("<unknown>"), *CacheAccessorLayerGuid.ToString());
+ 				Landscape->GetEditLayerConst(*EditingLayerGuid) ? *(Landscape->GetEditLayerConst(*EditingLayerGuid)->GetName().ToString()) : TEXT("<unknown>"), *EditingLayerGuid->ToString(),
+ 				Landscape->GetEditLayerConst(CacheAccessorLayerGuid) ? *(Landscape->GetEditLayerConst(CacheAccessorLayerGuid)->GetName().ToString()) : TEXT("<unknown>"), *CacheAccessorLayerGuid.ToString());
 
 			// Restore layers visibility
 			SetLayersVisibility(BackupLayerVisibility);
@@ -1039,16 +1037,16 @@ private:
 
 	void SetLayersVisibility(const TArray<bool>& InLayerVisibility)
 	{
-		check(InLayerVisibility.Num() == Landscape->GetLayersConst().Num());
+		check(InLayerVisibility.Num() == Landscape->GetEditLayersConst().Num());
 		for (int i = 0; i < InLayerVisibility.Num(); ++i)
 		{
-			if (const FLandscapeLayer* Layer = Landscape->GetLayerConst(i))
+			if (ULandscapeEditLayerBase* EditLayer = Landscape->GetEditLayer(i))
 			{
-				if (Layer->EditLayer->IsVisible() != InLayerVisibility[i])
+				if (EditLayer->IsVisible() != InLayerVisibility[i])
 				{
 					// Pass bInModify and bInBroadcastDataChange = false here, because we don't want to call Modify on the landscape, we don't actually want to truly change 
 					//  the layer visibility, only temporarily for the purpose of the intermediate render : 
-					Layer->EditLayer->SetVisible(InLayerVisibility[i], /*bInModify = */ false, /*bInBroadcastDataChange = */false);
+					EditLayer->SetVisible(InLayerVisibility[i], /*bInModify = */ false, /*bInBroadcastDataChange = */false);
 					bVisibilityChanged = true;
 				}
 			}
diff --git a/Engine/Source/Editor/LandscapeEditor/Private/LandscapeEditorDetailCustomization_Blueprint.cpp b/Engine/Source/Editor/LandscapeEditor/Private/LandscapeEditorDetailCustomization_Blueprint.cpp
--- a/Engine/Source/Editor/LandscapeEditor/Private/LandscapeEditorDetailCustomization_Blueprint.cpp
+++ b/Engine/Source/Editor/LandscapeEditor/Private/LandscapeEditorDetailCustomization_Blueprint.cpp
@@ -28,14 +28,13 @@ void FLandscapeEditorDetailCustomization_Blueprint::CustomizeDetails(IDetailLayo
 		// Tool is disabled for landscapes with edit layers disabled
 		if (LandscapeEdMode->CanHaveLandscapeLayersContent())
 		{
-			const FLandscapeLayer* Layer = LandscapeEdMode->GetCurrentLayerConst();
-			check(Layer != nullptr);
-
-			const bool bSupportsBPBrushes = Layer->EditLayer->SupportsBlueprintBrushes();
+			const ULandscapeEditLayerBase* EditLayer = LandscapeEdMode->GetCurrentEditLayerConst();
+			check(EditLayer != nullptr);
+			const bool bSupportsBPBrushes = EditLayer->SupportsBlueprintBrushes();
 
 			const FText EnabledToolTip = LOCTEXT("LandscapeBlueprintBrush_EnabledToolTip", "Selects the blueprint brush to apply to the current edit layer. Click on the landscape to apply it.");
 			const FText DisabledToolTip = FText::Format(LOCTEXT("LandscapeBlueprintBrush_DisabledToolTip", "Cannot add blueprint brush : the type of layer {0} ({1}) doesn't support blueprint brushes."),
-				FText::FromName(Layer->EditLayer->GetName()), Layer->EditLayer.GetClass()->GetDisplayNameText());
+				FText::FromName(EditLayer->GetName()), EditLayer->GetClass()->GetDisplayNameText());
 
 			TSharedRef<IPropertyHandle> PropertyHandle_Blueprint = DetailBuilder.GetProperty(GET_MEMBER_NAME_CHECKED(ULandscapeEditorObject, BlueprintBrush));
 			ToolsCategory.AddProperty(PropertyHandle_Blueprint)
diff --git a/Engine/Source/Editor/LandscapeEditor/Private/LandscapeEditorDetailCustomization_Layers.cpp b/Engine/Source/Editor/LandscapeEditor/Private/LandscapeEditorDetailCustomization_Layers.cpp
--- a/Engine/Source/Editor/LandscapeEditor/Private/LandscapeEditorDetailCustomization_Layers.cpp
+++ b/Engine/Source/Editor/LandscapeEditor/Private/LandscapeEditorDetailCustomization_Layers.cpp
@@ -537,15 +537,12 @@ TSharedPtr<SWidget> FLandscapeEditorCustomNodeBuilder_Layers::OnLayerContextMenu
 	ALandscape* Landscape = LandscapeEdMode ? LandscapeEdMode->GetLandscape() : nullptr;
 	if (Landscape && LandscapeEdMode && LandscapeEdMode->DoesCurrentToolAffectEditLayers())
 	{
-		const FLandscapeLayer* Layer = LandscapeEdMode->GetLayerConst(InLayerIndex);
+		const ULandscapeEditLayerBase* EditLayer = LandscapeEdMode->GetEditLayerConst(InLayerIndex);
 		FMenuBuilder MenuBuilder(true, nullptr);
 		MenuBuilder.BeginSection("LandscapeEditorLayerActions", LOCTEXT("LandscapeEditorLayerActions.Heading", "Edit Layers"));
 		{
-			if (Layer)
+			if (EditLayer)
 			{
-				const ULandscapeEditLayerBase* EditLayer = Layer->EditLayer;
-				check(EditLayer != nullptr);
-
 				// Rename Layer
 				FUIAction RenameLayerAction = FUIAction(
 					FExecuteAction::CreateSP(this, &FLandscapeEditorCustomNodeBuilder_Layers::RenameLayer, InLayerIndex),
@@ -600,9 +597,9 @@ TSharedPtr<SWidget> FLandscapeEditorCustomNodeBuilder_Layers::OnLayerContextMenu
 
 		MenuBuilder.BeginSection("LandscapeEditorLayerVisibility", LOCTEXT("LandscapeEditorLayerVisibility.Heading", "Visibility"));
 		{
-			if (Layer)
+			if (EditLayer)
 			{
-				if (Layer->EditLayer->IsVisible())
+				if (EditLayer->IsVisible())
 				{
 					// Hide Selected Layer
 					FUIAction HideSelectedLayerAction = FUIAction(FExecuteAction::CreateSPLambda(this, [this, InLayerIndex] { OnToggleVisibility(InLayerIndex); }),
@@ -655,15 +652,13 @@ TSharedPtr<SWidget> FLandscapeEditorCustomNodeBuilder_Layers::OnLayerContextMenu
 bool FLandscapeEditorCustomNodeBuilder_Layers::CanRenameLayer(int32 InLayerIndex, FText& OutReason) const
 {
 	FEdModeLandscape* LandscapeEdMode = GetEditorMode();
-	const FLandscapeLayer* Layer = LandscapeEdMode->GetLayerConst(InLayerIndex);
-	if (Layer == nullptr)
+	const ULandscapeEditLayerBase* EditLayer = LandscapeEdMode->GetEditLayerConst(InLayerIndex);
+
+	if (EditLayer == nullptr)
 	{
 		return false;
 	}
 
-	const ULandscapeEditLayerBase* EditLayer = Layer->EditLayer;
-	check(EditLayer != nullptr);
-
 	if (EditLayer->IsLocked())
 	{
 		OutReason = LOCTEXT("RenameLayer_CantRenameLocked", "Cannot rename a locked edit layer");
@@ -686,22 +681,20 @@ void FLandscapeEditorCustomNodeBuilder_Layers::RenameLayer(int32 InLayerIndex)
 bool FLandscapeEditorCustomNodeBuilder_Layers::CanClearTargetLayerOnLayer(int32 InLayerIndex, ULandscapeLayerInfoObject* InLayerInfo, FText& OutReason) const
 {
 	FEdModeLandscape* LandscapeEdMode = GetEditorMode();
-	const FLandscapeLayer* Layer = LandscapeEdMode->GetLayerConst(InLayerIndex);
-	if (Layer == nullptr)
+	const ULandscapeEditLayerBase* EditLayer = LandscapeEdMode->GetEditLayerConst(InLayerIndex);
+
+	if (EditLayer == nullptr)
 	{
 		return false;
 	}
 
-	const ULandscapeEditLayerBase* EditLayer = Layer->EditLayer;
-	check(EditLayer != nullptr);
-
 	// Visibility case : 
 	if (InLayerInfo->LayerName == UMaterialExpressionLandscapeVisibilityMask::ParameterName)
 	{
 		if (!EditLayer->SupportsTargetType(ELandscapeToolTargetType::Visibility))
 		{
 			OutReason = FText::Format(LOCTEXT("ClearTargetLayer_CantClearVisibilityOnLayerType", "Cannot clear visibility : the type of layer {0} ({1}) doesn't support visibility"),
-				FText::FromName(Layer->EditLayer->GetName()), EditLayer->GetClass()->GetDisplayNameText());
+				FText::FromName(EditLayer->GetName()), EditLayer->GetClass()->GetDisplayNameText());
 			return false;
 		}
 	}
@@ -709,7 +702,7 @@ bool FLandscapeEditorCustomNodeBuilder_Layers::CanClearTargetLayerOnLayer(int32
 	if (!EditLayer->SupportsTargetType(ELandscapeToolTargetType::Weightmap))
 	{
 		OutReason = FText::Format(LOCTEXT("ClearTargetLayer_CantClearWeightmapOnLayerType", "Cannot clear weightmap : the type of layer {0} ({1}) doesn't support weightmaps"),
-			FText::FromName(Layer->EditLayer->GetName()), EditLayer->GetClass()->GetDisplayNameText());
+			FText::FromName(EditLayer->GetName()), EditLayer->GetClass()->GetDisplayNameText());
 		return false;
 	}
 
@@ -723,10 +716,10 @@ void FLandscapeEditorCustomNodeBuilder_Layers::ClearTargetLayerOnLayer(int32 InL
 	ALandscape* Landscape = LandscapeEdMode ? LandscapeEdMode->GetLandscape() : nullptr;
 	if (Landscape)
 	{
-		const FLandscapeLayer* Layer = LandscapeEdMode->GetLayerConst(InLayerIndex);
-		if (Layer)
+		const ULandscapeEditLayerBase* EditLayer = LandscapeEdMode->GetEditLayerConst(InLayerIndex);
+		if (EditLayer)
 		{
-			EAppReturnType::Type Result = FMessageDialog::Open(EAppMsgType::YesNo, FText::Format(LOCTEXT("Landscape_ClearTargetLayer_Message", "The layer {0} : {1} content will be completely cleared.  Continue?"), FText::FromName(Layer->EditLayer->GetName()), FText::FromName(InLayerInfo->LayerName)));
+			EAppReturnType::Type Result = FMessageDialog::Open(EAppMsgType::YesNo, FText::Format(LOCTEXT("Landscape_ClearTargetLayer_Message", "The layer {0} : {1} content will be completely cleared.  Continue?"), FText::FromName(EditLayer->GetName()), FText::FromName(InLayerInfo->LayerName)));
 			if (Result == EAppReturnType::Yes)
 			{
 				const FScopedTransaction Transaction(LOCTEXT("Landscape_Layers_TargetClear", "Clear Target Layer"));
@@ -740,15 +733,13 @@ void FLandscapeEditorCustomNodeBuilder_Layers::ClearTargetLayerOnLayer(int32 InL
 bool FLandscapeEditorCustomNodeBuilder_Layers::CanClearLayer(int32 InLayerIndex, FText& OutReason) const
 {
 	FEdModeLandscape* LandscapeEdMode = GetEditorMode();
-	const FLandscapeLayer* Layer = LandscapeEdMode->GetLayerConst(InLayerIndex);
-	if (Layer == nullptr)
+	const ULandscapeEditLayerBase* EditLayer = LandscapeEdMode->GetEditLayerConst(InLayerIndex);
+
+	if (EditLayer == nullptr)
 	{
 		return false;
 	}
 
-	const ULandscapeEditLayerBase* EditLayer = Layer->EditLayer;
-	check(EditLayer != nullptr);
-
 	if (EditLayer->IsLocked())
 	{
 		OutReason = LOCTEXT("ClearLayer_CantClearLocked", "Cannot clear a locked edit layer");
@@ -758,7 +749,7 @@ bool FLandscapeEditorCustomNodeBuilder_Layers::CanClearLayer(int32 InLayerIndex,
 	if (!EditLayer->SupportsEditingTools())
 	{
 		OutReason = FText::Format(LOCTEXT("ClearLayer_CantClearLayerWithoutPersistentTextures", "Cannot clear an edit layer which doesn't have editable textures (procedural)"),
-			FText::FromName(Layer->EditLayer->GetName()), EditLayer->GetClass()->GetDisplayNameText());
+			FText::FromName(EditLayer->GetName()), EditLayer->GetClass()->GetDisplayNameText());
 		return false;
 	}
 
@@ -769,29 +760,26 @@ bool FLandscapeEditorCustomNodeBuilder_Layers::CanClearLayer(int32 InLayerIndex,
 bool FLandscapeEditorCustomNodeBuilder_Layers::CanClearTargetLayersOnLayer(int32 InLayerIndex, ELandscapeClearMode InClearMode, FText& OutReason) const
 {
 	FEdModeLandscape* LandscapeEdMode = GetEditorMode();
-	const FLandscapeLayer* Layer = LandscapeEdMode->GetLayerConst(InLayerIndex);
-	if (Layer == nullptr)
+	const ULandscapeEditLayerBase* EditLayer = LandscapeEdMode->GetEditLayerConst(InLayerIndex);
+	if (EditLayer == nullptr)
 	{
 		return false;
 	}
 
-	const ULandscapeEditLayerBase* EditLayer = Layer->EditLayer;
-	check(EditLayer != nullptr);
-
 	// TODO [jonathan.bard] : Ideally, ELandscapeClearMode should be deprecated and we would treat Visibility and Weightmap separately but for now just keep on treating visibility as a standard weightmap : 
 	if ((InClearMode == ELandscapeClearMode::Clear_Weightmap)
 		&& !EditLayer->SupportsTargetType(ELandscapeToolTargetType::Weightmap) 
 		&& !EditLayer->SupportsTargetType(ELandscapeToolTargetType::Visibility))
 	{
 			OutReason = FText::Format(LOCTEXT("ClearTargetLayers_CantClearWeightmapsOnLayerType", "Cannot clear weightmaps : the type of layer {0} ({1}) doesn't support weightmaps"), 
-				FText::FromName(Layer->EditLayer->GetName()), EditLayer->GetClass()->GetDisplayNameText());
+				FText::FromName(EditLayer->GetName()), EditLayer->GetClass()->GetDisplayNameText());
 			return false;
 	}
 	else if ((InClearMode == ELandscapeClearMode::Clear_Heightmap)
 		&& !EditLayer->SupportsTargetType(ELandscapeToolTargetType::Heightmap))
 	{
 		OutReason = FText::Format(LOCTEXT("ClearTargetLayers_CantClearHeightmapOnLayerType", "Cannot clear heightmap : the type of layer {0} ({1}) doesn't support heightmaps"),
-			FText::FromName(Layer->EditLayer->GetName()), EditLayer->GetClass()->GetDisplayNameText());
+			FText::FromName(EditLayer->GetName()), EditLayer->GetClass()->GetDisplayNameText());
 			return false;
 	}
 	if ((InClearMode == ELandscapeClearMode::Clear_All)
@@ -801,7 +789,7 @@ bool FLandscapeEditorCustomNodeBuilder_Layers::CanClearTargetLayersOnLayer(int32
 	if ((InClearMode == ELandscapeClearMode::Clear_Heightmap) && !EditLayer->SupportsTargetType(ELandscapeToolTargetType::Heightmap))
 	{
 		OutReason = FText::Format(LOCTEXT("ClearTargetLayers_CantClearOnLayerType", "Cannot clear : the type of layer {0} ({1}) doesn't support heightmaps or weightmaps"),
-			FText::FromName(Layer->EditLayer->GetName()), EditLayer->GetClass()->GetDisplayNameText());
+			FText::FromName(EditLayer->GetName()), EditLayer->GetClass()->GetDisplayNameText());
 		return false;
 	}
 
@@ -815,10 +803,10 @@ void FLandscapeEditorCustomNodeBuilder_Layers::ClearTargetLayersOnLayer(int32 In
 	ALandscape* Landscape = LandscapeEdMode ? LandscapeEdMode->GetLandscape() : nullptr;
 	if (Landscape)
 	{
-		const FLandscapeLayer* Layer = LandscapeEdMode->GetLayerConst(InLayerIndex);
-		if (Layer)
+		const ULandscapeEditLayerBase* EditLayer = LandscapeEdMode->GetEditLayerConst(InLayerIndex);
+		if (EditLayer)
 		{
-			EAppReturnType::Type Result = FMessageDialog::Open(EAppMsgType::YesNo, FText::Format(LOCTEXT("Landscape_ClearLayer_Message", "The layer {0} content will be completely cleared.  Continue?"), FText::FromName(Layer->EditLayer->GetName())));
+			EAppReturnType::Type Result = FMessageDialog::Open(EAppMsgType::YesNo, FText::Format(LOCTEXT("Landscape_ClearLayer_Message", "The layer {0} content will be completely cleared.  Continue?"), FText::FromName(EditLayer->GetName())));
 			if (Result == EAppReturnType::Yes)
 			{
 				const FScopedTransaction Transaction(LOCTEXT("Landscape_Layers_Clean", "Clear Layer"));
@@ -837,28 +825,24 @@ bool FLandscapeEditorCustomNodeBuilder_Layers::CanDeleteLayer(int32 InLayerIndex
 {
 	FEdModeLandscape* LandscapeEdMode = GetEditorMode();
 	ALandscape* Landscape = LandscapeEdMode ? LandscapeEdMode->GetLandscape() : nullptr;
-	if (!Landscape || Landscape->GetLayersConst().Num() <= 1)
+	if (!Landscape || Landscape->GetEditLayersConst().Num() <= 1)
 	{
 		OutReason = LOCTEXT("DeleteLayer_CantDeleteLastLayer", "The last layer cannot be deleted");
 		return false;
 	}
 
-	const FLandscapeLayer* Layer = LandscapeEdMode->GetLayerConst(InLayerIndex);
-	if (Layer == nullptr)
+	const ULandscapeEditLayerBase* EditLayer = LandscapeEdMode->GetEditLayerConst(InLayerIndex);
+	if (EditLayer == nullptr)
 	{
 		return false;
 	}
 
-	const ULandscapeEditLayerBase* EditLayer = Layer->EditLayer;
-	check(EditLayer != nullptr);
-
 	if (EditLayer->IsLocked())
 	{
 		OutReason = LOCTEXT("DeleteLayer_CantDeleteLocked", "Cannot delete a locked edit layer");
 		return false;
 	}
 
-
 	OutReason = LOCTEXT("DeleteLayer_CanDelete", "Delete the edit layer");
 	return true;
 }
@@ -867,12 +851,12 @@ void FLandscapeEditorCustomNodeBuilder_Layers::DeleteLayer(int32 InLayerIndex)
 {
 	FEdModeLandscape* LandscapeEdMode = GetEditorMode();
 	ALandscape* Landscape = LandscapeEdMode ? LandscapeEdMode->GetLandscape() : nullptr;
-	if (Landscape && Landscape->GetLayersConst().Num() > 1)
+	if (Landscape && Landscape->GetEditLayersConst().Num() > 1)
 	{
-		const FLandscapeLayer* Layer = LandscapeEdMode->GetLayerConst(InLayerIndex);
-		if (Layer)
+		const ULandscapeEditLayerBase* EditLayer = LandscapeEdMode->GetEditLayerConst(InLayerIndex);
+		if (EditLayer)
 		{
-			EAppReturnType::Type Result = FMessageDialog::Open(EAppMsgType::YesNo, FText::Format(LOCTEXT("Landscape_DeleteLayer_Message", "The layer {0} will be deleted.  Continue?"), FText::FromName(Layer->EditLayer->GetName())));
+			EAppReturnType::Type Result = FMessageDialog::Open(EAppMsgType::YesNo, FText::Format(LOCTEXT("Landscape_DeleteLayer_Message", "The layer {0} will be deleted.  Continue?"), FText::FromName(EditLayer->GetName())));
 			if (Result == EAppReturnType::Yes)
 			{
 				const FScopedTransaction Transaction(LOCTEXT("Landscape_Layers_Delete", "Delete Layer"));
@@ -889,7 +873,7 @@ bool FLandscapeEditorCustomNodeBuilder_Layers::CanCollapseLayer(int32 InLayerInd
 {
 	FEdModeLandscape* LandscapeEdMode = GetEditorMode();
 	ALandscape* Landscape = LandscapeEdMode ? LandscapeEdMode->GetLandscape() : nullptr;
-	if (!Landscape || Landscape->GetLayersConst().Num() <= 1)
+	if (!Landscape || Landscape->GetEditLayersConst().Num() <= 1)
 	{
 		OutReason = LOCTEXT("Landscape_CollapseLayer_Reason_NotEnoughLayersToCollapse", "Not enough layers to do collapse");
 		return false;
@@ -901,42 +885,43 @@ bool FLandscapeEditorCustomNodeBuilder_Layers::CanCollapseLayer(int32 InLayerInd
 		return false;
 	}
 	
-	const FLandscapeLayer* TopLayer = LandscapeEdMode->GetLayerConst(InLayerIndex);
-	const FLandscapeLayer* BottomLayer = LandscapeEdMode->GetLayerConst(InLayerIndex - 1);
+	const FLandscapeLayer* TopLayer = Landscape->GetLayerConst(InLayerIndex);
+	const FLandscapeLayer* BottomLayer = Landscape->GetLayerConst(InLayerIndex - 1);
 	if ((TopLayer == nullptr) || (BottomLayer == nullptr))
 	{
 		return false;
 	}
 
-	const ULandscapeEditLayerBase* TopEditLayer = TopLayer->EditLayer;
-	const ULandscapeEditLayerBase* BottomEditLayer = BottomLayer->EditLayer;
+	const ULandscapeEditLayerBase* TopEditLayer = LandscapeEdMode->GetEditLayerConst(InLayerIndex);
+	const ULandscapeEditLayerBase* BottomEditLayer = LandscapeEdMode->GetEditLayerConst(InLayerIndex - 1);
+
 	check((TopEditLayer != nullptr) && (BottomEditLayer != nullptr));
 
 	if (!TopEditLayer->SupportsCollapsingTo())
 	{
 		OutReason = FText::Format(LOCTEXT("Landscape_CollapseLayer_Reason_TopLayerDoesntSupportCollapsing", "Cannot collapse layer '{0}' onto layer '{1}'. The type of layer '{0}' ({2}) doesn't support collapsing to another one"), 
-			FText::FromName(TopLayer->EditLayer->GetName()), FText::FromName(BottomLayer->EditLayer->GetName()), TopEditLayer->GetClass()->GetDisplayNameText());
+			FText::FromName(TopEditLayer->GetName()), FText::FromName(BottomEditLayer->GetName()), TopEditLayer->GetClass()->GetDisplayNameText());
 		return false;
 	}
 
 	if (!BottomEditLayer->SupportsBeingCollapsedAway())
 	{
 		OutReason = FText::Format(LOCTEXT("Landscape_CollapseLayer_Reason_BottomLayerDoesntSupportCollapsing", "Cannot collapse layer '{0}' onto layer '{1}'. The type of layer '{1}' ({2}) doesn't support being collapsed away"),
-			FText::FromName(TopLayer->EditLayer->GetName()), FText::FromName(BottomLayer->EditLayer->GetName()), BottomEditLayer->GetClass()->GetDisplayNameText());
+			FText::FromName(TopEditLayer->GetName()), FText::FromName(BottomEditLayer->GetName()), BottomEditLayer->GetClass()->GetDisplayNameText());
 		return false;
 	}
 
 	if (TopEditLayer->IsLocked())
 	{
 		OutReason = FText::Format(LOCTEXT("Landscape_CollapseLayer_Reason_TopLayerIsLocked", "Cannot collapse layer '{0}' onto layer '{1}'. Layer '{0}' will be deleted in the operation but it is currently locked"),
-			FText::FromName(TopLayer->EditLayer->GetName()), FText::FromName(BottomLayer->EditLayer->GetName()));
+			FText::FromName(TopEditLayer->GetName()), FText::FromName(BottomEditLayer->GetName()));
 		return false;
 	}
 
 	if (BottomEditLayer->IsLocked())
 	{
 		OutReason = FText::Format(LOCTEXT("Landscape_CollapseLayer_Reason_BottomLayerIsLocked", "Cannot collapse layer '{0}' onto layer '{1}'. Destination layer '{1}' is currently locked"),
-			FText::FromName(TopLayer->EditLayer->GetName()), FText::FromName(BottomLayer->EditLayer->GetName()));
+			FText::FromName(TopEditLayer->GetName()), FText::FromName(BottomEditLayer->GetName()));
 		return false;
 	}
 
@@ -944,19 +929,19 @@ bool FLandscapeEditorCustomNodeBuilder_Layers::CanCollapseLayer(int32 InLayerInd
 	if (TopLayer->Brushes.Num() > 0)
 	{
 		OutReason = FText::Format(LOCTEXT("Landscape_CollapseLayer_Reason_TopLayerHasBrush", "Cannot collapse layer '{0}' onto layer '{1}'. Layer '{0}' contains brush(es)"),
-			FText::FromName(TopLayer->EditLayer->GetName()), FText::FromName(BottomLayer->EditLayer->GetName()));
+			FText::FromName(TopEditLayer->GetName()), FText::FromName(BottomEditLayer->GetName()));
 		return false;
 	}
 
 	if (BottomLayer->Brushes.Num() > 0)
 	{
 		OutReason = FText::Format(LOCTEXT("Landscape_CollapseLayer_Reason_BottomLayerHasBrush", "Cannot collapse layer '{0}' onto layer '{1}'. Layer '{1}' contains brush(es)"),
-			FText::FromName(TopLayer->EditLayer->GetName()), FText::FromName(BottomLayer->EditLayer->GetName()));
+			FText::FromName(TopEditLayer->GetName()), FText::FromName(BottomEditLayer->GetName()));
 		return false;
 	}
 	
 	OutReason = FText::Format(LOCTEXT("Landscape_CollapseLayer_Reason_Collapse", "Collapse layer '{0}' onto layer '{1}'"),
-		FText::FromName(TopLayer->EditLayer->GetName()), FText::FromName(BottomLayer->EditLayer->GetName()));
+		FText::FromName(TopEditLayer->GetName()), FText::FromName(BottomEditLayer->GetName()));
 	return true;
 }
 
@@ -968,11 +953,12 @@ void FLandscapeEditorCustomNodeBuilder_Layers::CollapseLayer(int32 InLayerIndex)
 	if (FEdModeLandscape* LandscapeEdMode = GetEditorMode())
 	{
 		ALandscape* Landscape = LandscapeEdMode->GetLandscape();
-		const FLandscapeLayer* Layer = LandscapeEdMode->GetLayerConst(InLayerIndex);
-		const FLandscapeLayer* BaseLayer = LandscapeEdMode->GetLayerConst(InLayerIndex - 1);
+		const ULandscapeEditLayerBase* Layer = LandscapeEdMode->GetEditLayerConst(InLayerIndex);
+		const ULandscapeEditLayerBase* BaseLayer = LandscapeEdMode->GetEditLayerConst(InLayerIndex - 1);
+
 		if (Landscape && Layer && BaseLayer)
 		{
-			EAppReturnType::Type Result = FMessageDialog::Open(EAppMsgType::YesNo, FText::Format(LOCTEXT("Landscape_CollapseLayer_Message", "The layer {0} will be collapsed into layer {1}.  Continue?"), FText::FromName(Layer->EditLayer->GetName()), FText::FromName(BaseLayer->EditLayer->GetName())));
+			EAppReturnType::Type Result = FMessageDialog::Open(EAppMsgType::YesNo, FText::Format(LOCTEXT("Landscape_CollapseLayer_Message", "The layer {0} will be collapsed into layer {1}.  Continue?"), FText::FromName(Layer->GetName()), FText::FromName(BaseLayer->GetName())));
 			if (Result == EAppReturnType::Yes)
 			{
 				const FScopedTransaction Transaction(LOCTEXT("Landscape_Layers_Collapse", "Collapse Layer"));
@@ -994,15 +980,13 @@ bool FLandscapeEditorCustomNodeBuilder_Layers::CanExecuteCustomLayerAction(int32
 		return false;
 	}
 
-	const FLandscapeLayer* Layer = LandscapeEdMode->GetLayerConst(InLayerIndex);
-	if (Layer == nullptr)
+	const ULandscapeEditLayerBase* EditLayer = LandscapeEdMode->GetEditLayerConst(InLayerIndex);
+	if (EditLayer == nullptr)
 	{
 		return false;
 	}
 
-	const ULandscapeEditLayerBase* EditLayer = Layer->EditLayer;
-	check(EditLayer != nullptr);
-	ULandscapeEditLayerBase::FEditLayerAction::FExecuteParams ExecuteParams(Layer, Landscape);
+	ULandscapeEditLayerBase::FEditLayerAction::FExecuteParams ExecuteParams(EditLayer, Landscape);
 	return InCustomLayerAction.GetCanExecuteDelegate().Execute(ExecuteParams, OutReason);
 }
 
@@ -1011,14 +995,11 @@ void FLandscapeEditorCustomNodeBuilder_Layers::ExecuteCustomLayerAction(int32 In
 	FEdModeLandscape* LandscapeEdMode = GetEditorMode();
 	if (ALandscape* Landscape = LandscapeEdMode ? LandscapeEdMode->GetLandscape() : nullptr)
 	{
-		const FLandscapeLayer* Layer = LandscapeEdMode->GetLayerConst(InLayerIndex);
-		if (Layer != nullptr)
+		const ULandscapeEditLayerBase* EditLayer = LandscapeEdMode->GetEditLayerConst(InLayerIndex);
+		if (EditLayer != nullptr)
 		{
-			const ULandscapeEditLayerBase* EditLayer = Layer->EditLayer;
-			check(EditLayer != nullptr);
-
 			const FScopedTransaction Transaction(InCustomLayerAction.GetLabel());
-			ULandscapeEditLayerBase::FEditLayerAction::FExecuteParams ExecuteParams(Layer, Landscape);
+			ULandscapeEditLayerBase::FEditLayerAction::FExecuteParams ExecuteParams(EditLayer, Landscape);
 			ULandscapeEditLayerBase::FEditLayerAction::FExecuteResult Result = InCustomLayerAction.GetExecuteDelegate().Execute(ExecuteParams);
 			if (!Result.bSuccess)
 			{
@@ -1122,7 +1103,7 @@ void FLandscapeEditorCustomNodeBuilder_Layers::CreateLayer()
 
 			const FScopedTransaction Transaction(LOCTEXT("Landscape_Layers_Create", "Create Layer"));
 			Landscape->CreateLayer(FName(EditLayerClass.GetDefaultObject()->GetDefaultName()), EditLayerClass);
-			OnLayerSelectionChanged(Landscape->GetLayersConst().Num() - 1);
+			OnLayerSelectionChanged(Landscape->GetEditLayersConst().Num() - 1);
 
 			LandscapeEdMode->RefreshDetailPanel();
 		}
@@ -1135,7 +1116,7 @@ FText FLandscapeEditorCustomNodeBuilder_Layers::GetNumLayersText() const
 	ALandscape* Landscape = LandscapeEdMode ? LandscapeEdMode->GetLandscape() : nullptr;
 	if (Landscape)
 	{
-		return FText::Format(LOCTEXT("NumEditLayersText", "{0} Edit {0}|plural(one=Layer, other=Layers)"), Landscape->GetLayersConst().Num());
+		return FText::Format(LOCTEXT("NumEditLayersText", "{0} Edit {0}|plural(one=Layer, other=Layers)"), Landscape->GetEditLayersConst().Num());
 	}
 
 	return FText();
@@ -1172,9 +1153,7 @@ TOptional<float> FLandscapeEditorCustomNodeBuilder_Layers::GetLayerAlpha(int32 I
 {
 	if (FEdModeLandscape* LandscapeEdMode = GetEditorMode())
 	{
-		const FLandscapeLayer* Layer = LandscapeEdMode->GetLayerConst(InLayerIndex);
-		check(Layer != nullptr);
-		const ULandscapeEditLayerBase* EditLayer = Layer->EditLayer;
+		const ULandscapeEditLayerBase* EditLayer = LandscapeEdMode->GetEditLayerConst(InLayerIndex);
 		check(EditLayer != nullptr);
 
 		return EditLayer->GetAlphaForTargetType(LandscapeEdMode->GetLandscapeToolTargetType());
@@ -1204,19 +1183,16 @@ bool FLandscapeEditorCustomNodeBuilder_Layers::CanSetLayerAlpha(int32 InLayerInd
 		return false;
 	}
 
-	const FLandscapeLayer* Layer = LandscapeEdMode->GetLayerConst(InLayerIndex);
-	if (Layer == nullptr)
+	const ULandscapeEditLayerBase* EditLayer = LandscapeEdMode->GetEditLayerConst(InLayerIndex);
+	if (EditLayer == nullptr)
 	{
 		return false;
 	}
 
-	const ULandscapeEditLayerBase* EditLayer = Layer->EditLayer;
-	check(EditLayer != nullptr);
-
 	if (!EditLayer->SupportsAlphaForTargetType(LandscapeEdMode->GetLandscapeToolTargetType()))
 	{
 		OutReason = FText::Format(LOCTEXT("SetLayerAlpha_LayerDoesntSupportAlpha", "Cannot change alpha : the type of layer {0} ({1}) doesn't support alpha"),
-			FText::FromName(Layer->EditLayer->GetName()), EditLayer->GetClass()->GetDisplayNameText());
+			FText::FromName(EditLayer->GetName()), EditLayer->GetClass()->GetDisplayNameText());
 		return false;
 	}
 
@@ -1248,15 +1224,13 @@ void FLandscapeEditorCustomNodeBuilder_Layers::SetLayerAlpha(float InAlpha, int3
 bool FLandscapeEditorCustomNodeBuilder_Layers::CanToggleVisibility(int32 InLayerIndex, FText& OutReason) const
 {
 	FEdModeLandscape* LandscapeEdMode = GetEditorMode();
-	const FLandscapeLayer* Layer = LandscapeEdMode->GetLayerConst(InLayerIndex);
-	if (Layer == nullptr)
+
+	const ULandscapeEditLayerBase* EditLayer = LandscapeEdMode->GetEditLayerConst(InLayerIndex);
+	if (EditLayer == nullptr)
 	{
 		return false;
 	}
 
-	const ULandscapeEditLayerBase* EditLayer = Layer->EditLayer;
-	check(EditLayer != nullptr);
-
 	if (EditLayer->IsLocked())
 	{
 		OutReason = LOCTEXT("ToggleVisibility_CantToggleLocked", "Cannot change the visibility of a locked edit layer");
@@ -1367,8 +1341,8 @@ int32 FLandscapeEditorCustomNodeBuilder_Layers::SlotIndexToLayerIndex(int32 Slot
 		return INDEX_NONE;
 	}
 	
-	check(Landscape->GetLayersConst().IsValidIndex(SlotIndex));
-	return Landscape->GetLayersConst().Num() - SlotIndex - 1;
+	check(Landscape->GetEditLayersConst().IsValidIndex(SlotIndex));
+	return Landscape->GetEditLayersConst().Num() - SlotIndex - 1;
 }
 
 FReply FLandscapeEditorCustomNodeBuilder_Layers::HandleDragDetected(const FGeometry& MyGeometry, const FPointerEvent& MouseEvent, int32 SlotIndex, SVerticalBox::FSlot* Slot)
@@ -1377,10 +1351,8 @@ FReply FLandscapeEditorCustomNodeBuilder_Layers::HandleDragDetected(const FGeome
 	if (LandscapeEdMode && LandscapeEdMode->DoesCurrentToolAffectEditLayers())
 	{
 		int32 LayerIndex = SlotIndexToLayerIndex(SlotIndex);
-		if (const FLandscapeLayer* Layer = LandscapeEdMode->GetLayerConst(LayerIndex))
+		if (const ULandscapeEditLayerBase* EditLayer = LandscapeEdMode->GetEditLayerConst(LayerIndex))
 		{
-			const ULandscapeEditLayerBase* EditLayer = Layer->EditLayer;
-			check(EditLayer != nullptr);
 			if (!EditLayer->IsLocked())
 			{
 				TSharedPtr<SWidget> Row = GenerateRow(LayerIndex);
diff --git a/Engine/Source/Editor/LandscapeEditor/Private/LandscapeEditorDetailCustomization_LayersBrushStack.cpp b/Engine/Source/Editor/LandscapeEditor/Private/LandscapeEditorDetailCustomization_LayersBrushStack.cpp
--- a/Engine/Source/Editor/LandscapeEditor/Private/LandscapeEditorDetailCustomization_LayersBrushStack.cpp
+++ b/Engine/Source/Editor/LandscapeEditor/Private/LandscapeEditorDetailCustomization_LayersBrushStack.cpp
@@ -87,8 +87,8 @@ void FLandscapeEditorCustomNodeBuilder_LayersBrushStack::GenerateChildContent(ID
 	if (FEdModeLandscape* LandscapeEdMode = GetEditorMode(); LandscapeEdMode && LandscapeEdMode->CurrentToolMode)
 	{
 		const int32 BrushCount = LandscapeEdMode->GetBrushesForCurrentLayer().Num();
-		const FLandscapeLayer* Layer = LandscapeEdMode->GetCurrentLayerConst();
-		check(Layer != nullptr && Layer->EditLayer != nullptr);
+		const ULandscapeEditLayerBase* EditLayer = LandscapeEdMode->GetCurrentEditLayerConst();
+		check(EditLayer != nullptr);
 
 		TSharedPtr<SDragAndDropVerticalBox> BrushesList = SNew(SDragAndDropVerticalBox)
 			.OnCanAcceptDropAdvanced(this, &FLandscapeEditorCustomNodeBuilder_LayersBrushStack::HandleCanAcceptDrop)
@@ -109,7 +109,7 @@ void FLandscapeEditorCustomNodeBuilder_LayersBrushStack::GenerateChildContent(ID
 
 		// Display a disabled message when the current edit layer does not support BP brushes
 		const FText DisabledToolTip = FText::Format(LOCTEXT("LandscapeEditLayerBlueprintBrushes_Disabled", "The type of layer {0} ({1}) doesn't support blueprint brushes."),
-			FText::FromName(Layer->EditLayer->GetName()), Layer->EditLayer->GetClass()->GetDisplayNameText());
+			FText::FromName(EditLayer->GetName()), EditLayer->GetClass()->GetDisplayNameText());
 
 		const TSharedPtr<SWidget> ErrorMessage = SNew(SMultiLineEditableTextBox)
 			.IsReadOnly(true)
@@ -121,7 +121,7 @@ void FLandscapeEditorCustomNodeBuilder_LayersBrushStack::GenerateChildContent(ID
 		ChildrenBuilder.AddCustomRow(LOCTEXT("LandscapeEditLayerBlueprintBrushes", "Edit Layer Blueprint Brushes"))
 			.WholeRowContent()
 			[
-				Layer->EditLayer->SupportsBlueprintBrushes() ? BrushesListWidget.ToSharedRef() : ErrorMessage.ToSharedRef()
+				EditLayer->SupportsBlueprintBrushes() ? BrushesListWidget.ToSharedRef() : ErrorMessage.ToSharedRef()
 			];
 
 		for (int32 i = 0; i < BrushCount; ++i)
@@ -518,11 +518,8 @@ FReply FLandscapeEditorCustomNodeBuilder_LayersBrushStack::HandleDragDetected(co
 {
 	if (FEdModeLandscape* LandscapeEdMode = GetEditorMode())
 	{
-		if (const FLandscapeLayer* Layer = LandscapeEdMode->GetCurrentLayerConst())
+		if (const ULandscapeEditLayerBase* EditLayer = LandscapeEdMode->GetCurrentEditLayerConst())
 		{
-			const ULandscapeEditLayerBase* EditLayer = Layer->EditLayer;
-			check(EditLayer != nullptr);
-
 			const TArray<ALandscapeBlueprintBrushBase*>& BrushStack = LandscapeEdMode->GetBrushesForCurrentLayer();
 			if (!EditLayer->IsLocked() && BrushStack.IsValidIndex(SlotIndex))
 			{
diff --git a/Engine/Source/Editor/LandscapeEditor/Private/LandscapeEditorDetailCustomization_NewLandscape.cpp b/Engine/Source/Editor/LandscapeEditor/Private/LandscapeEditorDetailCustomization_NewLandscape.cpp
--- a/Engine/Source/Editor/LandscapeEditor/Private/LandscapeEditorDetailCustomization_NewLandscape.cpp
+++ b/Engine/Source/Editor/LandscapeEditor/Private/LandscapeEditorDetailCustomization_NewLandscape.cpp
@@ -1111,11 +1111,11 @@ void FLandscapeEditorDetailCustomization_NewLandscape::AddComponents(ULandscapeI
 			TArray<ULandscapeComponent*> ComponentsUsingHeightmap;
 			ComponentsUsingHeightmap.Add(NewComponent);
 
-			for (const FLandscapeLayer& Layer : Landscape->GetLayersConst())
+			for (const ULandscapeEditLayerBase* EditLayer : Landscape->GetEditLayersConst())
 			{
 				// Since we do not share heightmap when adding new component, we will provided the required array, but they will only be used for 1 component
 				TMap<UTexture2D*, UTexture2D*> CreatedHeightmapTextures;
-				NewComponent->AddDefaultLayerData(Layer.EditLayer->GetGuid(), ComponentsUsingHeightmap, CreatedHeightmapTextures);
+				NewComponent->AddDefaultLayerData(EditLayer->GetGuid(), ComponentsUsingHeightmap, CreatedHeightmapTextures);
 			}
 		}
 
diff --git a/Engine/Source/Editor/LandscapeEditor/Private/LandscapeUIDetails.cpp b/Engine/Source/Editor/LandscapeEditor/Private/LandscapeUIDetails.cpp
--- a/Engine/Source/Editor/LandscapeEditor/Private/LandscapeUIDetails.cpp
+++ b/Engine/Source/Editor/LandscapeEditor/Private/LandscapeUIDetails.cpp
@@ -110,9 +110,9 @@ void FLandscapeUIDetails::ToggleCanHaveLayersContent(ALandscape* Landscape)
 	if (Landscape->bCanHaveLayersContent)
 	{
 		bool bHasHiddenLayers = false;
-		for (const FLandscapeLayer& Layer : Landscape->GetLayersConst())
+		for (const ULandscapeEditLayerBase* EditLayer : Landscape->GetEditLayersConst())
 		{
-			if (!Layer.EditLayer->IsVisible())
+			if (!EditLayer->IsVisible())
 			{
 				bHasHiddenLayers = true;
 				break;
diff --git a/Engine/Source/Editor/LandscapeEditor/Private/SLandscapeLayerListDialog.cpp b/Engine/Source/Editor/LandscapeEditor/Private/SLandscapeLayerListDialog.cpp
--- a/Engine/Source/Editor/LandscapeEditor/Private/SLandscapeLayerListDialog.cpp
+++ b/Engine/Source/Editor/LandscapeEditor/Private/SLandscapeLayerListDialog.cpp
@@ -15,9 +15,9 @@
 
 struct FWidgetLayerListItem
 {
-	FWidgetLayerListItem(const FLandscapeLayer& InLayer, const TFunction<void(void)>& InOnLayerListUpdated, TArray<TSharedPtr<FWidgetLayerListItem>>* InWidgetLayerList)
-	: LayerName(InLayer.EditLayer->GetName())
-	, LayerGuid(InLayer.EditLayer->GetGuid())
+	FWidgetLayerListItem(const ULandscapeEditLayerBase* InEditLayer, const TFunction<void(void)>& InOnLayerListUpdated, TArray<TSharedPtr<FWidgetLayerListItem>>* InWidgetLayerList)
+	: LayerName(InEditLayer->GetName())
+	, LayerGuid(InEditLayer->GetGuid())
 	, OnLayerListUpdated(InOnLayerListUpdated)
 	, WidgetLayerList(InWidgetLayerList)
 	{ 
@@ -182,15 +182,15 @@ void SLandscapeLayerListDialog::Construct(const FArguments& InArgs, const TWeakO
 	check(Landscape.IsValid());
 
 	int32 LayerIndex = 0;
-	TArrayView<const FLandscapeLayer> LandscapeLayers = Landscape->GetLayersConst();
-	for (const FLandscapeLayer& Layer : LandscapeLayers)
+	const TArray<const ULandscapeEditLayerBase*> LandscapeEditLayers = Landscape->GetEditLayersConst();
+	for (const ULandscapeEditLayerBase* EditLayer : LandscapeEditLayers)
 	{
-		WidgetLayerList.Add(MakeShared<FWidgetLayerListItem>(LandscapeLayers[LandscapeLayers.Num() - 1 - LayerIndex], [this](){ OnLayerListUpdated(); }, &WidgetLayerList));
+		WidgetLayerList.Add(MakeShared<FWidgetLayerListItem>(LandscapeEditLayers[LandscapeEditLayers.Num() - 1 - LayerIndex], [this](){ OnLayerListUpdated(); }, &WidgetLayerList));
 		++LayerIndex;
 	}
 
 	WidgetLayerList[0]->bAllowedToDrag = true;
-	InsertedLayerIndex = LandscapeLayers.Num() - 1;
+	InsertedLayerIndex = LandscapeEditLayers.Num() - 1;
 	
 	// Construct list view
 	SAssignNew(LayerListView, SWidgetLayerListView)
@@ -258,7 +258,7 @@ TSharedRef<ITableRow> SLandscapeLayerListDialog::OnGenerateRow(TSharedPtr<FWidge
 void SLandscapeLayerListDialog::OnComplete()
 {
 	check(Landscape.IsValid());
-	TArrayView<const FLandscapeLayer> LandscapeLayers = Landscape->GetLayersConst();
+	const TArray<const ULandscapeEditLayerBase*> LandscapeEditLayers = Landscape->GetEditLayersConst();
 	// Find the draggable layer in WidgetLayerList
 	for (int WidgetLayerListIndex = 0; WidgetLayerListIndex < WidgetLayerList.Num(); ++WidgetLayerListIndex)
 	{
@@ -267,12 +267,12 @@ void SLandscapeLayerListDialog::OnComplete()
 		if (WidgetLayer->bAllowedToDrag)
 		{
 			// Find actual edit layer corresponding to draggable layer
-			for (int32 LayerIndex = 0; LayerIndex < LandscapeLayers.Num(); ++LayerIndex)
+			for (int32 LayerIndex = 0; LayerIndex < LandscapeEditLayers.Num(); ++LayerIndex)
 			{
-				const FLandscapeLayer& Layer = LandscapeLayers[LayerIndex];
+				const ULandscapeEditLayerBase* EditLayer = LandscapeEditLayers[LayerIndex];
 
 				// Once found, remove and copy edit layer into correct spot in LayerList
-				if (WidgetLayer->LayerGuid == Layer.EditLayer->GetGuid())
+				if (WidgetLayer->LayerGuid == EditLayer->GetGuid())
 				{
 					InsertedLayerIndex = WidgetLayerListIndex;
 					Landscape->ReorderLayer(LayerIndex, InsertedLayerIndex);
diff --git a/Engine/Source/Editor/LandscapeEditor/Public/SLandscapeLayerListDialog.h b/Engine/Source/Editor/LandscapeEditor/Public/SLandscapeLayerListDialog.h
--- a/Engine/Source/Editor/LandscapeEditor/Public/SLandscapeLayerListDialog.h
+++ b/Engine/Source/Editor/LandscapeEditor/Public/SLandscapeLayerListDialog.h
@@ -45,7 +45,7 @@ private:
 	/* The list view of WidgetLayerList. */
 	TSharedPtr<SWidgetLayerListView> LayerListView;
 
-	/* Array of FWidgetLayerListItem, which contains FLandscapeLayer* in addition to intermediate data.
+	/* Array of FWidgetLayerListItem, which contains ULandscapeEditLayerBase* in addition to intermediate data.
 	 * Note: This array is a UI reflection of LayerList but the order is the reverse order of LayerList */
 	TArray<TSharedPtr<FWidgetLayerListItem>> WidgetLayerList;
 
diff --git a/Engine/Source/Runtime/Landscape/Classes/Landscape.h b/Engine/Source/Runtime/Landscape/Classes/Landscape.h
--- a/Engine/Source/Runtime/Landscape/Classes/Landscape.h
+++ b/Engine/Source/Runtime/Landscape/Classes/Landscape.h
@@ -426,6 +426,8 @@ public:
 
 	// ULandscapeEditLayerBase accessors : both const and non-const versions are provided because the landscape listens to data changes on the edit layer object 
 	//  and can therefore react to any change
+	LANDSCAPE_API const TArray<const ULandscapeEditLayerBase*> GetEditLayersConst() const;
+	LANDSCAPE_API const TArray<ULandscapeEditLayerBase*> GetEditLayers() const;
 	LANDSCAPE_API const ULandscapeEditLayerBase* GetEditLayerConst(int32 InLayerIndex) const;
 	LANDSCAPE_API ULandscapeEditLayerBase* GetEditLayer(int32 InLayerIndex) const;
 	LANDSCAPE_API const ULandscapeEditLayerBase* GetEditLayerConst(const FGuid& InLayerGuid) const;
@@ -463,6 +465,7 @@ public:
 	 * });
 	 */
 	LANDSCAPE_API void ForEachLayerConst(TFunctionRef<bool(const FLandscapeLayer&)> Fn);
+	LANDSCAPE_API void ForEachEditLayerConst(TFunctionRef<bool(const ULandscapeEditLayerBase*)> Fn);
 
 	LANDSCAPE_API void GetUsedPaintLayers(int32 InLayerIndex, TArray<ULandscapeLayerInfoObject*>& OutUsedLayerInfos) const;
 	LANDSCAPE_API void GetUsedPaintLayers(const FGuid& InLayerGuid, TArray<ULandscapeLayerInfoObject*>& OutUsedLayerInfos) const;
@@ -544,7 +547,8 @@ public:
 
 private:
 	FLandscapeLayer* GetLayerInternal(int32 InLayerIndex);
-	void OnLayerCreatedInternal(FLandscapeLayer& Layer);
+	void OnLayerCreatedInternal(ULandscapeEditLayerBase* EditLayer);
+	ULandscapeEditLayerBase* GetEditLayerInternal(int32 InLayerIndex);
 
 	bool SupportsEditLayersLocalMerge();
 
@@ -593,7 +597,7 @@ private:
 
 	void UpdateLayersMaterialInstances(const TArray<ULandscapeComponent*>& InLandscapeComponents);
 
-	void PrepareComponentDataToExtractMaterialLayersCS(const TArray<ULandscapeComponent*>& InLandscapeComponents, const FLandscapeLayer& InLayer, int32 InCurrentWeightmapToProcessIndex, const FIntPoint& InLandscapeBase, class FLandscapeTexture2DResource* InOutTextureData,
+	void PrepareComponentDataToExtractMaterialLayersCS(const TArray<ULandscapeComponent*>& InLandscapeComponents, const ULandscapeEditLayerBase* InEditLayer, int32 InCurrentWeightmapToProcessIndex, const FIntPoint& InLandscapeBase, class FLandscapeTexture2DResource* InOutTextureData,
 														  TArray<struct FLandscapeLayerWeightmapExtractMaterialLayersComponentData>& OutComponentData, TMap<ULandscapeLayerInfoObject*, int32>& OutLayerInfoObjects);
 	void PrepareComponentDataToPackMaterialLayersCS(int32 InCurrentWeightmapToProcessIndex, const FIntPoint& InLandscapeBase, const TArray<ULandscapeComponent*>& InAllLandscapeComponents, TArray<UTexture2D*>& InOutProcessedWeightmaps,
 													TArray<FLandscapeEditLayerReadback*>& OutProcessedCPUReadBacks, TArray<struct FLandscapeLayerWeightmapPackMaterialLayersComponentData>& OutComponentData);
diff --git a/Engine/Source/Runtime/Landscape/Classes/LandscapeEditLayer.h b/Engine/Source/Runtime/Landscape/Classes/LandscapeEditLayer.h
--- a/Engine/Source/Runtime/Landscape/Classes/LandscapeEditLayer.h
+++ b/Engine/Source/Runtime/Landscape/Classes/LandscapeEditLayer.h
@@ -58,6 +58,14 @@ public:
 			FExecuteParams() = delete;
 
 		public:
+			LANDSCAPE_API FExecuteParams(const ULandscapeEditLayerBase* InEditLayer, ALandscape* InLandscape)
+				: EditLayer(InEditLayer)
+				, Landscape(InLandscape)
+			{
+				check((InEditLayer != nullptr) && (InLandscape != nullptr));
+			}
+
+			UE_DEPRECATED(5.6, "This FLandscapeLayer constructor is deprecated. Use ULandscapeEditLayerBase constructor.")
 			LANDSCAPE_API FExecuteParams(const FLandscapeLayer* InLayer, ALandscape* InLandscape)
 				: Layer(InLayer)
 				, Landscape(InLandscape)
@@ -65,10 +73,13 @@ public:
 				check((InLayer != nullptr) && (InLandscape != nullptr));
 			}
 
+			UE_DEPRECATED(5.6, "Deprecated in preference of ULandscapeEditLayerBase*, use GetEditLayer instead. ")
 			inline const FLandscapeLayer* GetLayer() const { return Layer; }
+			inline const ULandscapeEditLayerBase* GetEditLayer() const { return EditLayer; }
 			inline ALandscape* GetLandscape() const { return Landscape; }
 
 		private:
+			const ULandscapeEditLayerBase* EditLayer = nullptr;
 			const FLandscapeLayer* Layer = nullptr;
 			ALandscape* Landscape = nullptr;
 		};
diff --git a/Engine/Source/Runtime/Landscape/Private/Landscape.cpp b/Engine/Source/Runtime/Landscape/Private/Landscape.cpp
--- a/Engine/Source/Runtime/Landscape/Private/Landscape.cpp
+++ b/Engine/Source/Runtime/Landscape/Private/Landscape.cpp
@@ -1408,11 +1408,11 @@ PRAGMA_ENABLE_DEPRECATION_WARNINGS
 		}
 
 		// Fixup missing/mismatching edit layer names :
-		if (const FLandscapeLayer* LandscapeLayer = GetLandscapeActor() ? GetLandscapeActor()->GetLayerConst(ItPair.Key) : nullptr)
+		if (const ULandscapeEditLayerBase* LandscapeEditLayer = GetLandscapeActor() ? GetLandscapeActor()->GetEditLayerConst(ItPair.Key) : nullptr)
 		{
-			if (LayerComponentData.DebugName != LandscapeLayer->EditLayer->GetName())
+			if (LayerComponentData.DebugName != LandscapeEditLayer->GetName())
 			{
-				LayerComponentData.DebugName = LandscapeLayer->EditLayer->GetName();
+				LayerComponentData.DebugName = LandscapeEditLayer->GetName();
 			}
 		}
 	}
@@ -2566,8 +2566,8 @@ void ULandscapeComponent::AddDefaultLayerData(const FGuid& InLayerGuid, const TA
 
 	if (LayerData == nullptr || !LayerData->IsInitialized())
 	{
-		const FLandscapeLayer* LandscapeLayer = GetLandscapeActor() ? GetLandscapeActor()->GetLayerConst(InLayerGuid) : nullptr;
-		FLandscapeLayerComponentData NewData(LandscapeLayer ? LandscapeLayer->EditLayer->GetName() : FName());
+		const ULandscapeEditLayerBase* LandscapeEditLayer = GetLandscapeActor() ? GetLandscapeActor()->GetEditLayerConst(InLayerGuid) : nullptr;
+		FLandscapeLayerComponentData NewData(LandscapeEditLayer ? LandscapeEditLayer->GetName() : FName());
 
 		// Setup Heightmap data
 		UTexture2D** LayerHeightmap = InOutCreatedHeightmapTextures.Find(ComponentHeightmap);
@@ -3812,9 +3812,9 @@ void ALandscapeProxy::PreSave(FObjectPreSaveContext ObjectSaveContext)
 			// Make sure edit layer debug names are synchronized upon save :
 			LandscapeComponent->ForEachLayer([&](const FGuid& LayerGuid, FLandscapeLayerComponentData& LayerData)
 			{
-				if (const FLandscapeLayer* EditLayer = Landscape->GetLayerConst(LayerGuid))
+				if (const ULandscapeEditLayerBase* EditLayer = Landscape->GetEditLayerConst(LayerGuid))
 				{
-					LayerData.DebugName = EditLayer->EditLayer->GetName();
+					LayerData.DebugName = EditLayer->GetName();
 				}
 			});
 		}
@@ -5012,12 +5012,12 @@ void ALandscapeProxy::FixupSharedData(ALandscape* Landscape, const bool bMapChec
 		bool bUpdated = !SynchronizedProperties.IsEmpty();
 
 		TSet<FGuid> LayerGuids;
-		Algo::Transform(Landscape->GetLayersConst(), LayerGuids, [](const FLandscapeLayer& Layer) { return Layer.EditLayer->GetGuid(); });
+		Algo::Transform(Landscape->GetEditLayersConst(), LayerGuids, [](const ULandscapeEditLayerBase* EditLayer) { return EditLayer->GetGuid(); });
 		bUpdated |= RemoveObsoleteLayers(LayerGuids);
 
-		for (const FLandscapeLayer& Layer : Landscape->GetLayersConst())
+		for (const ULandscapeEditLayerBase* EditLayer : Landscape->GetEditLayersConst())
 		{
-			bUpdated |= AddLayer(Layer.EditLayer->GetGuid());
+			bUpdated |= AddLayer(EditLayer->GetGuid());
 		}
 
 		if (bUpdated)
diff --git a/Engine/Source/Runtime/Landscape/Private/LandscapeBlueprintSupport.cpp b/Engine/Source/Runtime/Landscape/Private/LandscapeBlueprintSupport.cpp
--- a/Engine/Source/Runtime/Landscape/Private/LandscapeBlueprintSupport.cpp
+++ b/Engine/Source/Runtime/Landscape/Private/LandscapeBlueprintSupport.cpp
@@ -24,14 +24,14 @@ void ALandscapeProxy::EditorApplySpline(USplineComponent* InSplineComponent, flo
 	{
 		if (ALandscape* Landscape = GetLandscapeInfo()->LandscapeActor.Get())
 		{
-			const FLandscapeLayer* Layer = Landscape->GetLayerConst(EditLayerName);
-			if (Landscape->HasLayersContent() && (Layer == nullptr))
+			const ULandscapeEditLayerBase* EditLayer = Landscape->GetEditLayerConst(EditLayerName);
+			if (Landscape->HasLayersContent() && (EditLayer == nullptr))
 			{
 				UE_LOG(LogLandscape, Error, TEXT("Invalid landscape edit layer name (\"%s\") for Edit Layers-enabled landscape. Cannot apply spline. "), *EditLayerName.ToString());
 				return;
 			}
 
-			FScopedSetLandscapeEditingLayer Scope(Landscape, Layer ? Layer->EditLayer->GetGuid() : FGuid(), [=] { Landscape->RequestLayersContentUpdate(ELandscapeLayerUpdateMode::Update_All); });
+			FScopedSetLandscapeEditingLayer Scope(Landscape, EditLayer ? EditLayer->GetGuid() : FGuid(), [=] { Landscape->RequestLayersContentUpdate(ELandscapeLayerUpdateMode::Update_All); });
 
 			TArray<FLandscapeSplineInterpPoint> Points;
 			LandscapeSplineRaster::FPointifyFalloffs Falloffs(StartSideFalloff, EndSideFalloff);
diff --git a/Engine/Source/Runtime/Landscape/Private/LandscapeConfigHelper.cpp b/Engine/Source/Runtime/Landscape/Private/LandscapeConfigHelper.cpp
--- a/Engine/Source/Runtime/Landscape/Private/LandscapeConfigHelper.cpp
+++ b/Engine/Source/Runtime/Landscape/Private/LandscapeConfigHelper.cpp
@@ -380,28 +380,28 @@ void FLandscapeConfigHelper::CopyRegionToComponent(ULandscapeInfo* InLandscapeIn
 		}
 		InComponent->InitHeightmapData(HeightInitData, false);
 
-		for (const FLandscapeLayer& LandscapeLayer : InLandscapeInfo->LandscapeActor->GetLayersConst())
+		for (const ULandscapeEditLayerBase* LandscapeEditLayer : InLandscapeInfo->LandscapeActor->GetEditLayersConst())
 		{		
-			check(LandscapeLayer.EditLayer != nullptr);
-			if (!LandscapeLayer.EditLayer->NeedsPersistentTextures())
+			check(LandscapeEditLayer != nullptr);
+			if (!LandscapeEditLayer->NeedsPersistentTextures())
 			{
 				continue;
 			}
 
 			TMap<UTexture2D*, UTexture2D*> CreatedTextures;
-			InComponent->AddDefaultLayerData(LandscapeLayer.EditLayer->GetGuid(), { InComponent }, CreatedTextures);
+			InComponent->AddDefaultLayerData(LandscapeEditLayer->GetGuid(), { InComponent }, CreatedTextures);
 
 			// No need to copy the persistent textures if they are not manually edited (e.g. fully procedural layers with persistent textures : e.g. the splines layer) :
-			if (!LandscapeLayer.EditLayer->SupportsEditingTools())
+			if (!LandscapeEditLayer->SupportsEditingTools())
 			{
 				continue;
 			}
 
 			SrcWeightData.Empty();
-			ExtractLandscapeData(InLandscapeInfo, InRegion, LandscapeLayer.EditLayer->GetGuid(), SrcHeightData, SrcWeightData);
+			ExtractLandscapeData(InLandscapeInfo, InRegion, LandscapeEditLayer->GetGuid(), SrcHeightData, SrcWeightData);
 			CopyData(SrcHeightData, NewHeightData, InRegion, NewRegion, bInResample);
 
-			FScopedSetLandscapeEditingLayer LayerScope(NewLandscapeInfo->LandscapeActor.Get(), LandscapeLayer.EditLayer->GetGuid());
+			FScopedSetLandscapeEditingLayer LayerScope(NewLandscapeInfo->LandscapeActor.Get(), LandscapeEditLayer->GetGuid());
 			{
 				FLandscapeEditDataInterface LandscapeEdit(NewLandscapeInfo);
 				LandscapeEdit.SetHeightData(NewRegion.Min.X, NewRegion.Min.Y, NewRegion.Max.X, NewRegion.Max.Y, NewHeightData.GetData(), 0, false);
diff --git a/Engine/Source/Runtime/Landscape/Private/LandscapeEdit.cpp b/Engine/Source/Runtime/Landscape/Private/LandscapeEdit.cpp
--- a/Engine/Source/Runtime/Landscape/Private/LandscapeEdit.cpp
+++ b/Engine/Source/Runtime/Landscape/Private/LandscapeEdit.cpp
@@ -3902,7 +3902,7 @@ void ALandscapeProxy::Import(const FGuid& InGuid, int32 InMinX, int32 InMinY, in
 		// Create the default layer first
 		ALandscape* LandscapeActor = GetLandscapeActor();
 		check(LandscapeActor != nullptr);
-		if (LandscapeActor->GetLayersConst().IsEmpty() && InImportLayers.IsEmpty())
+		if (LandscapeActor->GetEditLayersConst().IsEmpty() && InImportLayers.IsEmpty())
 		{
 			LandscapeActor->CreateDefaultLayer();
 		}
@@ -3949,12 +3949,12 @@ void ALandscapeProxy::Import(const FGuid& InGuid, int32 InMinX, int32 InMinY, in
 			}
 
 			// And we will fill all the landscape components with the provided final layer content put into the default layer (aka layer index 0)
-			const FLandscapeLayer* DefaultLayer = LandscapeActor->GetLayerConst(0);
-			check(DefaultLayer != nullptr);
+			const ULandscapeEditLayerBase* DefaultEditLayer = LandscapeActor->GetEditLayerConst(0);
+			check(DefaultEditLayer != nullptr);
 
 			FLayerImportSettings ImportSettings;
 			ImportSettings.SourceLayerGuid = FinalLayerGuid;
-			ImportSettings.DestinationLayerGuid = DefaultLayer->EditLayer->GetGuid();
+			ImportSettings.DestinationLayerGuid = DefaultEditLayer->GetGuid();
 			LayerImportSettings.Add(ImportSettings);
 
 			ComponentsToProcess.Append(ToRawPtrTArrayUnsafe(LandscapeComponents));
diff --git a/Engine/Source/Runtime/Landscape/Private/LandscapeEditInterface.cpp b/Engine/Source/Runtime/Landscape/Private/LandscapeEditInterface.cpp
--- a/Engine/Source/Runtime/Landscape/Private/LandscapeEditInterface.cpp
+++ b/Engine/Source/Runtime/Landscape/Private/LandscapeEditInterface.cpp
@@ -2362,9 +2362,9 @@ void FLandscapeEditDataInterface::ReplaceLayer(ULandscapeLayerInfoObject* FromLa
 	{
 		LandscapeInfo->LandscapeActor->Modify(GetShouldDirtyPackage());
 		LandscapeInfo->LandscapeActor->ReplaceLayerSubstractiveBlendStatus(FromLayerInfo, ToLayerInfo, GetShouldDirtyPackage());
-		LandscapeInfo->LandscapeActor->ForEachLayerConst([DoReplace](const FLandscapeLayer& CurrentLayer)
+		LandscapeInfo->LandscapeActor->ForEachEditLayerConst([DoReplace](const ULandscapeEditLayerBase* CurrentLayer)
 		{
-			DoReplace(CurrentLayer.EditLayer->GetGuid());
+			DoReplace(CurrentLayer->GetGuid());
 			return true;
 		});
 		LandscapeInfo->LandscapeActor->RequestLayersContentUpdateForceAll(ELandscapeLayerUpdateMode::Update_Weightmap_All);
diff --git a/Engine/Source/Runtime/Landscape/Private/LandscapeEditLayer.cpp b/Engine/Source/Runtime/Landscape/Private/LandscapeEditLayer.cpp
--- a/Engine/Source/Runtime/Landscape/Private/LandscapeEditLayer.cpp
+++ b/Engine/Source/Runtime/Landscape/Private/LandscapeEditLayer.cpp
@@ -446,11 +446,11 @@ TArray<ULandscapeEditLayerSplines::FEditLayerAction> ULandscapeEditLayerSplines:
 		}),
 		FEditLayerAction::FCanExecuteDelegate::CreateWeakLambda(this, [](const FEditLayerAction::FExecuteParams& InParams, FText& OutReason)
 		{
-			const ULandscapeEditLayerBase* EditLayer = InParams.GetLayer()->EditLayer;
+			const ULandscapeEditLayerBase* EditLayer = InParams.GetEditLayer();
 			check(EditLayer != nullptr);
 			if (EditLayer->IsLocked())
 			{
-				OutReason = FText::Format(LOCTEXT("LandscapeEditLayerSplines_CannotUpdateSplinesOnLockedLayer", "Cannot update splines on layer '{0}' : the layer is currently locked"), FText::FromName(InParams.GetLayer()->EditLayer->GetName()));
+				OutReason = FText::Format(LOCTEXT("LandscapeEditLayerSplines_CannotUpdateSplinesOnLockedLayer", "Cannot update splines on layer '{0}' : the layer is currently locked"), FText::FromName(EditLayer->GetName()));
 				return false;
 			}
 
diff --git a/Engine/Source/Runtime/Landscape/Private/LandscapeEditLayers.cpp b/Engine/Source/Runtime/Landscape/Private/LandscapeEditLayers.cpp
--- a/Engine/Source/Runtime/Landscape/Private/LandscapeEditLayers.cpp
+++ b/Engine/Source/Runtime/Landscape/Private/LandscapeEditLayers.cpp
@@ -2701,8 +2701,8 @@ void ALandscape::CopyOldDataToDefaultLayer(ALandscapeProxy* InProxy)
 
 	InProxy->Modify();
 
-	const FLandscapeLayer* DefaultLayer = GetLayerConst(0);
-	check(DefaultLayer != nullptr);
+	const ULandscapeEditLayerBase* DefaultEditLayer = GetEditLayerConst(0);
+	check(DefaultEditLayer != nullptr);
 
 	struct FWeightmapTextureData
 	{
@@ -2716,7 +2716,7 @@ void ALandscape::CopyOldDataToDefaultLayer(ALandscapeProxy* InProxy)
 
 	for (ULandscapeComponent* Component : InProxy->LandscapeComponents)
 	{
-		FLandscapeLayerComponentData* LayerData = Component->GetLayerData(DefaultLayer->EditLayer->GetGuid());
+		FLandscapeLayerComponentData* LayerData = Component->GetLayerData(DefaultEditLayer->GetGuid());
 
 		if (ensure(LayerData != nullptr && LayerData->IsInitialized()))
 		{
@@ -2762,7 +2762,7 @@ void ALandscape::CopyOldDataToDefaultLayer(ALandscapeProxy* InProxy)
 				{
 					LayerData->WeightmapData.Textures[TextureIndex] = WeightmapTextureData->Texture;
 					LayerData->WeightmapData.TextureUsages[TextureIndex] = WeightmapTextureData->Usage;
-					check(WeightmapTextureData->Usage->LayerGuid == DefaultLayer->EditLayer->GetGuid());
+					check(WeightmapTextureData->Usage->LayerGuid == DefaultEditLayer->GetGuid());
 
 					for (int32 ChannelIndex = 0; ChannelIndex < ULandscapeWeightmapUsage::NumChannels; ++ChannelIndex)
 					{
@@ -2804,7 +2804,7 @@ void ALandscape::CopyOldDataToDefaultLayer(ALandscapeProxy* InProxy)
 						LayerData->WeightmapData.TextureUsages[TextureIndex]->ChannelUsage[ChannelIndex] = ComponentWeightmapTexturesUsage[TextureIndex]->ChannelUsage[ChannelIndex];
 					}
 
-					LayerData->WeightmapData.TextureUsages[TextureIndex]->LayerGuid = DefaultLayer->EditLayer->GetGuid();
+					LayerData->WeightmapData.TextureUsages[TextureIndex]->LayerGuid = DefaultEditLayer->GetGuid();
 
 					// Create new Usage for the "final" layer as the other one will now be used by the Default layer
 					for (const FWeightmapLayerAllocationInfo& Allocation : ComponentLayerAllocations)
@@ -2897,10 +2897,10 @@ void ALandscapeProxy::InitializeProxyLayersWeightmapUsage()
 			// Reinitialize the weightmap usages for the base (final) paint layers allocations :
 			Component->InitializeLayersWeightmapUsage(FGuid());
 
-			for (const FLandscapeLayer& Layer : Landscape->GetLayersConst())
+			for (const ULandscapeEditLayerBase* EditLayer : Landscape->GetEditLayersConst())
 			{
 				// Reinitialize each edit layer's weightmap usages list :
-				Component->InitializeLayersWeightmapUsage(Layer.EditLayer->GetGuid());
+				Component->InitializeLayersWeightmapUsage(EditLayer->GetGuid());
 			}
 		}
 	}
@@ -2915,8 +2915,8 @@ void ULandscapeComponent::InitializeLayersWeightmapUsage(const FGuid& InLayerGui
 	check(Proxy);
 	ALandscape* Landscape = GetLandscapeActor();
 	check(Landscape);
-	const FLandscapeLayer* SplinesLayer = Landscape->FindLayerOfTypeConst(ULandscapeEditLayerSplines::StaticClass());
-	FGuid SplinesEditLayerGuid = SplinesLayer ? SplinesLayer->EditLayer->GetGuid() : FGuid();
+	const ULandscapeEditLayerBase* SplinesEditLayer = Landscape->FindEditLayerOfTypeConst(ULandscapeEditLayerSplines::StaticClass());
+	FGuid SplinesEditLayerGuid = SplinesEditLayer ? SplinesEditLayer->GetGuid() : FGuid();
 
 	// Don't consider invalid edit layers : 
 	if (InLayerGuid.IsValid())
@@ -3062,13 +3062,13 @@ void ALandscapeProxy::ValidateProxyLayersWeightmapUsage() const
 
 			// Validate edit layers weightmap allocations : 
 			{
-				const FLandscapeLayer* SplinesLayer = Landscape->FindLayerOfTypeConst(ULandscapeEditLayerSplines::StaticClass());
-				for (const FLandscapeLayer& Layer : Landscape->GetLayersConst())
+				const ULandscapeEditLayerBase* SplinesEditLayer = Landscape->FindEditLayerOfTypeConst(ULandscapeEditLayerSplines::StaticClass());
+				for (const ULandscapeEditLayerBase* EditLayer : Landscape->GetEditLayersConst())
 				{
-					FLandscapeLayerComponentData* LayerData = Component->GetLayerData(Layer.EditLayer->GetGuid());
+					FLandscapeLayerComponentData* LayerData = Component->GetLayerData(EditLayer->GetGuid());
 
 					// Skip validation on SplinesLayer since it can momentarily contain duplicated layer allocations after undo (since it's updated outside of a transaction) :
-					if (LayerData != nullptr && LayerData->IsInitialized() && (&Layer != SplinesLayer))
+					if (LayerData != nullptr && LayerData->IsInitialized() && (EditLayer != SplinesEditLayer))
 					{
 						for (int32 LayerIdx = 0; LayerIdx < LayerData->WeightmapData.LayerAllocations.Num(); LayerIdx++)
 						{
@@ -3078,7 +3078,7 @@ void ALandscapeProxy::ValidateProxyLayersWeightmapUsage() const
 								UTexture2D* WeightmapTexture = LayerData->WeightmapData.Textures[Allocation.WeightmapTextureIndex];
 								if (ULandscapeWeightmapUsage* Usage = LayerData->WeightmapData.TextureUsages.IsValidIndex(Allocation.WeightmapTextureIndex) ? LayerData->WeightmapData.TextureUsages[Allocation.WeightmapTextureIndex].Get() : nullptr)
 								{
-									ValidateWeightmapAllocationAndUsage(WeightmapTexture, Allocation, Usage, Layer.EditLayer->GetGuid());
+									ValidateWeightmapAllocationAndUsage(WeightmapTexture, Allocation, Usage, EditLayer->GetGuid());
 								}
 							}
 						}
@@ -8562,7 +8562,7 @@ bool ALandscape::ResolveLayersTexture(
 	return bChanged;
 }
 
-void ALandscape::PrepareComponentDataToExtractMaterialLayersCS(const TArray<ULandscapeComponent*>& InLandscapeComponents, const FLandscapeLayer& InLayer, int32 InCurrentWeightmapToProcessIndex, const FIntPoint& InLandscapeBase, FLandscapeTexture2DResource* InOutTextureData,
+void ALandscape::PrepareComponentDataToExtractMaterialLayersCS(const TArray<ULandscapeComponent*>& InLandscapeComponents, const ULandscapeEditLayerBase* InEditLayer, int32 InCurrentWeightmapToProcessIndex, const FIntPoint& InLandscapeBase, FLandscapeTexture2DResource* InOutTextureData,
 	TArray<FLandscapeLayerWeightmapExtractMaterialLayersComponentData>& OutComponentData, TMap<ULandscapeLayerInfoObject*, int32>& OutLayerInfoObjects)
 {
 	ULandscapeInfo* Info = GetLandscapeInfo();
@@ -8578,7 +8578,7 @@ void ALandscape::PrepareComponentDataToExtractMaterialLayersCS(const TArray<ULan
 	const int32 LocalComponentSizeVerts = (SubsectionSizeQuads + 1) * NumSubsections;
 	for (ULandscapeComponent* Component : InLandscapeComponents)
 	{
-		FLandscapeLayerComponentData* ComponentLayerData = Component->GetLayerData(InLayer.EditLayer->GetGuid());
+		FLandscapeLayerComponentData* ComponentLayerData = Component->GetLayerData(InEditLayer->GetGuid());
 
 		if (ComponentLayerData != nullptr)
 		{
@@ -8594,13 +8594,13 @@ void ALandscape::PrepareComponentDataToExtractMaterialLayersCS(const TArray<ULan
 				FVector2D SourcePositionOffset(FMath::RoundToInt((float)ComponentSectionBase.X / LocalComponentSizeQuad), FMath::RoundToInt((float)ComponentSectionBase.Y / LocalComponentSizeQuad));
 				FIntPoint SourceComponentVertexPosition = FIntPoint(static_cast<int32>(SourcePositionOffset.X * LocalComponentSizeVerts), static_cast<int32>(SourcePositionOffset.Y * LocalComponentSizeVerts));
 
-				FLandscapeLayersCopyTextureParams& CopyTextureParams = DeferredCopyTextures.Add_GetRef(FLandscapeLayersCopyTextureParams(*LayerWeightmap->GetName(), LayerWeightmap->GetResource(), FString::Printf(TEXT("%s WeightmapScratchTexture"), *InLayer.EditLayer->GetName().ToString()), InOutTextureData));
+				FLandscapeLayersCopyTextureParams& CopyTextureParams = DeferredCopyTextures.Add_GetRef(FLandscapeLayersCopyTextureParams(*LayerWeightmap->GetName(), LayerWeightmap->GetResource(), FString::Printf(TEXT("%s WeightmapScratchTexture"), *InEditLayer->GetName().ToString()), InOutTextureData));
 				// Only copy the size that's actually needed : 
 				CopyTextureParams.CopySize.X = LayerWeightmap->GetResource()->GetSizeX();
 				CopyTextureParams.CopySize.Y = LayerWeightmap->GetResource()->GetSizeY();
 				// Copy from the top-left corner of the weightmap to the composited texture's position
 				CopyTextureParams.DestPosition = SourceComponentVertexPosition;
-				PrintLayersDebugTextureResource(FString::Printf(TEXT("LS Weight: %s WeightmapScratchTexture %s"), *InLayer.EditLayer->GetName().ToString(), TEXT("WeightmapScratchTextureResource")), InOutTextureData, 0, false);
+				PrintLayersDebugTextureResource(FString::Printf(TEXT("LS Weight: %s WeightmapScratchTexture %s"), *InEditLayer->GetName().ToString(), TEXT("WeightmapScratchTextureResource")), InOutTextureData, 0, false);
 
 				for (const FWeightmapLayerAllocationInfo& WeightmapLayerAllocation : ComponentLayerData->WeightmapData.LayerAllocations)
 				{
@@ -10018,7 +10018,7 @@ int32 ALandscape::PerformLayersWeightmapsGlobalMerge(FUpdateLayersContentContext
 
 				// Prepare compute shader data
 				TArray<FLandscapeLayerWeightmapExtractMaterialLayersComponentData> ComponentsData;
-				PrepareComponentDataToExtractMaterialLayersCS(InUpdateLayersContentContext.LandscapeComponentsWeightmapsToRender, Layer, CurrentWeightmapToProcessIndex, LandscapeExtent.Min, WeightmapScratchExtractLayerTextureResource,
+				PrepareComponentDataToExtractMaterialLayersCS(InUpdateLayersContentContext.LandscapeComponentsWeightmapsToRender, Layer.EditLayer, CurrentWeightmapToProcessIndex, LandscapeExtent.Min, WeightmapScratchExtractLayerTextureResource,
 					ComponentsData, LayerInfoObjects);
 
 				HasFoundWeightmapToProcess = ComponentsData.Num() > 0;
@@ -11348,8 +11348,8 @@ void ALandscape::MonitorLandscapeEdModeChanges()
 		bRequiredEditingClientFullUpdate = true;
 	}
 
-	const FLandscapeLayer* SelectedLayer = LandscapeEdMode ? LandscapeEdMode->GetLandscapeSelectedLayer() : nullptr;
-	FGuid NewSelectedLayer = SelectedLayer && SelectedLayer->EditLayer->IsVisible() ? SelectedLayer->EditLayer->GetGuid() : FGuid();
+	const ULandscapeEditLayerBase* SelectedEditLayer = LandscapeEdMode ? LandscapeEdMode->GetLandscapeSelectedLayer() : nullptr;
+	FGuid NewSelectedLayer = SelectedEditLayer && SelectedEditLayer->IsVisible() ? SelectedEditLayer->GetGuid() : FGuid();
 	if (LandscapeEdModeInfo.SelectedLayer != NewSelectedLayer)
 	{
 		LandscapeEdModeInfo.SelectedLayer = NewSelectedLayer;
@@ -12186,8 +12186,8 @@ bool ALandscapeProxy::AddLayer(const FGuid& InLayerGuid)
 	{
 		if ((Component != nullptr) && !Component->GetLayerData(InLayerGuid))
 		{
-			const FLandscapeLayer* EditLayer = GetLandscapeActor() ? GetLandscapeActor()->GetLayerConst(InLayerGuid) : nullptr;
-			Component->AddLayerData(InLayerGuid, FLandscapeLayerComponentData(EditLayer ? EditLayer->EditLayer->GetName() : FName()));
+			const ULandscapeEditLayerBase* EditLayer = GetLandscapeActor() ? GetLandscapeActor()->GetEditLayerConst(InLayerGuid) : nullptr;
+			Component->AddLayerData(InLayerGuid, FLandscapeLayerComponentData(EditLayer ? EditLayer->GetName() : FName()));
 			bModified = true;
 		}
 	}
@@ -12381,19 +12381,19 @@ void ALandscape::OnEditLayerDataChanged(const FOnLandscapeEditLayerDataChangedPa
 void ALandscape::SetLayerName(int32 InLayerIndex, const FName& InName)
 {
 	const ULandscapeInfo* LandscapeInfo = GetLandscapeInfo();
-	if (FLandscapeLayer* Layer = GetLayerInternal(InLayerIndex); (Layer != nullptr) && (LandscapeInfo != nullptr))
+	if (ULandscapeEditLayerBase* EditLayer = GetEditLayerInternal(InLayerIndex); (EditLayer != nullptr) && (LandscapeInfo != nullptr))
 	{
-		Layer->EditLayer->SetName(InName, /*bInModify = */true);;
+		EditLayer->SetName(InName, /*bInModify = */true);;
 	}
 }
 
 // Deprecated
 float ALandscape::GetLayerAlpha(int32 InLayerIndex, bool bInHeightmap) const
 {
-	if (const FLandscapeLayer* Layer = GetLayerConst(InLayerIndex))
+	if (const ULandscapeEditLayerBase* Layer = GetEditLayerConst(InLayerIndex))
 	{
-		check(Layer->EditLayer != nullptr);
-		return Layer->EditLayer->GetAlphaForTargetType(bInHeightmap ? ELandscapeToolTargetType::Heightmap : ELandscapeToolTargetType::Weightmap);
+		check(Layer != nullptr);
+		return Layer->GetAlphaForTargetType(bInHeightmap ? ELandscapeToolTargetType::Heightmap : ELandscapeToolTargetType::Weightmap);
 	}
 
 	return 1.0f;
@@ -12410,9 +12410,9 @@ float ALandscape::GetClampedLayerAlpha(float InAlpha, bool bInHeightmap) const
 void ALandscape::SetLayerAlpha(int32 InLayerIndex, float InAlpha, bool bInHeightmap)
 {
 	const ULandscapeInfo* LandscapeInfo = GetLandscapeInfo();
-	if (FLandscapeLayer* Layer = GetLayerInternal(InLayerIndex); (Layer != nullptr) && (LandscapeInfo != nullptr))
+	if (ULandscapeEditLayerBase* EditLayer = GetEditLayerInternal(InLayerIndex); (EditLayer != nullptr) && (LandscapeInfo != nullptr))
 	{
-		Layer->EditLayer->SetAlphaForTargetType(bInHeightmap ? ELandscapeToolTargetType::Heightmap : ELandscapeToolTargetType::Weightmap, InAlpha, /*bInModify = */true, EPropertyChangeType::ValueSet);
+		EditLayer->SetAlphaForTargetType(bInHeightmap ? ELandscapeToolTargetType::Heightmap : ELandscapeToolTargetType::Weightmap, InAlpha, /*bInModify = */true, EPropertyChangeType::ValueSet);
 	}
 }
 
@@ -12420,18 +12420,18 @@ void ALandscape::SetLayerAlpha(int32 InLayerIndex, float InAlpha, bool bInHeight
 void ALandscape::SetLayerVisibility(int32 InLayerIndex, bool bInVisible, bool bInForIntermediateRender)
 {
 	const ULandscapeInfo* LandscapeInfo = GetLandscapeInfo();
-	if (FLandscapeLayer* Layer = GetLayerInternal(InLayerIndex); (Layer != nullptr) && (LandscapeInfo != nullptr))
+	if (ULandscapeEditLayerBase* EditLayer = GetEditLayerInternal(InLayerIndex); (EditLayer != nullptr) && (LandscapeInfo != nullptr))
 	{
-		Layer->EditLayer->SetVisible(bInVisible, /*bInModify = */bInForIntermediateRender);
+		EditLayer->SetVisible(bInVisible, /*bInModify = */bInForIntermediateRender);
 	}
 }
 
 // Deprecated 
 void ALandscape::SetLayerLocked(int32 InLayerIndex, bool bLocked)
 {
-	if (FLandscapeLayer* Layer = GetLayerInternal(InLayerIndex); (Layer != nullptr))
+	if (ULandscapeEditLayerBase* EditLayer = GetEditLayerInternal(InLayerIndex); (EditLayer != nullptr))
 	{
-		Layer->EditLayer->SetLocked(bLocked, /*bInModify = */true);
+		EditLayer->SetLocked(bLocked, /*bInModify = */true);
 	}
 }
 
@@ -12456,6 +12456,15 @@ FLandscapeLayer* ALandscape::GetLayerInternal(int32 InLayerIndex)
 	return nullptr;
 }
 
+ULandscapeEditLayerBase* ALandscape::GetEditLayerInternal(int32 InLayerIndex)
+{
+	if (LandscapeEditLayers.IsValidIndex(InLayerIndex))
+	{
+		return LandscapeEditLayers[InLayerIndex].EditLayer;
+	}
+	return nullptr;
+}
+
 // Deprecated
 const FLandscapeLayer* ALandscape::GetLayer(int32 InLayerIndex) const
 {
@@ -12503,6 +12512,34 @@ const FLandscapeLayer* ALandscape::GetLayerConst(const FName& InLayerName) const
 	return LandscapeEditLayers.FindByPredicate([InLayerName](const FLandscapeLayer& Layer) { return Layer.EditLayer->GetName() == InLayerName; });
 }
 
+const TArray<const ULandscapeEditLayerBase*> ALandscape::GetEditLayersConst() const
+{
+	TArray<const ULandscapeEditLayerBase*> EditLayers;
+
+	for (const FLandscapeLayer& Layer : LandscapeEditLayers)
+	{
+		const ULandscapeEditLayerBase* EditLayer = Layer.EditLayer;
+		check(EditLayer != nullptr);
+		EditLayers.Add(EditLayer);
+	}
+
+	return EditLayers;
+}
+
+const TArray<ULandscapeEditLayerBase*> ALandscape::GetEditLayers() const
+{
+	TArray<ULandscapeEditLayerBase*> EditLayers;
+
+	for (const FLandscapeLayer& Layer : GetLayersConst())
+	{
+		ULandscapeEditLayerBase* EditLayer = Layer.EditLayer;
+		check(EditLayer != nullptr);
+		EditLayers.Add(EditLayer);
+	}
+
+	return EditLayers;
+}
+
 const ULandscapeEditLayerBase* ALandscape::GetEditLayerConst(int32 InLayerIndex) const
 {
 	if (const FLandscapeLayer* Layer = GetLayerConst(InLayerIndex))
@@ -12617,6 +12654,17 @@ void ALandscape::ForEachLayerConst(TFunctionRef<bool(const FLandscapeLayer&)> Fn
 	}
 }
 
+void ALandscape::ForEachEditLayerConst(TFunctionRef<bool(const ULandscapeEditLayerBase*)> Fn)
+{
+	for (const ULandscapeEditLayerBase* EditLayer : GetEditLayersConst())
+	{
+		if (!Fn(EditLayer))
+		{
+			return;
+		}
+	}
+}
+
 const FLandscapeLayer* ALandscape::FindLayerOfTypeConst(const TSubclassOf<ULandscapeEditLayerBase>& InLayerClass) const
 {
 	return LandscapeEditLayers.FindByPredicate([&InLayerClass](const FLandscapeLayer& InLayer) { check(InLayer.EditLayer != nullptr); return InLayer.EditLayer->GetClass()->IsChildOf(InLayerClass); });
@@ -12766,10 +12814,9 @@ void ALandscape::CollapseLayer(int32 InLayerIndex)
 
 void ALandscape::GetUsedPaintLayers(int32 InLayerIndex, TArray<ULandscapeLayerInfoObject*>& OutUsedLayerInfos) const
 {
-	const FLandscapeLayer* Layer = GetLayerConst(InLayerIndex);
-	if (Layer)
+	if (const ULandscapeEditLayerBase* EditLayer = GetEditLayer(InLayerIndex))
 	{
-		GetUsedPaintLayers(Layer->EditLayer->GetGuid(), OutUsedLayerInfos);
+		GetUsedPaintLayers(EditLayer->GetGuid(), OutUsedLayerInfos);
 	}
 }
 
@@ -12786,10 +12833,9 @@ void ALandscape::GetUsedPaintLayers(const FGuid& InLayerGuid, TArray<ULandscapeL
 
 void ALandscape::ClearPaintLayer(int32 InLayerIndex, ULandscapeLayerInfoObject* InLayerInfo)
 {
-	const FLandscapeLayer* Layer = GetLayerConst(InLayerIndex);
-	if (Layer)
+	if (const ULandscapeEditLayerBase* EditLayer = GetEditLayer(InLayerIndex))
 	{
-		ClearPaintLayer(Layer->EditLayer->GetGuid(), InLayerInfo);
+		ClearPaintLayer(EditLayer->GetGuid(), InLayerInfo);
 	}
 }
 
@@ -12818,10 +12864,9 @@ void ALandscape::ClearPaintLayer(const FGuid& InLayerGuid, ULandscapeLayerInfoOb
 
 void ALandscape::ClearLayer(int32 InLayerIndex, TSet<TObjectPtr<ULandscapeComponent>>* InComponents, ELandscapeClearMode InClearMode)
 {
-	const FLandscapeLayer* Layer = GetLayerConst(InLayerIndex);
-	if (Layer)
+	if (const ULandscapeEditLayerBase* EditLayer = GetEditLayer(InLayerIndex))
 	{
-		ClearLayer(Layer->EditLayer->GetGuid(), InComponents, InClearMode);
+		ClearLayer(EditLayer->GetGuid(), InComponents, InClearMode);
 	}
 }
 
@@ -12829,15 +12874,15 @@ void ALandscape::ClearLayer(const FGuid& InLayerGuid, TSet<TObjectPtr<ULandscape
 {
 	ensure(HasLayersContent());
 
-	const FLandscapeLayer* Layer = GetLayerConst(InLayerGuid);
+	const ULandscapeEditLayerBase* EditLayer = GetEditLayerConst(InLayerGuid);
 	ULandscapeInfo* LandscapeInfo = GetLandscapeInfo();
-	if (!LandscapeInfo || !Layer)
+	if (!LandscapeInfo || !EditLayer)
 	{
 		return;
 	}
 
 	Modify(bMarkPackageDirty);
-	FScopedSetLandscapeEditingLayer Scope(this, Layer->EditLayer->GetGuid(), [this] { RequestLayersContentUpdate(ELandscapeLayerUpdateMode::Update_All); });
+	FScopedSetLandscapeEditingLayer Scope(this, EditLayer->GetGuid(), [this] { RequestLayersContentUpdate(ELandscapeLayerUpdateMode::Update_All); });
 
 	TArray<uint16> NewHeightData;
 	NewHeightData.AddZeroed(FMath::Square(ComponentSizeQuads + 1));
@@ -12852,7 +12897,7 @@ void ALandscape::ClearLayer(const FGuid& InLayerGuid, TSet<TObjectPtr<ULandscape
 
 	if (InClearMode & ELandscapeClearMode::Clear_Heightmap)
 	{
-		if (Layer->EditLayer->GetBlendMode() == LSBM_AlphaBlend)
+		if (EditLayer->GetBlendMode() == LSBM_AlphaBlend)
 		{
 			NewHeightAlphaBlendData.Init(MAX_uint16, FMath::Square(ComponentSizeQuads + 1));
 			NewHeightFlagsData.AddZeroed(FMath::Square(ComponentSizeQuads + 1));
@@ -12917,32 +12962,32 @@ void ALandscape::ClearLayer(const FGuid& InLayerGuid, TSet<TObjectPtr<ULandscape
 
 void ALandscape::ShowOnlySelectedLayer(int32 InLayerIndex)
 {
-	const FLandscapeLayer* VisibleLayer = GetLayerConst(InLayerIndex);
-	if (VisibleLayer)
+	const ULandscapeEditLayerBase* VisibleEditLayer = GetEditLayerConst(InLayerIndex);
+	if (VisibleEditLayer)
 	{
-		for (FLandscapeLayer& Layer : LandscapeEditLayers)
+		for (ULandscapeEditLayerBase* EditLayer : GetEditLayers())
 		{
-			bool bDesiredVisible = (&Layer == VisibleLayer);
-			if (Layer.EditLayer->IsVisible() != bDesiredVisible)
-				{
-				Layer.EditLayer->SetVisible(bDesiredVisible, /* bInModify = */ true);
+			bool bDesiredVisible = (EditLayer == VisibleEditLayer);
+			if (EditLayer->IsVisible() != bDesiredVisible)
+			{
+				EditLayer->SetVisible(bDesiredVisible, /* bInModify = */ true);
 			}
 		}
-}
+	}
 }
 
 void ALandscape::ShowAllLayers()
 {
 	if (LandscapeEditLayers.Num() > 0)
 	{
-		for (FLandscapeLayer& Layer : LandscapeEditLayers)
+		for (ULandscapeEditLayerBase* EditLayer : GetEditLayers())
 		{
-			if (Layer.EditLayer->IsVisible() != true)
-				{
-				Layer.EditLayer->SetVisible(/*bInVisible = */ true, /*bInModify = */ true);
-		}
+			if (EditLayer->IsVisible() != true)
+			{
+				EditLayer->SetVisible(/*bInVisible = */ true, /*bInModify = */ true);
+			}
 		}
-}
+	}
 }
 
 // Deprecated
@@ -13017,16 +13062,16 @@ void ALandscape::UpdateLandscapeSplines(const FGuid& InTargetLayer, bool bInUpda
 	TRACE_CPUPROFILER_EVENT_SCOPE(LandscapeLayers_UpdateLandscapeSplines);
 	check(CanHaveLayersContent());
 	ULandscapeInfo* LandscapeInfo = GetLandscapeInfo();
-	const FLandscapeLayer* SplinesLayer = FindLayerOfTypeConst(ULandscapeEditLayerSplines::StaticClass());
-	FGuid TargetLayerGuid = (SplinesLayer != nullptr) ? SplinesLayer->EditLayer->GetGuid() : InTargetLayer;
-	const FLandscapeLayer* TargetLayer = GetLayerConst(TargetLayerGuid);
+	const ULandscapeEditLayerBase* SplinesEditLayer = FindEditLayerOfTypeConst(ULandscapeEditLayerSplines::StaticClass());
+	FGuid TargetLayerGuid = (SplinesEditLayer != nullptr) ? SplinesEditLayer->GetGuid() : InTargetLayer;
+	const ULandscapeEditLayerBase* TargetLayer = GetEditLayerConst(TargetLayerGuid);
 	if (LandscapeInfo && TargetLayer)
 	{
 		FScopedSetLandscapeEditingLayer Scope(this, TargetLayerGuid, [this] { this->RequestLayersContentUpdate(ELandscapeLayerUpdateMode::Update_All); });
 		// Temporarily disable material instance updates since it will be done once at the end (requested by RequestLayersContentUpdateForceAll)
 		GDisableUpdateLandscapeMaterialInstances = true;
 		TSet<TObjectPtr<ULandscapeComponent>>* ModifiedComponent = nullptr;
-		if (SplinesLayer != nullptr)
+		if (SplinesEditLayer != nullptr)
 		{
 			// Check that we can modify data
 			if (!LandscapeInfo->AreAllComponentsRegistered())
@@ -13055,7 +13100,7 @@ void ALandscape::UpdateLandscapeSplines(const FGuid& InTargetLayer, bool bInUpda
 
 			// Clear layers without affecting weightmap allocations
 			const bool bMarkPackageDirty = false;
-			ClearLayer(SplinesLayer->EditLayer->GetGuid(), (!bInForceUpdateAllCompoments && LandscapeSplinesAffectedComponents.Num()) ? &LandscapeSplinesAffectedComponents : nullptr, Clear_All, bMarkPackageDirty);
+			ClearLayer(SplinesEditLayer->GetGuid(), (!bInForceUpdateAllCompoments && LandscapeSplinesAffectedComponents.Num()) ? &LandscapeSplinesAffectedComponents : nullptr, Clear_All, bMarkPackageDirty);
 			LandscapeSplinesAffectedComponents.Empty();
 			ModifiedComponent = &LandscapeSplinesAffectedComponents;
 			// For now, in Landscape Layer System Mode with a reserved layer for splines, we always update all the splines since we clear the whole layer first
@@ -13191,7 +13236,7 @@ FLandscapeLayer* ALandscape::DuplicateLayerAndMoveBrushes(const FLandscapeLayer&
 	}
 
 	int32 AddedIndex = LandscapeEditLayers.Add(NewLayer);
-	OnLayerCreatedInternal(NewLayer);
+	OnLayerCreatedInternal(NewLayer.EditLayer);
 
 	return &LandscapeEditLayers[AddedIndex];
 }
@@ -13220,7 +13265,7 @@ int32 ALandscape::CreateLayer(FName InName, const TSubclassOf<ULandscapeEditLaye
 	NewLayer.EditLayer->SetBackPointer(this);
 	NewLayer.EditLayer->SetName(GenerateUniqueLayerName(InName), /*bInModify = */true);
 
-	OnLayerCreatedInternal(NewLayer);
+	OnLayerCreatedInternal(NewLayer.EditLayer);
 	return LayerIndex;
 }
 
@@ -13243,33 +13288,33 @@ int32 ALandscape::CreateLayerFrom(const FLandscapeLayer& InLayer)
 	int32 LayerIndex = LandscapeEditLayers.Add(InLayer);
 	FLandscapeLayer& NewLayer = LandscapeEditLayers[LayerIndex];
 	NewLayer.EditLayer = DuplicateObject<ULandscapeEditLayerBase>(InLayer.EditLayer, this, MakeUniqueObjectName(this, InLayer.EditLayer->GetClass()));
-	OnLayerCreatedInternal(NewLayer);
+	OnLayerCreatedInternal(NewLayer.EditLayer);
 	return LayerIndex;
 }
 
-void ALandscape::OnLayerCreatedInternal(FLandscapeLayer& Layer)
+void ALandscape::OnLayerCreatedInternal(ULandscapeEditLayerBase* EditLayer)
 {
 	// TODO: Might not be necessary eventually, if EditLayer has ability to trigger landscape updates
 	// and has access to its guid.
-	Layer.EditLayer->SetBackPointer(this);
-	Layer.EditLayer->SetFlags(RF_Transactional);
+	EditLayer->SetBackPointer(this);
+	EditLayer->SetFlags(RF_Transactional);
 
 	// Register to data change events on the edit layer so that we can update the landscape accordingly : 
-	check(!Layer.EditLayer->OnLayerDataChanged().IsBoundToObject(this));
-	Layer.EditLayer->OnLayerDataChanged().AddUObject(this, &ALandscape::OnEditLayerDataChanged);
+	check(!EditLayer->OnLayerDataChanged().IsBoundToObject(this));
+	EditLayer->OnLayerDataChanged().AddUObject(this, &ALandscape::OnEditLayerDataChanged);
 
 	// Add to self first instead of from the loop.  Avoids potential load-order problems if this is called at load-time,
 	// before the LandscapeInfo is created.  The other proxies can fix themselves if they load later.
-	AddLayer(Layer.EditLayer->GetGuid());
+	AddLayer(EditLayer->GetGuid());
 
 	// Create associated layer data in each landscape proxy
 	if (ULandscapeInfo* LandscapeInfo = GetLandscapeInfo())
 	{
-		LandscapeInfo->ForEachLandscapeProxy([&Layer, this](ALandscapeProxy* Proxy)
+		LandscapeInfo->ForEachLandscapeProxy([&EditLayer, this](ALandscapeProxy* Proxy)
 		{
 			if (Proxy != this)
 			{
-			Proxy->AddLayer(Layer.EditLayer->GetGuid());
+			Proxy->AddLayer(EditLayer->GetGuid());
 			}
 			return true;
 		});
@@ -13292,9 +13337,9 @@ void ALandscape::AddLayersToProxy(ALandscapeProxy* InProxy)
 	check(InProxy != this);
 	check(InProxy != nullptr);
 
-	ForEachLayerConst([&](const FLandscapeLayer& Layer)
+	ForEachEditLayerConst([&](const ULandscapeEditLayerBase* EditLayer)
 	{
-		InProxy->AddLayer(Layer.EditLayer->GetGuid());
+		InProxy->AddLayer(EditLayer->GetGuid());
 		return true;
 	});
 
@@ -13340,14 +13385,14 @@ FName ALandscape::GenerateUniqueLayerName(FName InName) const
 
 bool ALandscape::IsLayerBlendSubstractive(int32 InLayerIndex, const TWeakObjectPtr<ULandscapeLayerInfoObject>& InLayerInfoObj) const
 {
-	const FLandscapeLayer* Layer = GetLayerConst(InLayerIndex);
+	const ULandscapeEditLayerBase* EditLayer = GetEditLayerConst(InLayerIndex);
 
-	if (Layer == nullptr)
+	if (EditLayer == nullptr)
 	{
 		return false;
 	}
 
-	const bool* AllocationBlend = Layer->EditLayer->GetWeightmapLayerAllocationBlend().Find(InLayerInfoObj.Get());
+	const bool* AllocationBlend = EditLayer->GetWeightmapLayerAllocationBlend().Find(InLayerInfoObj.Get());
 
 	if (AllocationBlend != nullptr)
 	{
@@ -13359,20 +13404,20 @@ bool ALandscape::IsLayerBlendSubstractive(int32 InLayerIndex, const TWeakObjectP
 
 void ALandscape::SetLayerSubstractiveBlendStatus(int32 InLayerIndex, bool InStatus, const TWeakObjectPtr<ULandscapeLayerInfoObject>& InLayerInfoObj)
 {
-	if (FLandscapeLayer* Layer = GetLayerInternal(InLayerIndex))
+	if (ULandscapeEditLayerBase* EditLayer = GetEditLayerInternal(InLayerIndex))
 	{
-		Layer->EditLayer->AddOrUpdateWeightmapAllocationLayerBlend(InLayerInfoObj.Get(), InStatus, /*bInModify =*/ true);
+		EditLayer->AddOrUpdateWeightmapAllocationLayerBlend(InLayerInfoObj.Get(), InStatus, /*bInModify =*/ true);
 	}
 }
 
 void ALandscape::ReplaceLayerSubstractiveBlendStatus(ULandscapeLayerInfoObject* InFromLayerInfo, ULandscapeLayerInfoObject* InToLayerInfo, bool bInShouldDirtyPackage)
 {
-	for (FLandscapeLayer& Layer : LandscapeEditLayers)
+	for (ULandscapeEditLayerBase* EditLayer : GetEditLayers())
 	{
 		bool OutValue;
-		if (Layer.EditLayer->RemoveAndCopyWeightmapAllocationLayerBlend(InFromLayerInfo, OutValue, /*bInModify = */ bInShouldDirtyPackage))
+		if (EditLayer->RemoveAndCopyWeightmapAllocationLayerBlend(InFromLayerInfo, OutValue, /*bInModify = */ bInShouldDirtyPackage))
 		{
-			Layer.EditLayer->AddOrUpdateWeightmapAllocationLayerBlend(InToLayerInfo, OutValue, /*bInModify = */ bInShouldDirtyPackage);
+			EditLayer->AddOrUpdateWeightmapAllocationLayerBlend(InToLayerInfo, OutValue, /*bInModify = */ bInShouldDirtyPackage);
 		}
 	}
 }
diff --git a/Engine/Source/Runtime/Landscape/Private/LandscapeSplineRaster.cpp b/Engine/Source/Runtime/Landscape/Private/LandscapeSplineRaster.cpp
--- a/Engine/Source/Runtime/Landscape/Private/LandscapeSplineRaster.cpp
+++ b/Engine/Source/Runtime/Landscape/Private/LandscapeSplineRaster.cpp
@@ -291,11 +291,11 @@ void RasterizeHeight(int32& MinX, int32& MinY, int32& MaxX, int32& MaxY, FLandsc
 	const ALandscape* Landscape = LandscapeEdit.GetTargetLandscape();
 	check(Landscape != nullptr);
 
-	const FLandscapeLayer* EditingLayer = Landscape->GetLayerConst(Landscape->GetEditingLayer());
-	check((EditingLayer == nullptr) || (EditingLayer->EditLayer != nullptr));
-	const FLandscapeLayer* SplinesLayer = Landscape->FindLayerOfTypeConst(ULandscapeEditLayerSplines::StaticClass());
+	const ULandscapeEditLayerBase* EditingLayer = Landscape->GetEditLayerConst(Landscape->GetEditingLayer());
+	check(EditingLayer != nullptr);
+	const ULandscapeEditLayerBase* SplinesLayer = Landscape->FindEditLayerOfTypeConst(ULandscapeEditLayerSplines::StaticClass());
 	bool bIsEditingLayerReservedForSplines = (EditingLayer != nullptr) && (SplinesLayer == EditingLayer);
-	check(!bIsEditingLayerReservedForSplines || (SplinesLayer->EditLayer->GetBlendMode() == LSBM_AlphaBlend));
+	check(!bIsEditingLayerReservedForSplines || (SplinesLayer->GetBlendMode() == LSBM_AlphaBlend));
 
 	TArray<uint16> HeightAlphaBlendData;
 	TArray<uint8> HeightFlagsData;
@@ -577,8 +577,9 @@ bool ULandscapeInfo::ApplySplines(bool bOnlySelected, TSet<TObjectPtr<ULandscape
 	bool bResult = false;
 
 	ALandscape* Landscape = LandscapeActor.Get();
-	const FLandscapeLayer* Layer = Landscape ? Landscape->FindLayerOfTypeConst(ULandscapeEditLayerSplines::StaticClass()) : nullptr;
-	FGuid SplinesTargetLayerGuid = Layer ? Layer->EditLayer->GetGuid() : Landscape ? Landscape->GetEditingLayer() : FGuid();
+
+	const ULandscapeEditLayerBase* Layer = Landscape ? Landscape->FindEditLayerOfTypeConst(ULandscapeEditLayerSplines::StaticClass()) : nullptr;
+	FGuid SplinesTargetLayerGuid = Layer ? Layer->GetGuid() : Landscape ? Landscape->GetEditingLayer() : FGuid();
 	FScopedSetLandscapeEditingLayer Scope(Landscape, SplinesTargetLayerGuid, [=] { Landscape->RequestLayersContentUpdate(ELandscapeLayerUpdateMode::Update_All); });
 
 	TMap<ULandscapeLayerInfoObject*, TSharedPtr<FModulateAlpha>> ModulatePerLayerInfo;
diff --git a/Engine/Source/Runtime/Landscape/Public/LandscapeEdit.h b/Engine/Source/Runtime/Landscape/Public/LandscapeEdit.h
--- a/Engine/Source/Runtime/Landscape/Public/LandscapeEdit.h
+++ b/Engine/Source/Runtime/Landscape/Public/LandscapeEdit.h
@@ -41,7 +41,7 @@ class ILandscapeEdModeInterface
 public:
 	virtual void PostUpdateLayerContent() = 0;
 	virtual ELandscapeToolTargetType GetLandscapeToolTargetType() const = 0;
-	virtual const FLandscapeLayer* GetLandscapeSelectedLayer() const = 0;
+	virtual const ULandscapeEditLayerBase* GetLandscapeSelectedLayer() const = 0;
 	virtual ULandscapeLayerInfoObject* GetSelectedLandscapeLayerInfo() const = 0;
 	virtual void OnCanHaveLayersContentChanged() = 0;
 };
