Fix for invalid landscape LOD being returned from ComputeRequiredLandscapeLOD when using the automatic HLOD source mesh policy.

diff --git a/Engine/Source/Runtime/Landscape/Private/LandscapeHLODBuilder.cpp b/Engine/Source/Runtime/Landscape/Private/LandscapeHLODBuilder.cpp
--- a/Engine/Source/Runtime/Landscape/Private/LandscapeHLODBuilder.cpp
+++ b/Engine/Source/Runtime/Landscape/Private/LandscapeHLODBuilder.cpp
@@ -225,14 +225,12 @@ static int32 ComputeRequiredLandscapeLOD(const ALandscapeProxy* InLandscapeProxy
 			const float ScreenHeight = 1080.0f;
 			const FPerspectiveMatrix ProjMatrix(HalfFOV, ScreenWidth, ScreenHeight, 1.0f);
 
-			
-
 			const ULandscapeComponent* LSComponent = InLandscapeProxy->LandscapeComponents[0];
 			const float ComponentRadiusScaled = static_cast<float>(LSComponent->GetLocalBounds().SphereRadius * LSComponent->GetComponentTransform().GetScale3D().GetAbsMax());
 			const float ExpectedScreenSize = ComputeBoundsScreenSize(FVector::ZeroVector, ComponentRadiusScaled, FVector(0.0f, 0.0f, InViewDistance), ProjMatrix);
 
-			
-			for (RequiredLOD = 0; RequiredLOD < LODScreenSizes.Num(); ++RequiredLOD)
+			// Find the matching LOD for the screen size. No need to test the last LOD screen size if we get to it.
+			for (RequiredLOD = 0; RequiredLOD < LODScreenSizes.Num() - 1; ++RequiredLOD)
 			{
 				if (ExpectedScreenSize > LODScreenSizes[RequiredLOD])
 				{
