Added ptr validation in landscapegrass to avoid crash

diff --git a/Engine/Source/Runtime/Landscape/Private/LandscapeGrass.cpp b/Engine/Source/Runtime/Landscape/Private/LandscapeGrass.cpp
--- a/Engine/Source/Runtime/Landscape/Private/LandscapeGrass.cpp
+++ b/Engine/Source/Runtime/Landscape/Private/LandscapeGrass.cpp
@@ -1026,6 +1026,11 @@ void ALandscapeProxy::UpdateGrassTypeSummary()
 	double ProxyMaxGrassCullDistance = 0.0;
 	for (ULandscapeComponent* Component : LandscapeComponents)
 	{
+		if (Component == nullptr)
+		{
+			continue;
+		}
+
 		double ComponentMaxGrassCullDistance;
 		int32 ComponentGrassVarietyCount = 0;
 		if (Component->GrassTypeSummary.bInvalid)
