Added option for landscape BP brushes to tell if they require power-of-2 render targets to operate. For backwards compatibility reasons, it's set to true by default, although it's not the best option. This ensures that the water brush, in particular, is perfectly aligned when running with batched merge vs. global merge

diff --git a/Engine/Source/Runtime/Landscape/Private/LandscapeBlueprintBrushBase.cpp b/Engine/Source/Runtime/Landscape/Private/LandscapeBlueprintBrushBase.cpp
--- a/Engine/Source/Runtime/Landscape/Private/LandscapeBlueprintBrushBase.cpp
+++ b/Engine/Source/Runtime/Landscape/Private/LandscapeBlueprintBrushBase.cpp
@@ -275,6 +275,21 @@ void ALandscapeBlueprintBrushBase::SetCanAffectVisibilityLayer(bool bInCanAffect
 #endif // WITH_EDITORONLY_DATA
 }
 
+void ALandscapeBlueprintBrushBase::SetUsePowerOfTwoRenderTarget(bool bInUsePowerOfTwoRenderTarget)
+{
+#if WITH_EDITORONLY_DATA
+	if (bInUsePowerOfTwoRenderTarget != bUsePowerOfTwoRenderTarget)
+	{
+		Modify();
+		bUsePowerOfTwoRenderTarget = bInUsePowerOfTwoRenderTarget;
+		if (OwningLandscape)
+		{
+			OwningLandscape->OnBlueprintBrushChanged();
+		}
+	}
+#endif // WITH_EDITORONLY_DATA
+}
+
 #if WITH_EDITOR
 
 void ALandscapeBlueprintBrushBase::GetRendererStateInfo(const UE::Landscape::EditLayers::FMergeContext* InMergeContext,
@@ -380,7 +395,14 @@ bool ALandscapeBlueprintBrushBase::RenderLayer(UE::Landscape::EditLayers::FRende
 	//  render target, so if, from one call to another, this size is changing because there happens to be a scratch render target big enough, we will have to call Initialize again 
 	//  on the brush, which might have performance implications. So using a size that matches what we've already used in the past here allows such brushes not to uselessly go 
 	//  through the initialization process again 
-	const FIntPoint RenderTargetIdealSize = RenderParams.RenderAreaSectionRect.Size().ComponentMax(CurrentRenderAreaSize);
+	FIntPoint RenderTargetIdealSize = RenderParams.RenderAreaSectionRect.Size().ComponentMax(CurrentRenderAreaSize);
+	if (bUsePowerOfTwoRenderTarget)
+	{
+		// For backwards-compatibility with how it used to be, round up the RT size to the next power of 2 : 
+		RenderTargetIdealSize.X = FMath::RoundUpToPowerOfTwo(RenderTargetIdealSize.X);
+		RenderTargetIdealSize.Y = FMath::RoundUpToPowerOfTwo(RenderTargetIdealSize.Y);
+	}
+
 	// Use exact dimensions for the scratch texture here for the reason explained above (see RenderTargetIdealSize) :
 	FScratchRenderTargetParams ScratchRenderTargetParams(TEXT("BPBrushScratchRT"), /*bInExactDimensions = */true, /*bInUseUAV = */false, /*bInTargetArraySlicesIndependently = */false,
 		RenderTargetIdealSize, /*InNumSlices = */0, CurrentLayerReadRT->GetFormat(), CurrentLayerReadRT->GetClearColor(), ERHIAccess::CopyDest);
diff --git a/Engine/Source/Runtime/Landscape/Public/LandscapeBlueprintBrushBase.h b/Engine/Source/Runtime/Landscape/Public/LandscapeBlueprintBrushBase.h
--- a/Engine/Source/Runtime/Landscape/Public/LandscapeBlueprintBrushBase.h
+++ b/Engine/Source/Runtime/Landscape/Public/LandscapeBlueprintBrushBase.h
@@ -62,6 +62,13 @@ protected:
 	UPROPERTY(Category = "Settings", EditAnywhere, BlueprintReadWrite)
 	TArray<FName> AffectedWeightmapLayers;
 
+	/** 
+	 * It is possible to define whether this BP brush works exclusively with power-of-2 render targets. This used to be mandatory so for full backwards-compatibility, 
+	 * this is still the case by default but it can be turned off. 
+	 */
+	UPROPERTY(Category = "Settings", EditAnywhere, BlueprintReadWrite, Setter="SetCanAffectVisibilityLayer")
+	bool bUsePowerOfTwoRenderTarget = true;
+
 #if WITH_EDITORONLY_DATA
 	UPROPERTY(Transient)
 	TObjectPtr<class ALandscape> OwningLandscape;
@@ -97,6 +104,7 @@ public:
 	LANDSCAPE_API void SetCanAffectHeightmap(bool bInCanAffectHeightmap);
 	LANDSCAPE_API void SetCanAffectWeightmap(bool bInCanAffectWeightmap);
 	LANDSCAPE_API void SetCanAffectVisibilityLayer(bool bInCanAffectVisibilityLayer);
+	LANDSCAPE_API void SetUsePowerOfTwoRenderTarget(bool bInUsePowerOfTwoRenderTarget);
 
 #if WITH_EDITOR
 	//~ Begin ILandscapeEditLayerRenderer implementation
