Landscape Spline - setting CVarLandscapeSplineEnableFixedSpriteIcons to false, enabling dynamic spline icons by default

diff --git a/Engine/Source/Runtime/Landscape/Private/LandscapeSplines.cpp b/Engine/Source/Runtime/Landscape/Private/LandscapeSplines.cpp
--- a/Engine/Source/Runtime/Landscape/Private/LandscapeSplines.cpp
+++ b/Engine/Source/Runtime/Landscape/Private/LandscapeSplines.cpp
@@ -76,7 +76,7 @@ static FAutoConsoleVariableRef CVarLandscapeSplineToSplineComponentMaxIterations
 #if WITH_EDITOR
 static FAutoConsoleVariable CVarLandscapeSplineEnableFixedSpriteIcons(
 	TEXT("landscape.Splines.EnableFixedSpriteIcons"),
-	true,
+	false,
 	TEXT("Enables fixed scaling and Z-offset for spline icons. If disabled, configure dynamic scale and offset in project settings."));
 #endif // WITH_EDITOR 
 
diff --git a/Engine/Source/Runtime/Landscape/Public/LandscapeSettings.h b/Engine/Source/Runtime/Landscape/Public/LandscapeSettings.h
--- a/Engine/Source/Runtime/Landscape/Public/LandscapeSettings.h
+++ b/Engine/Source/Runtime/Landscape/Public/LandscapeSettings.h
@@ -128,10 +128,10 @@ protected:
  
 	/** Offset in Z for the landscape spline icon in world-space. */
 	UPROPERTY(EditAnywhere, config, Category = "Spline", meta = (ClampMin = "0", UIMax = "4096"))
-	float SplineIconWorldZOffset = 75.0f;
+	float SplineIconWorldZOffset = 100.0f;
  
 	UPROPERTY(EditAnywhere, config, Category = "Spline", meta = (ClampMin = "32", UIMax = "2048"))
-	float SplineIconScale = 50.0f;
+	float SplineIconScale = 125.0f;
  
 	UPROPERTY(EditAnywhere, config, Category = "Configuration", meta = (ToolTip = 
 		"When true, each target layer will have a representative thumbnail in landscape mode. However, the thumbnail system can be slow. Plus, certain landscape materials (e.g. using RVT) are not able to properly render thumbnails. Thus, setting this to false will skip needlessly rendering landscape layer thumbnails, which can improve the editing experience"))
