Fix for landscape heightmap deep valley artifacts.

diff --git a/Engine/Source/Runtime/Landscape/Classes/LandscapeEdgeFixup.h b/Engine/Source/Runtime/Landscape/Classes/LandscapeEdgeFixup.h
--- a/Engine/Source/Runtime/Landscape/Classes/LandscapeEdgeFixup.h
+++ b/Engine/Source/Runtime/Landscape/Classes/LandscapeEdgeFixup.h
@@ -233,6 +233,7 @@ private:
 
 	// per-group settings (apply to the active group/component)
 	bool bMapped = false;
+	bool bForceUpdateSnapshot = true;					// set to true initially so that we do a force update on the very first request
 	FIntPoint GroupCoord = FIntPoint::ZeroValue;		// coordinate of this heightmap in the active group (when bMapped)
 
 	// components that also want to use this heightmap & edge fixup, but were disabled as we can only support one active component at a time.
@@ -276,7 +277,7 @@ public:
 	void RequestEdgeSnapshotUpdateFromHeightmapSource(bool bUpdateGPUEdgeHashes = false);
 
 	// Update the edge snapshot from heightmap source.  Returns the set of edges that changed since the previous snapshot.
-	UE::Landscape::EEdgeFlags UpdateEdgeSnapshotFromHeightmapSource(const FVector& LandscapeGridScale);
+	UE::Landscape::EEdgeFlags UpdateEdgeSnapshotFromHeightmapSource(const FVector& LandscapeGridScale, bool bForceUpdate = false);
 #endif // WITH_EDITOR
 
 	// Patch the GPU texture edges if needed, using the current snapshot and corresponding neighbor snapshots as source data.
diff --git a/Engine/Source/Runtime/Landscape/Private/Landscape.cpp b/Engine/Source/Runtime/Landscape/Private/Landscape.cpp
--- a/Engine/Source/Runtime/Landscape/Private/Landscape.cpp
+++ b/Engine/Source/Runtime/Landscape/Private/Landscape.cpp
@@ -2138,7 +2138,8 @@ ULandscapeHeightmapTextureEdgeFixup* ULandscapeComponent::InstallOrUpdateTexture
 #if WITH_EDITOR
 		if (bUpdateSnapshotNow)
 		{
-			EdgeFixup->UpdateEdgeSnapshotFromHeightmapSource(LandscapeGridScale);
+			const bool bForceUpdate = true;
+			EdgeFixup->UpdateEdgeSnapshotFromHeightmapSource(LandscapeGridScale, bForceUpdate);
 		}
 #endif // WITH_EDITOR
 	}
diff --git a/Engine/Source/Runtime/Landscape/Private/LandscapeEdgeFixup.cpp b/Engine/Source/Runtime/Landscape/Private/LandscapeEdgeFixup.cpp
--- a/Engine/Source/Runtime/Landscape/Private/LandscapeEdgeFixup.cpp
+++ b/Engine/Source/Runtime/Landscape/Private/LandscapeEdgeFixup.cpp
@@ -1372,7 +1372,7 @@ void ULandscapeHeightmapTextureEdgeFixup::RequestEdgeSnapshotUpdateFromHeightmap
 	ActiveGroup->HeightmapsNeedingEdgeSnapshotCapture.Add(this);
 }
 
-UE::Landscape::ENeighborFlags ULandscapeHeightmapTextureEdgeFixup::UpdateEdgeSnapshotFromHeightmapSource(const FVector& LandscapeGridScale)
+UE::Landscape::ENeighborFlags ULandscapeHeightmapTextureEdgeFixup::UpdateEdgeSnapshotFromHeightmapSource(const FVector& LandscapeGridScale, bool bForceUpdate)
 {
 	TRACE_CPUPROFILER_EVENT_SCOPE(UpdateEdgeSnapshotFromHeightmapSource);
 
@@ -1380,9 +1380,15 @@ UE::Landscape::ENeighborFlags ULandscapeHeightmapTextureEdgeFixup::UpdateEdgeSna
 
 	ENeighborFlags ChangedEdges = ENeighborFlags::None;
 
+	if (bForceUpdateSnapshot)
+	{
+		bForceUpdate = true;
+		bForceUpdateSnapshot = false;
+	}
+
 	// no need to update if source ID matches (we assume texture has not actually changed)		
 	FGuid HeightmapSourceID = ULandscapeTextureHash::GetHash(HeightmapTexture);
-	if (HeightmapSourceID != EdgeSnapshot->TextureSourceID)
+	if (bForceUpdate || (HeightmapSourceID != EdgeSnapshot->TextureSourceID))
 	{
 		FIXUP_DEBUG_LOG_DETAIL(TEXT("  UpdateEdgeSnapshotFromHeightmapSource %p HeightMap: %p -- UPDATED (ID: %s --> %s)"), ActiveComponent.Get(), HeightmapTexture.Get(),
 			*EdgeSnapshot->TextureSourceID.ToString(), *HeightmapSourceID.ToString());
@@ -1413,7 +1419,7 @@ UE::Landscape::ENeighborFlags ULandscapeHeightmapTextureEdgeFixup::UpdateEdgeSna
 	bUpdateGPUEdgeHashes = false;
 	bDoNotPatchUntilGPUEdgeHashesUpdated = false;
 
-	return ENeighborFlags::None;
+	return ChangedEdges;
 }
 #endif // WITH_EDITOR
 
