Avoid crash if saving when landscape TextureStreamingManager is not available.  Flush code isn't needed in that circumstance anyway, since pending async work wouldn't have been possible.

diff --git a/Engine/Source/Runtime/Landscape/Classes/Landscape.h b/Engine/Source/Runtime/Landscape/Classes/Landscape.h
--- a/Engine/Source/Runtime/Landscape/Classes/Landscape.h
+++ b/Engine/Source/Runtime/Landscape/Classes/Landscape.h
@@ -558,6 +558,7 @@ private:
 	bool CreateLayersRenderingResource();
 	void PrepareEditLayersLocalMergeResources();
 	void UpdateLayersContent(bool bInWaitForStreaming = false, bool bInSkipMonitorLandscapeEdModeChanges = false, bool bIntermediateRender = false, bool bFlushRender = false);
+	bool CanUpdateLayersContent() const;
 	void MonitorShaderCompilation();
 	void MonitorLandscapeEdModeChanges();
 	
diff --git a/Engine/Source/Runtime/Landscape/Private/LandscapeEditLayers.cpp b/Engine/Source/Runtime/Landscape/Private/LandscapeEditLayers.cpp
--- a/Engine/Source/Runtime/Landscape/Private/LandscapeEditLayers.cpp
+++ b/Engine/Source/Runtime/Landscape/Private/LandscapeEditLayers.cpp
@@ -11467,6 +11467,17 @@ void ALandscape::FWaitingForResourcesNotificationHelper::Reset()
 	WaitingForResourcesStartTime = -1.0;
 }
 
+bool ALandscape::CanUpdateLayersContent() const
+{
+	ULandscapeInfo* LandscapeInfo = GetLandscapeInfo();
+	UWorld* World = GetWorld();
+	ULandscapeSubsystem* Subsystem = World ? GetWorld()->GetSubsystem<ULandscapeSubsystem>() : nullptr;
+
+	return FApp::CanEverRender() && (LandscapeInfo != nullptr) && CanHaveLayersContent() &&
+		LandscapeInfo->AreAllComponentsRegistered() && LandscapeInfo->SupportsLandscapeEditing() &&
+		Subsystem && Subsystem->GetTextureStreamingManager();
+}
+
 void ALandscape::UpdateLayersContent(bool bInWaitForStreaming, bool bInSkipMonitorLandscapeEdModeChanges, bool bIntermediateRender, bool bFlushRender)
 {
 	TRACE_CPUPROFILER_EVENT_SCOPE(LandscapeLayers_UpdateLayersContent);
@@ -11507,7 +11518,7 @@ void ALandscape::UpdateLayersContent(bool bInWaitForStreaming, bool bInSkipMonit
 	bool bResourcesReady = PrepareTextureResources(bInWaitForStreaming);
 
 	ULandscapeInfo* LandscapeInfo = GetLandscapeInfo();
-	if ((LandscapeInfo == nullptr) || !CanHaveLayersContent() || !LandscapeInfo->AreAllComponentsRegistered() || !FApp::CanEverRender() || !LandscapeInfo->SupportsLandscapeEditing())
+	if (!CanUpdateLayersContent())
 	{
 		return;
 	}
@@ -11905,8 +11916,9 @@ void ALandscape::OnPreSave()
 void ALandscape::FlushLayerContentThisFrame()
 {
 	// Only call ForceUpdateLayersContent once per frame. Each proxy might trigger this, so only do it for the first one.
+	// Don't attempt this if the ULandscapeSubsystem or TextureStreamingManager don't exist.
 	uint32 CurrentFrame = GFrameNumber;
-	if (LastFlushedLayerUpdateFrame != CurrentFrame)
+	if (LastFlushedLayerUpdateFrame != CurrentFrame && CanUpdateLayersContent())
 	{
 		LastFlushedLayerUpdateFrame = CurrentFrame;
 		ForceUpdateLayersContent(/*bIntermediateRender = */ false);
