Fixed OOBB->AABB computation in landscape batched merge, which was leading to bad blend groups for landscape patches

diff --git a/Engine/Source/Runtime/Landscape/Private/LandscapeEditLayerMergeRenderContext.cpp b/Engine/Source/Runtime/Landscape/Private/LandscapeEditLayerMergeRenderContext.cpp
--- a/Engine/Source/Runtime/Landscape/Private/LandscapeEditLayerMergeRenderContext.cpp
+++ b/Engine/Source/Runtime/Landscape/Private/LandscapeEditLayerMergeRenderContext.cpp
@@ -1111,7 +1111,13 @@ bool FMergeRenderContext::IsVisualLogEnabled() const
 
 FBox FOOBox2D::BuildAABB() const
 {
-	return FBox::BuildAABB(Transform.GetTranslation(), Transform.TransformVector(FVector(Extents, 0.0)));
+	return FBox(
+		{
+			Transform.TransformPosition(FVector(+Extents.X, +Extents.Y, 0.0)),
+			Transform.TransformPosition(FVector(+Extents.X, -Extents.Y, 0.0)),
+			Transform.TransformPosition(FVector(-Extents.X, +Extents.Y, 0.0)),
+			Transform.TransformPosition(FVector(-Extents.X, -Extents.Y, 0.0)),
+		});
 }
 
 
