Use a shared material update context to update landscape material instances on undo. That can greatly cut down the undo time when undoing large changes by preventing iterating through all primitive components for each updated material instance

diff --git a/Engine/Source/Runtime/Landscape/Private/LandscapeEdit.cpp b/Engine/Source/Runtime/Landscape/Private/LandscapeEdit.cpp
--- a/Engine/Source/Runtime/Landscape/Private/LandscapeEdit.cpp
+++ b/Engine/Source/Runtime/Landscape/Private/LandscapeEdit.cpp
@@ -1047,18 +1047,26 @@ void ULandscapeComponent::PostEditUndo()
 		}
 
 		// Update Material Instances on each modified component, in case the material property has changed (uses the weightmap allocations fixed by FixupWeightmaps)
-		for (ULandscapeComponent* Component : UndoRedoModifiedComponents)
 		{
-			if (IsValid(Component)) // Components can be pending kill, if they were removed by the undo operation
+			// We're not having the material update context recreate the render state because we will manually do it for only these components
+			TArray<FComponentRecreateRenderStateContext> RecreateRenderStateContexts;
 			{
-				ALandscapeProxy* ComponentProxy = Component->GetLandscapeProxy();
-				if (IsValid(ComponentProxy))
+				FMaterialUpdateContext MaterialUpdateContext(FMaterialUpdateContext::EOptions::Default & ~FMaterialUpdateContext::EOptions::RecreateRenderStates);
+				for (ULandscapeComponent* Component : UndoRedoModifiedComponents)
 				{
-					Component->UpdateMaterialInstances();
+					if (IsValid(Component)) // Components can be pending kill, if they were removed by the undo operation
+					{
+						ALandscapeProxy* ComponentProxy = Component->GetLandscapeProxy();
+						if (IsValid(ComponentProxy))
+						{
+							Component->UpdateMaterialInstances(MaterialUpdateContext, RecreateRenderStateContexts);
+						}
+					}
 				}
 			}
 		}
 
+
 		// Clear out the temporarily recorded undo info
 		UndoRedoModifiedComponentCount = 0;
 		UndoRedoModifiedComponents.Empty();
