Landscape : added a DeleteActors helper in order to let the editor's delete actors dialog run when possible but also to let actors be safely deleted without using the UI in other cases

diff --git a/Engine/Plugins/Experimental/LandscapePatch/Source/LandscapePatch/Private/LandscapePatchManager.cpp b/Engine/Plugins/Experimental/LandscapePatch/Source/LandscapePatch/Private/LandscapePatchManager.cpp
--- a/Engine/Plugins/Experimental/LandscapePatch/Source/LandscapePatch/Private/LandscapePatchManager.cpp
+++ b/Engine/Plugins/Experimental/LandscapePatch/Source/LandscapePatch/Private/LandscapePatchManager.cpp
@@ -571,17 +571,12 @@ void ALandscapePatchManager::MigrateToPrioritySystemAndDeleteInternal(bool bAllo
 		// MarkModifiedLandscapesAsDirty, the same place that the deferred dirty state from MarkObjectDirty is finally applied.  This
 		// will leave the scene in a consistent state before and after the user clicks the "Update" button from the landscape check
 		// code.  If the scene is closed without clicking update, it will remain fully un-migrated on disk.
-		LandscapeInfo->DeleteActorWhenApplyingModifiedStatus(this);
+		LandscapeInfo->DeleteActorWhenApplyingModifiedStatus(this, bAllowUI);
 	}
-	else if (UTypedElementSelectionSet* SelectionSet = GEditor->GetEditorSubsystem<ULevelEditorSubsystem>()->GetSelectionSet())
+	else 
 	{
 		// No LandscapeInfo, likely because no OwningLandscape.  Try to delete directly.
-		GUnrealEd->DeleteActors({ this }, GetWorld(), SelectionSet);
-	}
-	else
-	{
-		// No SelectionSet because Editor isn't ready yet.  Print a message and leave the disconnected patch manager in place.
-		UE_LOG(LogLandscapePatch, Display, TEXT("ALandscapePatchManager \"%s\" is unused and can be deleted."), *GetActorLabel(false));
+		UE::Landscape::DeleteActors({ this }, GetWorld(), bAllowUI);
 	}
 }
 
diff --git a/Engine/Source/Runtime/Landscape/Classes/LandscapeInfo.h b/Engine/Source/Runtime/Landscape/Classes/LandscapeInfo.h
--- a/Engine/Source/Runtime/Landscape/Classes/LandscapeInfo.h
+++ b/Engine/Source/Runtime/Landscape/Classes/LandscapeInfo.h
@@ -295,8 +295,9 @@ public:
 	/**
 	 * Enqueue an actor to be deleted when deferred dirty status from MarkObjectDirty is finally applied.
 	 * @param InActor Actor to delete.
+	 * @param bInAllowUI allows the standard delete actors UX to be displayed, allowing the user to remove lingering reference to these actors, etc.
 	 */
-	LANDSCAPE_API void DeleteActorWhenApplyingModifiedStatus(AActor* InActor);
+	LANDSCAPE_API void DeleteActorWhenApplyingModifiedStatus(AActor* InActor, bool bInAllowUI);
 
 	/** Helper function to invalidate the runtime virtual texture data associated with the specified landscape region. 
 	 *   This forces the dirtied RVT pages to be updated in priority
diff --git a/Engine/Source/Runtime/Landscape/Private/Landscape.cpp b/Engine/Source/Runtime/Landscape/Private/Landscape.cpp
--- a/Engine/Source/Runtime/Landscape/Private/Landscape.cpp
+++ b/Engine/Source/Runtime/Landscape/Private/Landscape.cpp
@@ -5316,14 +5316,14 @@ bool ULandscapeInfo::ModifyObject(UObject* InObject, bool bAlwaysMarkDirty)
 	return bWasAddedToModifiedPackages;
 }
 
-void ULandscapeInfo::DeleteActorWhenApplyingModifiedStatus(AActor* InActor)
+void ULandscapeInfo::DeleteActorWhenApplyingModifiedStatus(AActor* InActor, bool bInAllowUI)
 {
 	check(InActor->GetWorld() && InActor->GetWorld()->WorldType == EWorldType::Editor);
 
 	// If we can mark the package dirty, then we can also delete it right away.  If we can't, then enqueue it for deletion at the same time as other deferred package dirtying.
 	if (InActor->MarkPackageDirty())
 	{
-		GUnrealEd->DeleteActors({ InActor }, InActor->GetWorld(), GUnrealEd->GetSelectedActors()->GetElementSelectionSet());
+		UE::Landscape::DeleteActors({ InActor }, InActor->GetWorld(), bInAllowUI);
 	}
 	else
 	{
@@ -5333,13 +5333,19 @@ void ULandscapeInfo::DeleteActorWhenApplyingModifiedStatus(AActor* InActor)
 
 void ULandscapeInfo::ProcessDeferredDeletions()
 {
+	TArray<AActor*> FinalActorsToDelete;
+	FinalActorsToDelete.Reserve(ActorsToDelete.Num());
 	for (TWeakObjectPtr<AActor> WeakActorPtr : ActorsToDelete)
 	{
 		if (AActor* Actor = WeakActorPtr.Get())
 		{
-			GUnrealEd->DeleteActors({ Actor }, Actor->GetWorld(), GUnrealEd->GetSelectedActors()->GetElementSelectionSet());
+			FinalActorsToDelete.Add(Actor);
 		}
 	}
+	if (!FinalActorsToDelete.IsEmpty())
+	{
+		ensure(UE::Landscape::DeleteActors(FinalActorsToDelete, FinalActorsToDelete[0]->GetWorld(), /* bInAllowUI = */true));
+	}
 	ActorsToDelete.Empty();
 }
 
diff --git a/Engine/Source/Runtime/Landscape/Private/LandscapeUtils.cpp b/Engine/Source/Runtime/Landscape/Private/LandscapeUtils.cpp
--- a/Engine/Source/Runtime/Landscape/Private/LandscapeUtils.cpp
+++ b/Engine/Source/Runtime/Landscape/Private/LandscapeUtils.cpp
@@ -12,8 +12,11 @@
 #include "TextureResource.h"
 
 #if WITH_EDITOR
+#include "Algo/AllOf.h"
+#include "Algo/ForEach.h"
 #include "AssetRegistry/AssetRegistryModule.h"
 #include "Editor/EditorEngine.h"
+#include "Editor/UnrealEdEngine.h"
 #include "EditorDirectories.h"
 #include "Engine/Texture2D.h"
 #include "MaterialDomain.h"
@@ -21,6 +24,8 @@
 #include "LandscapeMaterialInstanceConstant.h"
 #include "LandscapeSettings.h"
 #include "ObjectTools.h"
+#include "Selection.h"
+#include "UnrealEdGlobals.h"
 #endif // WITH_EDITOR
 
 #if WITH_EDITOR
@@ -546,6 +551,26 @@ FString ConvertTargetLayerNamesToString(const TArrayView<const FName>& InTargetL
 	return *FString::Join(TargetLayerStrings, TEXT(","));
 }
 
+bool DeleteActors(const TArray<AActor*>& InActorsToDelete, UWorld* InWorld, bool bInAllowUI)
+{
+	bool bSuccess = true;
+	check(Algo::AllOf(InActorsToDelete, [InWorld](AActor* Actor) { return (Actor != nullptr) && (Actor->GetWorld() == InWorld); }));
+	// If UI is allowed, prefer UUnrealEdEngine::DeleteActors, which handles references to the actor being deleted and asks the user what to do about it :
+	if (bInAllowUI && GUnrealEd && GUnrealEd->GetSelectedActors() && GUnrealEd->GetSelectedActors()->GetElementSelectionSet())
+	{
+		bSuccess = GUnrealEd->DeleteActors(InActorsToDelete, InWorld, GUnrealEd->GetSelectedActors()->GetElementSelectionSet());
+	}
+	else
+	{
+		Algo::ForEach(InActorsToDelete, [InWorld, &bSuccess](AActor* Actor)
+			{
+				bSuccess &= InWorld->DestroyActor(Actor);
+			});
+	}
+
+	return bSuccess;
+}
+
 #endif // WITH_EDITOR
 
 } // end namespace UE::Landscape
\ No newline at end of file
diff --git a/Engine/Source/Runtime/Landscape/Public/LandscapeUtils.h b/Engine/Source/Runtime/Landscape/Public/LandscapeUtils.h
--- a/Engine/Source/Runtime/Landscape/Public/LandscapeUtils.h
+++ b/Engine/Source/Runtime/Landscape/Public/LandscapeUtils.h
@@ -13,6 +13,7 @@
 #include "RenderGraphBuilder.h"
 
 enum EShaderPlatform : uint16;
+class AActor;
 class FRDGBuilder;
 class FRDGEventName;
 class ULandscapeComponent;
@@ -21,6 +22,7 @@ class ULandscapeMaterialInstanceConstant;
 class ULevel;
 class UTexture;
 class UTexture2D;
+class UWorld;
 class FTextureResource;
 class FMaterialUpdateContext;
 enum class ELandscapeToolTargetType : uint8;
@@ -354,6 +356,16 @@ LANDSCAPE_API ULandscapeMaterialInstanceConstant* CreateLandscapeLayerThumbnailM
 /** Concatenates the target layer names in parameter into a string */
 LANDSCAPE_API FString ConvertTargetLayerNamesToString(const TArrayView<const FName>& InTargetLayerNames);
 
+/**
+ * Helper to delete one or multiple actors. 
+ * @param InActorsToDelete the list of actors to delete. Cannot contain null entries and all actors should be part of the world passed in InWorld
+ * @param InWorld world to which all actors to be deleted belong
+ * @param bInAllowUI allows the standard delete actors UX to be displayed, allowing the user to remove lingering reference to these actors, etc.
+ * 
+ * @return true if all actors could be properly deleted
+ */
+LANDSCAPE_API bool DeleteActors(const TArray<AActor*>& InActorsToDelete, UWorld* InWorld, bool bInAllowUI);
+
 #endif // WITH_EDITOR
 
 } // end namespace UE::Landscape
