Fix for landscape brushes allowing brush operations on unloaded border regions.

diff --git a/Engine/Source/Editor/LandscapeEditor/Private/LandscapeEdModeBrushes.cpp b/Engine/Source/Editor/LandscapeEditor/Private/LandscapeEdModeBrushes.cpp
--- a/Engine/Source/Editor/LandscapeEditor/Private/LandscapeEdModeBrushes.cpp
+++ b/Engine/Source/Editor/LandscapeEditor/Private/LandscapeEdModeBrushes.cpp
@@ -298,6 +298,8 @@ public:
 
 	virtual FLandscapeBrushData ApplyBrush(const TArray<FLandscapeToolInteractorPosition>& InInteractorPositions) override
 	{
+		// bCanPaint is evaluated at an old mouse position, and also does not take interactor positions into account,
+		// so cannot be relied upon to be correct vis-a-vis overlapping unloaded components.  We check those again below.
 		if (!bCanPaint)
 		{
 			return FLandscapeBrushData();
@@ -353,6 +355,14 @@ public:
 		}
 		InclusiveBounds.Clip(FIntRect(MinX, MinY, MaxX, MaxY));
 
+		// check if region to be modified overlaps with an unloaded component
+		// TODO [chris.tchou] : The +1 shouldn't be necessary, but without it, edges can be modified by the brush. Possibly a bug in the brush apply code.
+		bool bHasUnloadedComponents = LandscapeInfo->HasUnloadedComponentsInRegion(InclusiveBounds.Min.X, InclusiveBounds.Min.Y, InclusiveBounds.Max.X+1, InclusiveBounds.Max.Y+1);
+		if (bHasUnloadedComponents)
+		{
+			return FLandscapeBrushData();
+		}
+
 		// FLandscapeBrushData requires exclusive bounds : 
 		FLandscapeBrushData BrushData(FIntRect(InclusiveBounds.Min, InclusiveBounds.Max + 1));
 
diff --git a/Engine/Source/Runtime/Landscape/Private/LandscapeEdit.cpp b/Engine/Source/Runtime/Landscape/Private/LandscapeEdit.cpp
--- a/Engine/Source/Runtime/Landscape/Private/LandscapeEdit.cpp
+++ b/Engine/Source/Runtime/Landscape/Private/LandscapeEdit.cpp
@@ -3051,6 +3051,7 @@ bool ULandscapeInfo::HasUnloadedComponentsInRegion(int32 X1, int32 Y1, int32 X2,
 	{
 		UWorld* World = LandscapeActor->GetWorld();
 		ULevel* Level = LandscapeActor->GetLevel();
+		FGuid OriginalGUID = LandscapeActor->GetOriginalLandscapeGuid();
 
 		if (UWorldPartition* WorldPartition = World->GetWorldPartition())
 		{
@@ -3062,11 +3063,11 @@ bool ULandscapeInfo::HasUnloadedComponentsInRegion(int32 X1, int32 Y1, int32 X2,
 				const UActorPartitionSubsystem::FCellCoord MinCoord = UActorPartitionSubsystem::FCellCoord::GetCellCoord(FIntPoint(ComponentIndexX1 * ComponentSizeQuads, ComponentIndexY1 * ComponentSizeQuads), Level, LandscapeActor->GetGridSize());
 				const UActorPartitionSubsystem::FCellCoord MaxCoord = UActorPartitionSubsystem::FCellCoord::GetCellCoord(FIntPoint(ComponentIndexX2 * ComponentSizeQuads, ComponentIndexY2 * ComponentSizeQuads), Level, LandscapeActor->GetGridSize());
 
-				FWorldPartitionHelpers::ForEachActorDescInstance<ALandscapeProxy>(WorldPartition, [this, World, &MinCoord, &MaxCoord, &bResult](const FWorldPartitionActorDescInstance* ActorDescInstance)
+				FWorldPartitionHelpers::ForEachActorDescInstance<ALandscapeProxy>(WorldPartition, [this, World, &MinCoord, &MaxCoord, OriginalGUID, &bResult](const FWorldPartitionActorDescInstance* ActorDescInstance)
 				{
 					FLandscapeActorDesc* LandscapeActorDesc = (FLandscapeActorDesc*)ActorDescInstance->GetActorDesc();
 
-					if (LandscapeActorDesc->GridGuid == LandscapeGuid)
+					if (LandscapeActorDesc->GridGuid == OriginalGUID)
 					{
 						const UActorPartitionSubsystem::FCellCoord ActorCoord(LandscapeActorDesc->GridIndexX, LandscapeActorDesc->GridIndexY, LandscapeActorDesc->GridIndexZ, World->PersistentLevel);
 						if (ActorCoord.X >= MinCoord.X && ActorCoord.Y >= MinCoord.Y && ActorCoord.X <= MaxCoord.X && ActorCoord.Y <= MaxCoord.Y)
