[HWRT] Disable r.RayTracing.Geometry.Landscape.DetectTextureStreaming by default.
- it adds significant CPU overhead and is only necessary when landscape material uses WPO that depends on texture sampling.

diff --git a/Engine/Source/Runtime/Landscape/Private/LandscapeRender.cpp b/Engine/Source/Runtime/Landscape/Private/LandscapeRender.cpp
--- a/Engine/Source/Runtime/Landscape/Private/LandscapeRender.cpp
+++ b/Engine/Source/Runtime/Landscape/Private/LandscapeRender.cpp
@@ -288,7 +288,7 @@ static FAutoConsoleVariableRef CVarLandscapeRayTracingGeometryLODsThatUpdateEver
 	TEXT("If on, LODs that are lower than the specified level will be updated every frame, which can be used to workaround some artifacts caused by texture streaming if you're using WorldPositionOffset on the landscape")
 );
 
-int32 GLandscapeRayTracingGeometryDetectTextureStreaming = 1;
+int32 GLandscapeRayTracingGeometryDetectTextureStreaming = 0;
 static FAutoConsoleVariableRef CVarLandscapeRayTracingGeometryDetectTextureStreaming(
 	TEXT("r.RayTracing.Geometry.Landscape.DetectTextureStreaming"),
 	GLandscapeRayTracingGeometryDetectTextureStreaming,
