Remove redundant and unsafe non-blocking landscape nanite build from PreSave.  Don't build build nanite at all on auto-saves.

diff --git a/Engine/Source/Runtime/Landscape/Private/Landscape.cpp b/Engine/Source/Runtime/Landscape/Private/Landscape.cpp
--- a/Engine/Source/Runtime/Landscape/Private/Landscape.cpp
+++ b/Engine/Source/Runtime/Landscape/Private/Landscape.cpp
@@ -3864,9 +3864,13 @@ void ALandscapeProxy::PreSave(FObjectPreSaveContext ObjectSaveContext)
 		}
 	}
 
-	if (ULandscapeInfo* LandscapeInfo = GetLandscapeInfo())
+	// Update nanite (and block to wait for it).  Don't update nanite on auto-save, since its so slow.
+	if ((ObjectSaveContext.GetSaveFlags() & SAVE_FromAutosave) == 0)
 	{
-		LandscapeInfo->UpdateNanite(ObjectSaveContext.GetTargetPlatform());
+		if (ULandscapeInfo* LandscapeInfo = GetLandscapeInfo())
+		{
+			LandscapeInfo->UpdateNanite(ObjectSaveContext.GetTargetPlatform());
+		}
 	}
 
 	if (ALandscape* Landscape = GetLandscapeActor())
@@ -3884,7 +3888,6 @@ void ALandscapeProxy::PreSave(FObjectPreSaveContext ObjectSaveContext)
 				}
 			});
 		}
-		UpdateNaniteRepresentation(ObjectSaveContext.GetTargetPlatform());
 		UpdateRenderingMethod();
 	}
 
