Fix PCG refresh on undo/redo landscape painting.

diff --git a/Engine/Source/Runtime/Landscape/Classes/Landscape.h b/Engine/Source/Runtime/Landscape/Classes/Landscape.h
--- a/Engine/Source/Runtime/Landscape/Classes/Landscape.h
+++ b/Engine/Source/Runtime/Landscape/Classes/Landscape.h
@@ -370,7 +370,7 @@ public:
 	LANDSCAPE_API virtual bool HasLayersContent() const override;
 	LANDSCAPE_API virtual void UpdateCachedHasLayersContent(bool bInCheckComponentDataIntegrity) override;
 	LANDSCAPE_API void RequestSplineLayerUpdate();
-	LANDSCAPE_API void RequestLayersInitialization(bool bInRequestContentUpdate = true);
+	LANDSCAPE_API void RequestLayersInitialization(bool bInRequestContentUpdate = true, bool bInForceLayerResourceReset = true);
 	LANDSCAPE_API void RequestLayersContentUpdateForceAll(ELandscapeLayerUpdateMode InModeMask = ELandscapeLayerUpdateMode::Update_All, bool bInUserTriggered = false);
 	LANDSCAPE_API void RequestLayersContentUpdate(ELandscapeLayerUpdateMode InModeMask);
 	LANDSCAPE_API bool ReorderLayer(int32 InStartingLayerIndex, int32 InDestinationLayerIndex);
@@ -556,7 +556,7 @@ private:
 	bool SupportsEditLayersLocalMerge();
 
 	ELandscapeEditLayersMergeMode GetEditLayersMergeMode();
-	bool CreateLayersRenderingResource();
+	bool CreateLayersRenderingResource(bool bInForceLayerResourceReset);
 	void PrepareEditLayersLocalMergeResources();
 	void UpdateLayersContent(bool bInWaitForStreaming = false, bool bInSkipMonitorLandscapeEdModeChanges = false, bool bIntermediateRender = false, bool bFlushRender = false);
 	bool CanUpdateLayersContent() const;
@@ -755,6 +755,9 @@ private:
 	UPROPERTY(Transient)
 	bool bLandscapeLayersAreInitialized;
 
+	UPROPERTY(Transient)
+	bool bLandscapeLayersForceResourceReset;
+
 	UPROPERTY(Transient)
 	bool bLandscapeLayersAreUsingLocalMerge;
 
diff --git a/Engine/Source/Runtime/Landscape/Classes/LandscapeComponent.h b/Engine/Source/Runtime/Landscape/Classes/LandscapeComponent.h
--- a/Engine/Source/Runtime/Landscape/Classes/LandscapeComponent.h
+++ b/Engine/Source/Runtime/Landscape/Classes/LandscapeComponent.h
@@ -534,6 +534,8 @@ private:
 	UPROPERTY(Transient)
 	bool bPendingLayerCollisionDataUpdate;
 
+	bool bUndoChangedWeightmapAllocs = false;
+
 	/** Dirtied collision height region when painting (only used by Landscape Layer System) */
 	FIntRect LayerDirtyCollisionHeightData;
 #endif // WITH_EDITORONLY_DATA
@@ -838,6 +840,9 @@ public:
 	bool GetPendingCollisionDataUpdate() const { return bPendingCollisionDataUpdate; }
 	void SetPendingLayerCollisionDataUpdate(bool bInPendingLayerCollisionDataUpdate) { bPendingLayerCollisionDataUpdate = bInPendingLayerCollisionDataUpdate; }
 	bool GetPendingLayerCollisionDataUpdate() const { return bPendingLayerCollisionDataUpdate; }
+
+	void SetUndoChangedWeightmapAllocs(bool bInVal) { bUndoChangedWeightmapAllocs = bInVal; }
+	bool GetUndoChangedWeightmapAllocs() const { return bUndoChangedWeightmapAllocs; }
 #endif // WITH_EDITOR
 
 	virtual bool IsShown(const FEngineShowFlags& ShowFlags) const override;
@@ -1290,6 +1295,6 @@ protected:
 public:
 	/** Records the ULandscapeComponents that are modified in any undo/redo operation that is being applied currently */
 	static uint32 UndoRedoModifiedComponentCount;
-	static TArray<ULandscapeComponent*> UndoRedoModifiedComponents;
+	static TMap<ULandscapeComponent*, uint32> UndoRedoModifiedComponents;
 #endif // WITH_EDITOR
 };
diff --git a/Engine/Source/Runtime/Landscape/Private/Landscape.cpp b/Engine/Source/Runtime/Landscape/Private/Landscape.cpp
--- a/Engine/Source/Runtime/Landscape/Private/Landscape.cpp
+++ b/Engine/Source/Runtime/Landscape/Private/Landscape.cpp
@@ -1693,6 +1693,7 @@ ALandscape::ALandscape(const FObjectInitializer& ObjectInitializer)
 	WeightmapScratchExtractLayerTextureResource = nullptr;
 	WeightmapScratchPackLayerTextureResource = nullptr;
 	bLandscapeLayersAreInitialized = false;
+	bLandscapeLayersForceResourceReset = true;
 	LandscapeEdMode = nullptr;
 	bGrassUpdateEnabled = true;
 	bIsSpatiallyLoaded = false;
@@ -4191,7 +4192,7 @@ bool ULandscapeInfo::UpdateLayerInfoMap(ALandscapeProxy* Proxy /*= nullptr*/, bo
 		ALandscape* Landscape = LandscapeActor.Get();
 		if (Landscape && Landscape->HasLayersContent())
 		{
-			Landscape->RequestLayersInitialization(/*bInRequestContentUpdate*/false);
+			Landscape->RequestLayersInitialization(/*bInRequestContentUpdate*/false, /*bInForceLayerResourceReset*/false);
 		}
 	}
 	return bLayerInfoMapChanged;
diff --git a/Engine/Source/Runtime/Landscape/Private/LandscapeComponent.cpp b/Engine/Source/Runtime/Landscape/Private/LandscapeComponent.cpp
--- a/Engine/Source/Runtime/Landscape/Private/LandscapeComponent.cpp
+++ b/Engine/Source/Runtime/Landscape/Private/LandscapeComponent.cpp
@@ -11,7 +11,7 @@
 #include "LandscapeHLODBuilder.h"
 
 uint32 ULandscapeComponent::UndoRedoModifiedComponentCount;
-TArray<ULandscapeComponent*> ULandscapeComponent::UndoRedoModifiedComponents;
+TMap<ULandscapeComponent*, uint32> ULandscapeComponent::UndoRedoModifiedComponents;
 
 #endif
 
diff --git a/Engine/Source/Runtime/Landscape/Private/LandscapeEdit.cpp b/Engine/Source/Runtime/Landscape/Private/LandscapeEdit.cpp
--- a/Engine/Source/Runtime/Landscape/Private/LandscapeEdit.cpp
+++ b/Engine/Source/Runtime/Landscape/Private/LandscapeEdit.cpp
@@ -952,9 +952,9 @@ void ULandscapeComponent::PreEditUndo()
 	ALandscapeProxy* Proxy = GetLandscapeProxy();
 	Proxy->CurrentVersion++;
 	
-	// Count the number of components involved in the undo/redo
-	UndoRedoModifiedComponentCount++;
-	check(UndoRedoModifiedComponents.Num() == 0);	
+	// Keep track of components involved in the undo/redo and store their weightmap hashes before they get changed.
+	// Note the hash covers allocations and pointers, but not weightmap contents.
+	UndoRedoModifiedComponents.Add(this, ComputeWeightmapsHash());
 }
 
 void ULandscapeComponent::PostEditUndo()
@@ -1011,16 +1011,15 @@ void ULandscapeComponent::PostEditUndo()
 	//  rebuild or will simply invalidate it) : 
 	Proxy->InvalidateOrUpdateNaniteRepresentation(/* bInCheckContentId = */true, /*InTargetPlatform = */nullptr);
 
-	check(!UndoRedoModifiedComponents.Contains(this));
-	UndoRedoModifiedComponents.Add(this);
-
 	// Only Fixup Weightmaps and Update Material Instances when the LAST modified component has been processed by PostEditUndo()
+	UndoRedoModifiedComponentCount++;
 	if (UndoRedoModifiedComponents.Num() == UndoRedoModifiedComponentCount)
 	{
 		// process all the proxies across all of the components (except proxies that are pending kill)
 		TSet<ALandscapeProxy*> Proxies;
-		for (ULandscapeComponent* Component : UndoRedoModifiedComponents)
+		for (auto& Iter : UndoRedoModifiedComponents)
 		{
+			ULandscapeComponent* Component = Iter.Key;
 			ALandscapeProxy* ComponentProxy = Component->GetLandscapeProxy();
 			if (IsValid(ComponentProxy))
 			{
@@ -1052,21 +1051,30 @@ void ULandscapeComponent::PostEditUndo()
 			TArray<FComponentRecreateRenderStateContext> RecreateRenderStateContexts;
 			{
 				FMaterialUpdateContext MaterialUpdateContext(FMaterialUpdateContext::EOptions::Default & ~FMaterialUpdateContext::EOptions::RecreateRenderStates);
-				for (ULandscapeComponent* Component : UndoRedoModifiedComponents)
+				for (auto& Iter : UndoRedoModifiedComponents)
 				{
+					ULandscapeComponent* Component = Iter.Key;
 					if (IsValid(Component)) // Components can be pending kill, if they were removed by the undo operation
 					{
 						ALandscapeProxy* ComponentProxy = Component->GetLandscapeProxy();
 						if (IsValid(ComponentProxy))
 						{
 							Component->UpdateMaterialInstances(MaterialUpdateContext, RecreateRenderStateContexts);
+
+							uint32 PreHash = Iter.Value;
+							uint32 PostHash = Component->ComputeWeightmapsHash();
+							if (PreHash != PostHash)
+							{
+								// Component's weightmap allocations changed by Undo/Redo.  Keep track of this so that appropriate invalidations can be
+								// done during merge.
+								Component->bUndoChangedWeightmapAllocs = true;
+							}
 						}
 					}
 				}
 			}
 		}
 
-
 		// Clear out the temporarily recorded undo info
 		UndoRedoModifiedComponentCount = 0;
 		UndoRedoModifiedComponents.Empty();
diff --git a/Engine/Source/Runtime/Landscape/Private/LandscapeEditLayers.cpp b/Engine/Source/Runtime/Landscape/Private/LandscapeEditLayers.cpp
--- a/Engine/Source/Runtime/Landscape/Private/LandscapeEditLayers.cpp
+++ b/Engine/Source/Runtime/Landscape/Private/LandscapeEditLayers.cpp
@@ -2386,7 +2386,7 @@ ELandscapeEditLayersMergeMode ALandscape::GetEditLayersMergeMode()
 	return DesiredMode;
 }
 
-bool ALandscape::CreateLayersRenderingResource()
+bool ALandscape::CreateLayersRenderingResource(bool bInForceLayerResourceReset)
 {
 	ULandscapeInfo* Info = GetLandscapeInfo();
 	if (Info == nullptr)
@@ -2407,14 +2407,17 @@ bool ALandscape::CreateLayersRenderingResource()
 	ALandscape* Landscape = GetLandscapeActor();
 	check(Landscape);
 
-	if (GetEditLayersMergeMode() != ELandscapeEditLayersMergeMode::GlobalMerge)
+	// Always reset layer resources in global merge mode (using the second block).  In other merge modes, only release
+	// resources when forced, often because we switched out of global merge mode.
+	bool bCurrentlyGlobalMerge = GetEditLayersMergeMode() == ELandscapeEditLayersMergeMode::GlobalMerge;
+	if (!bCurrentlyGlobalMerge && bInForceLayerResourceReset)
 	{
 		ReleaseLayersRenderingResource();
 
 		Landscape->HeightmapRTList.Empty();
 		Landscape->WeightmapRTList.Empty();
 	}
-	else
+	else if (bCurrentlyGlobalMerge)
 	{
 		// Heightmap mip size
 		int32 ComponentVerts = ((SubsectionSizeQuads + 1) * NumSubsections);
@@ -8800,61 +8803,67 @@ void ALandscape::ReallocateLayersWeightmaps(FUpdateLayersContentContext& InUpdat
 	TMap<ULandscapeComponent*, TArray<ULandscapeLayerInfoObject*>> LayerAllocsPerComponent;
 	// If the job of building the per-component allocations has already been done, just use them : 
 	if (InPerComponentAllocations != nullptr)
-	{ 
+	{
 		LayerAllocsPerComponent = *InPerComponentAllocations;
 	}
 	else
 	{
-	for (FLandscapeLayer& Layer : LandscapeEditLayers)
-	{
-		for (ULandscapeComponent* Component : InUpdateLayersContentContext.LandscapeComponentsWeightmapsToResolve)
+		for (FLandscapeLayer& Layer : LandscapeEditLayers)
 		{
-			TArray<ULandscapeLayerInfoObject*>* ComponentLayerAlloc = LayerAllocsPerComponent.Find(Component);
-			if (ComponentLayerAlloc == nullptr)
-			{
-				TArray<ULandscapeLayerInfoObject*> NewLayerAllocs;
-				ComponentLayerAlloc = &LayerAllocsPerComponent.Add(Component, NewLayerAllocs);
-			}
-
-			// No need for an allocation if the edit layer is invisible : 
-			if (Layer.EditLayer->IsVisible())
+			for (ULandscapeComponent* Component : InUpdateLayersContentContext.LandscapeComponentsWeightmapsToResolve)
 			{
-				if (FLandscapeLayerComponentData* LayerComponentData = Component->GetLayerData(Layer.EditLayer->GetGuid()))
+				TArray<ULandscapeLayerInfoObject*>* ComponentLayerAlloc = LayerAllocsPerComponent.Find(Component);
+				if (ComponentLayerAlloc == nullptr)
 				{
-				for (const FWeightmapLayerAllocationInfo& LayerWeightmapAllocation : LayerComponentData->WeightmapData.LayerAllocations)
+					TArray<ULandscapeLayerInfoObject*> NewLayerAllocs;
+					ComponentLayerAlloc = &LayerAllocsPerComponent.Add(Component, NewLayerAllocs);
+				}
+
+				// No need for an allocation if the edit layer is invisible : 
+				if (Layer.EditLayer->IsVisible())
 				{
-					if (LayerWeightmapAllocation.LayerInfo != nullptr)
+					if (FLandscapeLayerComponentData* LayerComponentData = Component->GetLayerData(Layer.EditLayer->GetGuid()))
 					{
-						ComponentLayerAlloc->AddUnique(LayerWeightmapAllocation.LayerInfo);
+						for (const FWeightmapLayerAllocationInfo& LayerWeightmapAllocation : LayerComponentData->WeightmapData.LayerAllocations)
+						{
+							if (LayerWeightmapAllocation.LayerInfo != nullptr)
+							{
+								ComponentLayerAlloc->AddUnique(LayerWeightmapAllocation.LayerInfo);
+							}
+						}
 					}
 				}
-			}
-			}
 
-			// Add the brush alloc also (only if !InMergeParams.bSkipBrush, but InBrushRequiredAllocations should be empty already if InMergeParams.bSkipBrush is true) :
-			for (ULandscapeLayerInfoObject* BrushLayerInfo : InBrushRequiredAllocations)
-			{
-				if (BrushLayerInfo != nullptr)
+				// Add the brush alloc also (only if !InMergeParams.bSkipBrush, but InBrushRequiredAllocations should be empty already if InMergeParams.bSkipBrush is true) :
+				for (ULandscapeLayerInfoObject* BrushLayerInfo : InBrushRequiredAllocations)
 				{
-					ComponentLayerAlloc->AddUnique(BrushLayerInfo);
+					if (BrushLayerInfo != nullptr)
+					{
+						ComponentLayerAlloc->AddUnique(BrushLayerInfo);
+					}
 				}
 			}
 		}
 	}
-	}
 
 	// Trim the components that don't need weightmaps anymore (e.g. all edit layers are made invisible : there were some
 	// components in LandscapeComponentsWeightmapsToResolve but there aren't anymore now).
-	// Remember which components are removed.  These might need InvalidateGeneratedComponentData to fix up nanite dependencies.
-	TArray<ULandscapeComponent*> ComponentsWithNoAllocs;
-	InUpdateLayersContentContext.LandscapeComponentsWeightmapsToResolve.RemoveAllSwap([&LayerAllocsPerComponent, &ComponentsWithNoAllocs](ULandscapeComponent* Component) -> bool
+	// Record which components have no allocations or had previously had their allocations changed by undo/redo.  These might need InvalidateGeneratedComponentData to fix up dependencies.
+	TArray<ULandscapeComponent*> ComponentsWithChangedAllocs;
+	ComponentsWithChangedAllocs.Reserve(InUpdateLayersContentContext.LandscapeComponentsWeightmapsToResolve.Num());
+	InUpdateLayersContentContext.LandscapeComponentsWeightmapsToResolve.RemoveAllSwap([&LayerAllocsPerComponent, &ComponentsWithChangedAllocs](ULandscapeComponent* Component) -> bool
 		{
 			TArray<ULandscapeLayerInfoObject*>* ComponentLayerAlloc = LayerAllocsPerComponent.Find(Component);
 			check(ComponentLayerAlloc != nullptr);
 			bool bNoAllocs = ComponentLayerAlloc->IsEmpty();
-			if (bNoAllocs)
+			bool bUndoChangedAllocs = Component->GetUndoChangedWeightmapAllocs();
+			if (bNoAllocs || bUndoChangedAllocs)
 			{
-				ComponentsWithNoAllocs.Add(Component);
+				ComponentsWithChangedAllocs.Add(Component);
+				if (bUndoChangedAllocs)
+				{
+					Component->SetUndoChangedWeightmapAllocs(false);  // Reset the flag since it's being handled.
+				}
 			}
 			return bNoAllocs;
 		},
@@ -8895,7 +8904,7 @@ void ALandscape::ReallocateLayersWeightmaps(FUpdateLayersContentContext& InUpdat
 	for (ULandscapeComponent* Component : InUpdateLayersContentContext.LandscapeComponentsWeightmapsToResolve)
 	{
 		// If requested, don't allow the component to use textures that have a weightmap in another component than those provided when reallocating weightmaps :
-		NewCreatedTextures.Append(Component->ReallocateWeightmaps(/*DataInterface = */nullptr, /*InEditLayerGuid = */FGuid(), /*bInSaveToTransactionBuffer = */false, /*bool bInForceReallocate = */false, 
+		NewCreatedTextures.Append(Component->ReallocateWeightmaps(/*DataInterface = */nullptr, /*InEditLayerGuid = */FGuid(), /*bInSaveToTransactionBuffer = */false, /*bool bInForceReallocate = */false,
 			/*InTargetProxy = */nullptr, InRestrictTextureSharingToComponents));
 	}
 
@@ -8914,55 +8923,56 @@ void ALandscape::ReallocateLayersWeightmaps(FUpdateLayersContentContext& InUpdat
 
 	// Clean-up unused weightmap CPUReadback resources
 	Info->ForEachLandscapeProxy([](ALandscapeProxy* Proxy)
-	{
-		TArray<UTexture2D*, TInlineAllocator<64>> EntriesToRemoveFromMap;
-		for (auto& Pair : Proxy->WeightmapsCPUReadback)
 		{
-			UTexture2D* WeightmapTextureKey = Pair.Key;
-			bool IsTextureReferenced = false;
-			for (ULandscapeComponent* Component : Proxy->LandscapeComponents)
+			TArray<UTexture2D*, TInlineAllocator<64>> EntriesToRemoveFromMap;
+			for (auto& Pair : Proxy->WeightmapsCPUReadback)
 			{
-				for (UTexture2D* WeightmapTexture : Component->GetWeightmapTextures(false))
+				UTexture2D* WeightmapTextureKey = Pair.Key;
+				bool IsTextureReferenced = false;
+				for (ULandscapeComponent* Component : Proxy->LandscapeComponents)
 				{
-					if (WeightmapTexture == WeightmapTextureKey)
+					for (UTexture2D* WeightmapTexture : Component->GetWeightmapTextures(false))
 					{
-						IsTextureReferenced = true;
-						break;
+						if (WeightmapTexture == WeightmapTextureKey)
+						{
+							IsTextureReferenced = true;
+							break;
+						}
 					}
 				}
+				if (!IsTextureReferenced)
+				{
+					EntriesToRemoveFromMap.Add(WeightmapTextureKey);
+				}
 			}
-			if (!IsTextureReferenced)
-			{
-				EntriesToRemoveFromMap.Add(WeightmapTextureKey);
-			}
-		}
 
-		if (EntriesToRemoveFromMap.Num())
-		{
-			for (UTexture2D* OldWeightmapTexture : EntriesToRemoveFromMap)
+			if (EntriesToRemoveFromMap.Num())
 			{
-				if (FLandscapeEditLayerReadback** CPUReadbackToDelete = Proxy->WeightmapsCPUReadback.Find(OldWeightmapTexture))
+				for (UTexture2D* OldWeightmapTexture : EntriesToRemoveFromMap)
 				{
-					check(*CPUReadbackToDelete);
-					delete* CPUReadbackToDelete;
-					Proxy->WeightmapsCPUReadback.Remove(OldWeightmapTexture);
+					if (FLandscapeEditLayerReadback** CPUReadbackToDelete = Proxy->WeightmapsCPUReadback.Find(OldWeightmapTexture))
+					{
+						check(*CPUReadbackToDelete);
+						delete* CPUReadbackToDelete;
+						Proxy->WeightmapsCPUReadback.Remove(OldWeightmapTexture);
+					}
 				}
 			}
-		}
 
-		return true;
-	});
+			return true;
+		});
 
 	// When the last pixels are removed from a component it loses its allocations, which means we stop tracking it through the update pipeline.  Add a tracking
 	// object if needed and mark it with bCleared so that the code later on will know to update the collision object.
-	for (ULandscapeComponent* Component : ComponentsWithNoAllocs)
+
+	for (ULandscapeComponent* Component : ComponentsWithChangedAllocs)
 	{
-		FLandscapeEditLayerComponentReadbackResult* ComponentReadbackResult =  UE::Landscape::Private::FindOrAddByComponent(InUpdateLayersContentContext.AllLandscapeComponentReadbackResults, Component, ELandscapeLayerUpdateMode::Update_None);
+		FLandscapeEditLayerComponentReadbackResult* ComponentReadbackResult = UE::Landscape::Private::FindOrAddByComponent(InUpdateLayersContentContext.AllLandscapeComponentReadbackResults, Component, ELandscapeLayerUpdateMode::Update_None);
 		ComponentReadbackResult->bCleared = true;
 		ComponentReadbackResult->bModified = true;
 	}
 
-	InvalidateGeneratedComponentData(ComponentsWithNoAllocs, false);
+	InvalidateGeneratedComponentData(ComponentsWithChangedAllocs, false);
 	ValidateProxyLayersWeightmapUsage();
 
 	InUpdateLayersContentContext.Refresh(FUpdateLayersContentContext::ERefreshFlags::RefreshWeightmapInfos | FUpdateLayersContentContext::ERefreshFlags::RefreshMapHelper);
@@ -11120,7 +11130,7 @@ void ALandscape::UpdateCachedHasLayersContent(bool bInCheckComponentDataIntegrit
 	bHasLayersContent |= ALandscape::HasLayersContent();
 }
 
-void ALandscape::RequestLayersInitialization(bool bInRequestContentUpdate)
+void ALandscape::RequestLayersInitialization(bool bInRequestContentUpdate, bool bInForceLayerResourceReset)
 {
 	if (!CanHaveLayersContent())
 	{
@@ -11128,6 +11138,7 @@ void ALandscape::RequestLayersInitialization(bool bInRequestContentUpdate)
 	}
 
 	bLandscapeLayersAreInitialized = false;
+	bLandscapeLayersForceResourceReset |= bInForceLayerResourceReset;
 	LandscapeSplinesAffectedComponents.Empty();
 
 	if (bInRequestContentUpdate)
@@ -11536,7 +11547,7 @@ void ALandscape::UpdateLayersContent(bool bInWaitForStreaming, bool bInSkipMonit
 	ELandscapeEditLayersMergeMode EditLayersMergeMode = GetEditLayersMergeMode();
 	if (CurrentEditLayersMergeMode != EditLayersMergeMode)
 	{
-		RequestLayersInitialization(/*bInRequestContentUpdate = */true);
+		RequestLayersInitialization(/*bInRequestContentUpdate = */true, /*bForceLayerResourceReset = */true);
 		CurrentEditLayersMergeMode = EditLayersMergeMode;
 	}
 
@@ -11899,10 +11910,11 @@ void ALandscape::InitializeLayers()
 {
 	check(HasLayersContent());
 
-	if (CreateLayersRenderingResource())
+	if (CreateLayersRenderingResource(bLandscapeLayersForceResourceReset))
 	{
 		InitializeLandscapeLayersWeightmapUsage();
 		bLandscapeLayersAreInitialized = true;
+		bLandscapeLayersForceResourceReset = false;
 	}
 }
 
