Fix removing from an container while iterating on it.

diff --git a/Engine/Source/Runtime/Landscape/Private/LandscapeGroup.cpp b/Engine/Source/Runtime/Landscape/Private/LandscapeGroup.cpp
--- a/Engine/Source/Runtime/Landscape/Private/LandscapeGroup.cpp
+++ b/Engine/Source/Runtime/Landscape/Private/LandscapeGroup.cpp
@@ -686,7 +686,11 @@ FLandscapeGroup::~FLandscapeGroup()
 	if (AllRegisteredFixups.Num() > 0)
 	{
 		UE_LOG(LogLandscape, Warning, TEXT("Landscape Group (%d) had registered components at destruction, this indicates some components were not properly unregistered"), this->LandscapeGroupKey);
-		for (ULandscapeHeightmapTextureEdgeFixup* Fixup : AllRegisteredFixups)
+		
+		// Avoid UnregisterComponent removing from AllRegisteredFixups while iterating through it
+		TSet<TObjectPtr<ULandscapeHeightmapTextureEdgeFixup>> LocalAllRegisteredFixups(AllRegisteredFixups);
+
+		for (ULandscapeHeightmapTextureEdgeFixup* Fixup : LocalAllRegisteredFixups)
 		{
 			if (Fixup->ActiveGroup == this)
 			{
