Stop over-invalidating on empty landscape components.  Change to invalidating only components that added or removed weightmap allocations.

diff --git a/Engine/Source/Runtime/Landscape/Private/LandscapeEditLayers.cpp b/Engine/Source/Runtime/Landscape/Private/LandscapeEditLayers.cpp
--- a/Engine/Source/Runtime/Landscape/Private/LandscapeEditLayers.cpp
+++ b/Engine/Source/Runtime/Landscape/Private/LandscapeEditLayers.cpp
@@ -8848,24 +8848,30 @@ void ALandscape::ReallocateLayersWeightmaps(FUpdateLayersContentContext& InUpdat
 		}
 	}
 
+	int32 NumToResolve = InUpdateLayersContentContext.LandscapeComponentsWeightmapsToResolve.Num();
+	int32 NumNoAllocs = 0;
+	int32 NumUndoFlagged = 0;
+	int32 NumWithAdds = 0;
+	int32 NumWithRemoves = 0;
+
 	// Trim the components that don't need weightmaps anymore (e.g. all edit layers are made invisible : there were some
 	// components in LandscapeComponentsWeightmapsToResolve but there aren't anymore now).
 	// Record which components have no allocations or had previously had their allocations changed by undo/redo.  These might need InvalidateGeneratedComponentData to fix up dependencies.
-	TArray<ULandscapeComponent*> ComponentsWithChangedAllocs;
+	TSet<ULandscapeComponent*> ComponentsWithChangedAllocs;
 	ComponentsWithChangedAllocs.Reserve(InUpdateLayersContentContext.LandscapeComponentsWeightmapsToResolve.Num());
-	InUpdateLayersContentContext.LandscapeComponentsWeightmapsToResolve.RemoveAllSwap([&LayerAllocsPerComponent, &ComponentsWithChangedAllocs](ULandscapeComponent* Component) -> bool
+	InUpdateLayersContentContext.LandscapeComponentsWeightmapsToResolve.RemoveAllSwap([&LayerAllocsPerComponent, &ComponentsWithChangedAllocs, &NumNoAllocs, &NumUndoFlagged](ULandscapeComponent* Component) -> bool
 		{
 			TArray<ULandscapeLayerInfoObject*>* ComponentLayerAlloc = LayerAllocsPerComponent.Find(Component);
 			check(ComponentLayerAlloc != nullptr);
 			bool bNoAllocs = ComponentLayerAlloc->IsEmpty();
+			NumNoAllocs += bNoAllocs ? 1 : 0;
 			bool bUndoChangedAllocs = Component->GetUndoChangedWeightmapAllocs();
-			if (bNoAllocs || bUndoChangedAllocs)
+
+			if (bUndoChangedAllocs)
 			{
+				NumUndoFlagged++;
 				ComponentsWithChangedAllocs.Add(Component);
-				if (bUndoChangedAllocs)
-				{
-					Component->SetUndoChangedWeightmapAllocs(false);  // Reset the flag since it's being handled.
-				}
+				Component->SetUndoChangedWeightmapAllocs(false);  // Reset the flag since it's being handled.
 			}
 			return bNoAllocs;
 		},
@@ -8878,6 +8884,9 @@ void ALandscape::ReallocateLayersWeightmaps(FUpdateLayersContentContext& InUpdat
 		TArray<ULandscapeLayerInfoObject*>& ComponentLayerAlloc = ItPair.Value;
 		TArray<FWeightmapLayerAllocationInfo>& ComponentBaseLayerAlloc = Component->GetWeightmapLayerAllocations();
 
+		bool bRemoved = false;
+		bool bAdded = false;
+
 		// Deal with the one that need removal
 		for (int32 i = ComponentBaseLayerAlloc.Num() - 1; i >= 0; --i)
 		{
@@ -8885,6 +8894,7 @@ void ALandscape::ReallocateLayersWeightmaps(FUpdateLayersContentContext& InUpdat
 
 			if (!ComponentLayerAlloc.Contains(Alloc.LayerInfo))
 			{
+				bRemoved = true;
 				ComponentBaseLayerAlloc.RemoveAt(i);
 			}
 		}
@@ -8896,9 +8906,18 @@ void ALandscape::ReallocateLayersWeightmaps(FUpdateLayersContentContext& InUpdat
 
 			if (!AllocExist)
 			{
+				bAdded = true;
 				ComponentBaseLayerAlloc.Add(FWeightmapLayerAllocationInfo(LayerAlloc));
 			}
 		}
+
+		if (bRemoved || bAdded)
+		{
+			ComponentsWithChangedAllocs.Add(Component);
+
+			NumWithRemoves += bRemoved ? 1 : 0;
+			NumWithAdds += bAdded ? 1 : 0;
+		}
 	}
 
 	// Realloc the weightmap so it will create proper texture (if needed) and will set the allocations information
@@ -8964,6 +8983,12 @@ void ALandscape::ReallocateLayersWeightmaps(FUpdateLayersContentContext& InUpdat
 			return true;
 		});
 
+	int32 NumInvalidated = ComponentsWithChangedAllocs.Num();
+
+	// Very spammy logging for interactive edits.
+	UE_LOG(LogLandscape, VeryVerbose, TEXT("ReallocateLayersWeightmaps - Components ToResolve: %d, Invalidated: %d - Added: %d, Removed: %d, NoAlloc: %d, UndoFlagged: %d"),
+		NumToResolve, NumInvalidated, NumWithAdds, NumWithRemoves, NumNoAllocs, NumUndoFlagged);
+
 	// When the last pixels are removed from a component it loses its allocations, which means we stop tracking it through the update pipeline.  Add a tracking
 	// object if needed and mark it with bCleared so that the code later on will know to update the collision object.
 
