Landscape Editor - Fix assert after using the resize tool on a landscape with edit layers. Auto-selects the first edit layer when duplicating an edit layer from one landscape actor to another

diff --git a/Engine/Source/Runtime/Landscape/Private/LandscapeEdit.cpp b/Engine/Source/Runtime/Landscape/Private/LandscapeEdit.cpp
--- a/Engine/Source/Runtime/Landscape/Private/LandscapeEdit.cpp
+++ b/Engine/Source/Runtime/Landscape/Private/LandscapeEdit.cpp
@@ -3939,7 +3939,8 @@ void ALandscapeProxy::Import(const FGuid& InGuid, int32 InMinX, int32 InMinY, in
 
 			ComponentsToProcess.Append(ToRawPtrTArrayUnsafe(LandscapeComponents));
 		}
-
+		// CreateDefaultLayer or DuplicateLayerAndMoveBrushes must set the SelectedEditLayerIndex during import
+		check(LandscapeActor->GetSelectedEditLayerIndex() != INDEX_NONE);
 		check(LayerImportSettings.Num() != 0);
 		// Currently only supports reimporting heightmap data into a single edit layer, which will always be the default layer
 		ReimportDestinationLayerGuid = LayerImportSettings[0].DestinationLayerGuid;
diff --git a/Engine/Source/Runtime/Landscape/Private/LandscapeEditLayers.cpp b/Engine/Source/Runtime/Landscape/Private/LandscapeEditLayers.cpp
--- a/Engine/Source/Runtime/Landscape/Private/LandscapeEditLayers.cpp
+++ b/Engine/Source/Runtime/Landscape/Private/LandscapeEditLayers.cpp
@@ -13322,6 +13322,14 @@ FLandscapeLayer* ALandscape::DuplicateLayerAndMoveBrushes(const FLandscapeLayer&
 	}
 
 	int32 AddedIndex = LandscapeEditLayers.Add(NewLayer);
+
+	// If this is the first edit layer duplicated, auto select it so our landscape actor maintains a valid selection state
+	if (LandscapeEditLayers.Num() == 1 && SelectedEditLayerIndex == INDEX_NONE)
+	{
+		SetSelectedEditLayerIndex(0);
+	}
+	check(SelectedEditLayerIndex != INDEX_NONE);
+
 	OnLayerCreatedInternal(NewLayer.EditLayer);
 
 	return &LandscapeEditLayers[AddedIndex];
