Landscape optimization for the editor : don't inspect materials for finding mobile layer names when generating mobile weightmap allocations while there's no weightmap allocations in a given landscape component in the first place

diff --git a/Engine/Source/Runtime/Landscape/Private/Landscape.cpp b/Engine/Source/Runtime/Landscape/Private/Landscape.cpp
--- a/Engine/Source/Runtime/Landscape/Private/Landscape.cpp
+++ b/Engine/Source/Runtime/Landscape/Private/Landscape.cpp
@@ -638,6 +638,8 @@ FGuid ALandscapeProxy::GetNaniteContentId() const
 
 void ULandscapeComponent::CheckGenerateMobilePlatformData(bool bIsCooking, const ITargetPlatform* TargetPlatform)
 {
+	TRACE_CPUPROFILER_EVENT_SCOPE(ULandscapeComponent::CheckGenerateMobilePlatformData);
+
 	// Regenerate platform data only when it's missing or there is a valid hash-mismatch.
 	FBufferArchive ComponentStateAr;
 	SerializeStateHashes(ComponentStateAr);
diff --git a/Engine/Source/Runtime/Landscape/Private/LandscapeEdit.cpp b/Engine/Source/Runtime/Landscape/Private/LandscapeEdit.cpp
--- a/Engine/Source/Runtime/Landscape/Private/LandscapeEdit.cpp
+++ b/Engine/Source/Runtime/Landscape/Private/LandscapeEdit.cpp
@@ -7885,6 +7885,8 @@ bool ALandscapeStreamingProxy::IsValidLandscapeActor(ALandscape* Landscape)
 /* Returns the list of layer names relevant to mobile platforms. Walks the material tree following feature level switch nodes. */
 static void GetAllMobileRelevantLayerNames(TSet<FName>& OutLayerNames, UMaterial* InMaterial)
 {
+	TRACE_CPUPROFILER_EVENT_SCOPE(GetAllMobileRelevantLayerNames);
+
 	TArray<FName> LayerNames;
 
 	bool bMobileUseRuntimeGrassMapGeneration = false;
@@ -7930,6 +7932,14 @@ static void GetAllMobileRelevantLayerNames(TSet<FName>& OutLayerNames, UMaterial
 
 void ULandscapeComponent::GenerateMobileWeightmapLayerAllocations()
 {
+	TRACE_CPUPROFILER_EVENT_SCOPE(ULandscapeComponent::GenerateMobileWeightmapLayerAllocations);
+
+	// Early-out if the component doesn't even have weightmap allocations in the first place, since GetAllMobileRelevantLayerNames can be quite expensive :
+	if (WeightmapLayerAllocations.IsEmpty())
+	{
+		return;
+	}
+
 	const bool bComponentHasHoles = ComponentHasVisibilityPainted();
 	UMaterialInterface* const HoleMaterial = bComponentHasHoles ? GetLandscapeHoleMaterial() : nullptr;
 	UMaterialInterface* const MaterialToUse = bComponentHasHoles && HoleMaterial ? HoleMaterial : GetLandscapeMaterial();
