Landscape Editor - Fix to ensure TargetLayers cannot be created with an invalid name. Filters GetLayersFromMaterial to exclude invalid layers and auto-renames invalid target layers in AddTargetLayer

diff --git a/Engine/Source/Runtime/Landscape/Classes/LandscapeProxy.h b/Engine/Source/Runtime/Landscape/Classes/LandscapeProxy.h
--- a/Engine/Source/Runtime/Landscape/Classes/LandscapeProxy.h
+++ b/Engine/Source/Runtime/Landscape/Classes/LandscapeProxy.h
@@ -1327,8 +1327,8 @@ PRAGMA_ENABLE_DEPRECATION_WARNINGS
 	void UpdateProxyLayersWeightmapUsage();
 	void RequestProxyLayersWeightmapUsageUpdate();
 
-	LANDSCAPE_API static const TArray<FName>& GetLayersFromMaterial(UMaterialInterface* Material);
-	LANDSCAPE_API const TArray<FName>& GetLayersFromMaterial() const;
+	LANDSCAPE_API static const TArray<FName> GetLayersFromMaterial(UMaterialInterface* Material);
+	LANDSCAPE_API const TArray<FName> GetLayersFromMaterial() const;
 
 	// Get all Layer names used by the materials assigned the proxy and components (Proxy Material, Component, Component LOD Material and Hole Material)
 	LANDSCAPE_API TArray<FName> RetrieveTargetLayerNamesFromMaterials() const;
@@ -1597,9 +1597,12 @@ PRAGMA_ENABLE_DEPRECATION_WARNINGS
 	// Add a TargetLayer with a unique name and no layer object info
 	LANDSCAPE_API FLandscapeTargetLayerSettings& AddTargetLayer();
 
-	// Add a named TargetLayer 
+	// Add a TargetLayer with the given name or a generated unique name if the given name is invalid
 	LANDSCAPE_API FLandscapeTargetLayerSettings& AddTargetLayer(const FName& Name, const FLandscapeTargetLayerSettings& TargetLayerSettings, bool bPostEditChange = true);
 
+	// Generate a unique TargetLayer name
+	LANDSCAPE_API const FName GenerateUniqueTargetLayerName() const;
+
 	// Update a existing Target Layer
 	LANDSCAPE_API bool UpdateTargetLayer(const FName& Name, const FLandscapeTargetLayerSettings& TargetLayerSettings, bool bPostEditChange = true);
 
diff --git a/Engine/Source/Runtime/Landscape/Private/Landscape.cpp b/Engine/Source/Runtime/Landscape/Private/Landscape.cpp
--- a/Engine/Source/Runtime/Landscape/Private/Landscape.cpp
+++ b/Engine/Source/Runtime/Landscape/Private/Landscape.cpp
@@ -6908,21 +6908,14 @@ bool ALandscapeProxy::RemoveTargetLayer(const FName& Name,  bool bPostEditChange
 
 FLandscapeTargetLayerSettings& ALandscapeProxy::AddTargetLayer()
 {
-	int32 StartIndex = GetTargetLayers().Num();
-	FName NewName;
-	do
-	{
-		NewName = FName(FString::Format(TEXT("Layer_{0}"), { StartIndex++ } ));
-	} while (HasTargetLayer(NewName));
-
-	return AddTargetLayer(NewName, FLandscapeTargetLayerSettings());
+	return AddTargetLayer(GenerateUniqueTargetLayerName(), FLandscapeTargetLayerSettings());
 }
 	
 FLandscapeTargetLayerSettings& ALandscapeProxy::AddTargetLayer(const FName& Name, const FLandscapeTargetLayerSettings& TargetLayerSettings, bool bPostEditChange)
 {
 	Modify();
 
-	FLandscapeTargetLayerSettings& Settings = TargetLayers.Add(Name, TargetLayerSettings);
+	FLandscapeTargetLayerSettings& Settings = TargetLayers.Add(Name.IsNone() ? GenerateUniqueTargetLayerName() : Name, TargetLayerSettings);
 	if (FProperty* Property = StaticClass()->FindPropertyByName(GET_MEMBER_NAME_CHECKED(ALandscapeProxy, TargetLayers)); Property != nullptr && bPostEditChange)
 	{
 		FPropertyChangedEvent PropertyChangedEvent(Property);
@@ -6932,6 +6925,18 @@ FLandscapeTargetLayerSettings& ALandscapeProxy::AddTargetLayer(const FName& Name
 	return Settings;
 }
 
+const FName ALandscapeProxy::GenerateUniqueTargetLayerName() const
+{
+	int32 StartIndex = GetTargetLayers().Num();
+	FName NewName;
+	do
+	{
+		NewName = FName(FString::Format(TEXT("Layer_{0}"), { StartIndex++ }));
+	} while (HasTargetLayer(NewName));
+
+	return NewName;
+}
+
 bool ALandscapeProxy::UpdateTargetLayer(const FName& Name, const FLandscapeTargetLayerSettings& InTargetLayerSettings, bool bPostEditChange)
 {
 	FLandscapeTargetLayerSettings* TargetLayerSettings = TargetLayers.Find(Name);
diff --git a/Engine/Source/Runtime/Landscape/Private/LandscapeEdit.cpp b/Engine/Source/Runtime/Landscape/Private/LandscapeEdit.cpp
--- a/Engine/Source/Runtime/Landscape/Private/LandscapeEdit.cpp
+++ b/Engine/Source/Runtime/Landscape/Private/LandscapeEdit.cpp
@@ -3151,20 +3151,26 @@ struct FHeightmapInfo
 	TArray<FColor*> HeightmapTextureMipData;
 };
 
-const TArray<FName>& ALandscapeProxy::GetLayersFromMaterial(UMaterialInterface* MaterialInterface)
+const TArray<FName> ALandscapeProxy::GetLayersFromMaterial(UMaterialInterface* MaterialInterface)
 {
 	if (MaterialInterface)
 	{
 		const FMaterialCachedExpressionData& CachedExpressionData = MaterialInterface->GetCachedExpressionData();
 		if (CachedExpressionData.EditorOnlyData)
 		{
-			return CachedExpressionData.EditorOnlyData->LandscapeLayerNames;
+			TArray<FName> LayerNames = CachedExpressionData.EditorOnlyData->LandscapeLayerNames;
+			LayerNames.RemoveAll([](const FName& Name)
+			{
+				return Name.IsNone();
+			});
+
+			return LayerNames;
 		}
 	}
 	return FMaterialCachedExpressionEditorOnlyData::EmptyData.LandscapeLayerNames;
 }
 
-const TArray<FName>& ALandscapeProxy::GetLayersFromMaterial() const
+const TArray<FName> ALandscapeProxy::GetLayersFromMaterial() const
 {
 	return GetLayersFromMaterial(LandscapeMaterial);
 }
