Fix crash when using landscape update button.

diff --git a/Engine/Source/Runtime/Landscape/Private/Landscape.cpp b/Engine/Source/Runtime/Landscape/Private/Landscape.cpp
--- a/Engine/Source/Runtime/Landscape/Private/Landscape.cpp
+++ b/Engine/Source/Runtime/Landscape/Private/Landscape.cpp
@@ -5200,7 +5200,10 @@ bool ULandscapeInfo::IsPackageModified(UPackage* InPackage) const
 int32 ULandscapeInfo::MarkModifiedPackagesAsDirty()
 {
 	int32 NumDirtied = 0;
-	for (TWeakObjectPtr<UPackage> WeakPackagePtr : ModifiedPackages)
+	// Move into a local set to avoid OnMarkPackageDirty triggering from this loop and changing it during iteration.
+	TSet<TWeakObjectPtr<UPackage>, TWeakObjectPtrSetKeyFuncs<TWeakObjectPtr<UPackage>>> LocalModifiedPackages(MoveTemp(ModifiedPackages));
+	ModifiedPackages.Empty();
+	for (TWeakObjectPtr<UPackage> WeakPackagePtr : LocalModifiedPackages)
 	{
 		if (UPackage* Package = WeakPackagePtr.Get())
 		{
@@ -5209,7 +5212,6 @@ int32 ULandscapeInfo::MarkModifiedPackagesAsDirty()
 			NumDirtied += (!bWasDirty && bIsDirty) ? 1 : 0;
 		}
 	}
-	ModifiedPackages.Empty();
 
 	ProcessDeferredDeletions();
 
