Fixed invalid checks in landscape batched merge

diff --git a/Engine/Source/Runtime/Landscape/Private/LandscapeBlueprintBrushBase.cpp b/Engine/Source/Runtime/Landscape/Private/LandscapeBlueprintBrushBase.cpp
--- a/Engine/Source/Runtime/Landscape/Private/LandscapeBlueprintBrushBase.cpp
+++ b/Engine/Source/Runtime/Landscape/Private/LandscapeBlueprintBrushBase.cpp
@@ -329,7 +329,7 @@ void ALandscapeBlueprintBrushBase::GetRendererStateInfo(const UE::Landscape::Edi
 		InMergeContext->ForEachValidTargetLayer(
 			[InMergeContext, this, &OutSupportedTargetTypeState, &OutEnabledTargetTypeState](int32 InTargetLayerIndex, const FName& InTargetLayerName, ULandscapeLayerInfoObject* InWeightmapLayerInfo)
 			{
-				if (CanAffectWeightmapLayer(InTargetLayerName))
+				if (CanAffectWeightmapLayer(InTargetLayerName) && InMergeContext->IsValidTargetLayerName(InTargetLayerName))
 				{
 					OutSupportedTargetTypeState.AddWeightmap(InTargetLayerIndex);
 					if (AffectsWeightmapLayer(InTargetLayerName))
diff --git a/Engine/Source/Runtime/Landscape/Private/LandscapeEditLayerMergeContext.cpp b/Engine/Source/Runtime/Landscape/Private/LandscapeEditLayerMergeContext.cpp
--- a/Engine/Source/Runtime/Landscape/Private/LandscapeEditLayerMergeContext.cpp
+++ b/Engine/Source/Runtime/Landscape/Private/LandscapeEditLayerMergeContext.cpp
@@ -83,10 +83,9 @@ int32 FMergeContext::IsValidTargetLayerNameChecked(const FName& InName) const
 	return ValidTargetLayerBitIndices[Index];
 }
 
-int32 FMergeContext::IsValidTargetLayerIndex(int32 InIndex) const
+int32 FMergeContext::IsTargetLayerIndexValid(int32 InIndex) const
 {
-	check(AllTargetLayerNames.IsValidIndex(InIndex));
-	return ValidTargetLayerBitIndices[InIndex];
+	return AllTargetLayerNames.IsValidIndex(InIndex);
 }
 
 int32 FMergeContext::GetTargetLayerIndexForName(const FName& InName) const
diff --git a/Engine/Source/Runtime/Landscape/Private/LandscapeEditLayerTargetTypeState.cpp b/Engine/Source/Runtime/Landscape/Private/LandscapeEditLayerTargetTypeState.cpp
--- a/Engine/Source/Runtime/Landscape/Private/LandscapeEditLayerTargetTypeState.cpp
+++ b/Engine/Source/Runtime/Landscape/Private/LandscapeEditLayerTargetTypeState.cpp
@@ -78,7 +78,7 @@ bool FEditLayerTargetTypeState::IsActiveChecked(ELandscapeToolTargetType InTarge
 
 bool FEditLayerTargetTypeState::IsActive(ELandscapeToolTargetType InTargetType, int32 InWeightmapLayerIndex) const
 {
-	check((InWeightmapLayerIndex == INDEX_NONE) || MergeContext->IsValidTargetLayerIndex(InWeightmapLayerIndex));
+	check((InWeightmapLayerIndex == INDEX_NONE) || MergeContext->IsTargetLayerIndexValid(InWeightmapLayerIndex));
 
 	if (EnumHasAnyFlags(TargetTypeMask, GetLandscapeToolTargetTypeAsFlags(InTargetType)))
 	{
@@ -176,7 +176,7 @@ void FEditLayerTargetTypeState::AddWeightmapChecked(const FName& InWeightmapLaye
 
 void FEditLayerTargetTypeState::AddWeightmap(int32 InWeightmapLayerIndex)
 {
-	check(MergeContext->IsValidTargetLayerIndex(InWeightmapLayerIndex));
+	check(MergeContext->IsTargetLayerIndexValid(InWeightmapLayerIndex));
 
 	checkf((InWeightmapLayerIndex != MergeContext->GetVisibilityTargetLayerIndex()) || EnumHasAnyFlags(TargetTypeMask, ELandscapeToolTargetTypeFlags::Visibility),
 		TEXT("The visibility layer may only be used for target type states that support visibility"));
@@ -202,7 +202,7 @@ void FEditLayerTargetTypeState::RemoveWeightmapChecked(const FName& InWeightmapL
 
 void FEditLayerTargetTypeState::RemoveWeightmap(int32 InWeightmapLayerIndex)
 {
-	check(MergeContext->IsValidTargetLayerIndex(InWeightmapLayerIndex));
+	check(MergeContext->IsTargetLayerIndexValid(InWeightmapLayerIndex));
 
 	checkf(!EnumHasAnyFlags(TargetTypeMask, ELandscapeToolTargetTypeFlags::Visibility) || (InWeightmapLayerIndex != MergeContext->GetVisibilityTargetLayerIndex()),
 		TEXT("Cannot remove weightmap %s from a target type state that supports visibility"), *MergeContext->GetTargetLayerNameForIndexChecked(InWeightmapLayerIndex).ToString());
diff --git a/Engine/Source/Runtime/Landscape/Private/LandscapeEditLayers.cpp b/Engine/Source/Runtime/Landscape/Private/LandscapeEditLayers.cpp
--- a/Engine/Source/Runtime/Landscape/Private/LandscapeEditLayers.cpp
+++ b/Engine/Source/Runtime/Landscape/Private/LandscapeEditLayers.cpp
@@ -5379,9 +5379,9 @@ TArray<UE::Landscape::EditLayers::FEditLayerRendererState> FLandscapeLayer::GetE
 	ELandscapeToolTargetTypeFlags LayerTargetTypeMask = EditLayer->GetEnabledTargetTypeMask();
 	// Then disable all types that are not in the layer's mask :
 	for (FEditLayerRendererState& LayerRendererState : RendererStates)
-{
+	{
 		LayerRendererState.DisableTargetTypeMask(~LayerTargetTypeMask);
-		}
+	}
 
 	return RendererStates;
 }
@@ -5453,7 +5453,7 @@ TArray<UE::Landscape::EditLayers::FEditLayerRenderItem> ULandscapeEditLayerPersi
 				// Iterate through all allocated weightmaps in order to find which one we will really write to : 
 				for (const FWeightmapLayerAllocationInfo& LayerAllocationInfo : LayerData->WeightmapData.LayerAllocations)
 				{
-					if (LayerAllocationInfo.IsAllocated())
+					if (LayerAllocationInfo.IsAllocated() && InMergeContext->IsValidTargetLayerName(LayerAllocationInfo.GetLayerName()))
 					{
 						if (LayerAllocationInfo.LayerInfo == ALandscapeProxy::VisibilityLayer)
 						{
@@ -6798,10 +6798,10 @@ UE::Landscape::EditLayers::FMergeRenderContext ALandscape::PrepareEditLayersMerg
 			for (const FEditLayerRenderItem& RenderItem : EditLayerRendererRenderInfo.RendererState.GetRenderer()->GetRenderItems(&MergeRenderContext))
 			{
 				checkf(RenderItem.GetTargetTypeState().Intersect(EditLayerRendererRenderInfo.RendererState.GetSupportedTargetTypeState()) == RenderItem.GetTargetTypeState(),
-					TEXT("One of edit layer renderer %s's render items target type state is not supported by the renderer's supported target type state. \
-					This is invalid : the renderer's supported target type state should be a superset of its render item's target type state. \n\
-					- Render item state: %s \n\
-					- Renderer supported state: %s \n"), *EditLayerRendererDebugName, *RenderItem.GetTargetTypeState().ToString(), *EditLayerRendererRenderInfo.RendererState.GetSupportedTargetTypeState().ToString());
+					TEXT("One of edit layer renderer %s's render items target type state is not supported by the renderer's supported target type state. "
+					"This is invalid : the renderer's supported target type state should be a superset of its render item's target type state. \n"
+					"- Render item state: %s \n"
+					"- Renderer supported state: %s \n"), *EditLayerRendererDebugName, *RenderItem.GetTargetTypeState().ToString(), *EditLayerRendererRenderInfo.RendererState.GetSupportedTargetTypeState().ToString());
 
 				// Only consider render items which intersect with the renderer's active state (e.g. one a given renderer, a render item (A) might affect heightmaps only and another one (B) weightmaps only, 
 				//  then if performing a heightmap merge, the RendererState's mask here will be ELandscapeToolTargetTypeFlags::Heightmap, so this allows to filter out item B: 
diff --git a/Engine/Source/Runtime/Landscape/Public/LandscapeEditLayerMergeContext.h b/Engine/Source/Runtime/Landscape/Public/LandscapeEditLayerMergeContext.h
--- a/Engine/Source/Runtime/Landscape/Public/LandscapeEditLayerMergeContext.h
+++ b/Engine/Source/Runtime/Landscape/Public/LandscapeEditLayerMergeContext.h
@@ -85,10 +85,11 @@ public:
 	LANDSCAPE_API int32 IsValidTargetLayerNameChecked(const FName& InName) const;
 
 	/**
-	 * @return true if the target layer name has an associated layer info object. Asserts if the layer index isn't known to the merge context
+	 * @return true if the target layer index is valid in the list of declared target layers on the merge context. 
+	 *  Important note : it doesn't mean that it's associated with a valid layer info object, use IsValidTargetLayerName for this!
 	 */
 	LANDSCAPE_EDIT_LAYERS_BATCHED_MERGE_EXPERIMENTAL
-	LANDSCAPE_API int32 IsValidTargetLayerIndex(int32 InIndex) const;
+	LANDSCAPE_API int32 IsTargetLayerIndexValid(int32 InIndex) const;
 
 	/**
 	 * @return the index of this target layer name in this context, INDEX_NONE if not found
