Fix for mobile landscape weightmaps not being taken into account by streaming

diff --git a/Engine/Source/Runtime/Landscape/Classes/LandscapeComponent.h b/Engine/Source/Runtime/Landscape/Classes/LandscapeComponent.h
--- a/Engine/Source/Runtime/Landscape/Classes/LandscapeComponent.h
+++ b/Engine/Source/Runtime/Landscape/Classes/LandscapeComponent.h
@@ -966,7 +966,7 @@ public:
 #endif
 
 	int32 GetCurrentRuntimeMaterialInstanceCount() const;
-	class UMaterialInterface* GetCurrentRuntimeMaterialInterface(int32 InIndex);
+	class UMaterialInterface* GetCurrentRuntimeMaterialInterface(int32 InIndex) const;
 
 	LANDSCAPE_API int32 GetMaterialInstanceCount(bool InDynamic = true) const;
 	LANDSCAPE_API class UMaterialInstance* GetMaterialInstance(int32 InIndex, bool InDynamic = true) const;
diff --git a/Engine/Source/Runtime/Landscape/Private/Landscape.cpp b/Engine/Source/Runtime/Landscape/Private/Landscape.cpp
--- a/Engine/Source/Runtime/Landscape/Private/Landscape.cpp
+++ b/Engine/Source/Runtime/Landscape/Private/Landscape.cpp
@@ -372,7 +372,7 @@ int32 ULandscapeComponent::GetCurrentRuntimeMaterialInstanceCount() const
 	return GetMaterialInstanceCount(bDynamic);
 }
 
-class UMaterialInterface* ULandscapeComponent::GetCurrentRuntimeMaterialInterface(int32 InIndex)
+class UMaterialInterface* ULandscapeComponent::GetCurrentRuntimeMaterialInterface(int32 InIndex) const
 {
 	ALandscapeProxy* Proxy = GetLandscapeProxy();
 	const ERHIFeatureLevel::Type FeatureLevel = GetLandscapeProxy()->GetWorld()->GetFeatureLevel();
diff --git a/Engine/Source/Runtime/Landscape/Private/LandscapeRender.cpp b/Engine/Source/Runtime/Landscape/Private/LandscapeRender.cpp
--- a/Engine/Source/Runtime/Landscape/Private/LandscapeRender.cpp
+++ b/Engine/Source/Runtime/Landscape/Private/LandscapeRender.cpp
@@ -4351,11 +4351,11 @@ void ULandscapeComponent::GetStreamingRenderAssetInfo(FStreamingTextureLevelCont
 	}
 
 	ERHIFeatureLevel::Type FeatureLevel = LevelContext.GetFeatureLevel();
-	int32 MaterialInstanceCount = GetMaterialInstanceCount();
+	int32 MaterialInstanceCount = GetCurrentRuntimeMaterialInstanceCount();
 
 	for (int32 MaterialIndex = 0; MaterialIndex < MaterialInstanceCount; ++MaterialIndex)
 	{
-		const UMaterialInterface* MaterialInterface = GetMaterialInstance(MaterialIndex);
+		const UMaterialInterface* MaterialInterface = GetCurrentRuntimeMaterialInterface(MaterialIndex);
 
 		// Normal usage...
 		// Enumerate the textures used by the material.
