Fixed invalidations of landscape physical material during painting.

The last few frames of a paint stroke could be lost because the invalidations were ignored while a physical material task was still pending.

diff --git a/Engine/Source/Runtime/Landscape/Classes/LandscapeComponent.h b/Engine/Source/Runtime/Landscape/Classes/LandscapeComponent.h
--- a/Engine/Source/Runtime/Landscape/Classes/LandscapeComponent.h
+++ b/Engine/Source/Runtime/Landscape/Classes/LandscapeComponent.h
@@ -536,6 +536,9 @@ private:
 
 	bool bUndoChangedWeightmapAllocs = false;
 
+	// Physical material has been invalidated and no new physical material task has been started since then
+	bool bPendingPhysicalMaterialInvalidation = false;
+
 	/** Dirtied collision height region when painting (only used by Landscape Layer System) */
 	FIntRect LayerDirtyCollisionHeightData;
 #endif // WITH_EDITORONLY_DATA
@@ -1163,6 +1166,8 @@ public:
 	 */
 	LANDSCAPE_API void UpdateCollisionLayerData();
 
+	/** Invalidate current or pending physical material data */
+	void InvalidatePhysicalMaterial();
 	/** Returns true if we can currently update physical materials. */
 	bool CanUpdatePhysicalMaterial();
 	/** Update physical material render tasks. */
diff --git a/Engine/Source/Runtime/Landscape/Private/Landscape.cpp b/Engine/Source/Runtime/Landscape/Private/Landscape.cpp
--- a/Engine/Source/Runtime/Landscape/Private/Landscape.cpp
+++ b/Engine/Source/Runtime/Landscape/Private/Landscape.cpp
@@ -6816,7 +6816,7 @@ void ALandscapeProxy::InvalidatePhysicalMaterial()
 {
 	for (ULandscapeComponent* Component : LandscapeComponents)
 	{
-		Component->PhysicalMaterialHash = 0;
+		Component->InvalidatePhysicalMaterial();
 	}
 }
 
diff --git a/Engine/Source/Runtime/Landscape/Private/LandscapeEdit.cpp b/Engine/Source/Runtime/Landscape/Private/LandscapeEdit.cpp
--- a/Engine/Source/Runtime/Landscape/Private/LandscapeEdit.cpp
+++ b/Engine/Source/Runtime/Landscape/Private/LandscapeEdit.cpp
@@ -1688,7 +1688,7 @@ void ULandscapeComponent::UpdateCollisionHeightData(const FColor* const Heightma
 
 	// Invalidate rendered physical materials
 	// These are updated in UpdatePhysicalMaterialTasks()
- 	PhysicalMaterialHash = 0;
+	InvalidatePhysicalMaterial();
 }
 
 void ULandscapeComponent::DestroyCollisionData()
@@ -2092,7 +2092,7 @@ void ULandscapeComponent::UpdateCollisionLayerData(const FColor* const* const We
 
 		// Invalidate rendered physical materials
 		// These are updated in UpdatePhysicalMaterialTasks()
- 		PhysicalMaterialHash = 0;
+		InvalidatePhysicalMaterial();
 
 		// We do not force an update of the physics data here. We don't need the layer information in the editor and it
 		// causes problems if we update it multiple times in a single frame.
@@ -2184,6 +2184,11 @@ bool ULandscapeComponent::GetRenderPhysicalMaterials(TArray<UPhysicalMaterial*>&
 	return bReturnValue;
 }
 
+void ULandscapeComponent::InvalidatePhysicalMaterial()
+{
+	PhysicalMaterialHash = 0;
+	bPendingPhysicalMaterialInvalidation = true;
+}
 
 bool ULandscapeComponent::CanUpdatePhysicalMaterial()
 {
@@ -2215,7 +2220,15 @@ void ULandscapeComponent::FinalizePhysicalMaterial(bool bInImmediatePhysicsRebui
 	}
 
 	UpdateCollisionPhysicalMaterialData(PhysicalMaterialTask.GetResultMaterials(), PhysicalMaterialTask.GetResultIds());
-	PhysicalMaterialHash = PhysicalMaterialTask.GetHash();
+
+	// bPendingPhysicalMaterialInvalidation is reset to false when the task was started.  If the physical material was
+	// re-invalidated since then, don't set PhysicalMaterialHash.  PhysicalMaterialTask will be released and PhysicalMaterialHash
+	// will remain mismatched (typically 0), so on the next tick a new task will be started.  It's important that we don't lose
+	// track of these invalidations while the task is running.  We might only get notified of a change once.
+	if (!bPendingPhysicalMaterialInvalidation)
+	{
+		PhysicalMaterialHash = PhysicalMaterialTask.GetHash();
+	}
 
 	PhysicalMaterialTask.Release();
 
@@ -2275,6 +2288,8 @@ void ULandscapeComponent::UpdatePhysicalMaterialTasks()
 		{
 			bool bSuccess = PhysicalMaterialTask.Init(this, Hash);
 			check(bSuccess && PhysicalMaterialTask.IsValid());
+
+			bPendingPhysicalMaterialInvalidation = false;  //new task is starting, so consume the pending invalidation
 		}
 		else
 		{
