Fixed landscape culling stealing the r.Lumen.DiffuseIndirect.Allow's on changed callback. This is not needed anymore since this CVar's callback already invalidates all components, including landscape ones.

diff --git a/Engine/Source/Runtime/Landscape/Private/LandscapeCulling.cpp b/Engine/Source/Runtime/Landscape/Private/LandscapeCulling.cpp
--- a/Engine/Source/Runtime/Landscape/Private/LandscapeCulling.cpp
+++ b/Engine/Source/Runtime/Landscape/Private/LandscapeCulling.cpp
@@ -147,14 +147,6 @@ constexpr static uint32 MIN_SECTION_SIZE_QUADS = 31u;
 constexpr static uint32 INDIRECT_ARGS_NUM_WORDS = 5u;
 constexpr static uint32 TILE_DATA_ENTRY_NUM_BYTES = 4u;
 
-static void LandscapeComponentsRecreateRenderState(IConsoleVariable* Variable)
-{
-	for (auto* LandscapeComponent : TObjectRange<ULandscapeComponent>())
-	{
-		LandscapeComponent->MarkRenderStateDirty();
-	}
-}
-
 struct FLumenCVarsState
 {
 	FLumenCVarsState()
@@ -162,11 +154,8 @@ struct FLumenCVarsState
 		IConsoleVariable* SupportedCVar = IConsoleManager::Get().FindConsoleVariable(TEXT("r.Lumen.Supported"));
 		bLumenSupported = (SupportedCVar && SupportedCVar->GetInt() != 0);
 		
+		// No need to register to changes : all primitive components are invalidated when this CVar dynamically changes so render states of landscape components get appropriately re-created
 		IConsoleVariable* DiffuseIndirectCVar = IConsoleManager::Get().FindConsoleVariable(TEXT("r.Lumen.DiffuseIndirect.Allow"));
-		if (bLumenSupported && DiffuseIndirectCVar)
-		{
-			DiffuseIndirectCVar->SetOnChangedCallback(FConsoleVariableDelegate::CreateStatic(&LandscapeComponentsRecreateRenderState));
-		}
 		DiffuseIndirectInt = DiffuseIndirectCVar->AsVariableInt();
 	}
 
@@ -866,4 +855,4 @@ bool GetViewArguments(const FSceneView& View, uint32 LandscapeKey, FIntPoint Ren
 	return false;
 }
 
-}//namespace UE::Landscape::Culling
\ No newline at end of file
+}//namespace UE::Landscape::Culling
