Removed 5.5 deprecated landscape code

#rnx

diff --git a/Engine/Plugins/Experimental/LandscapePatch/Source/LandscapePatch/Public/LandscapePatchComponent.h b/Engine/Plugins/Experimental/LandscapePatch/Source/LandscapePatch/Public/LandscapePatchComponent.h
--- a/Engine/Plugins/Experimental/LandscapePatch/Source/LandscapePatch/Public/LandscapePatchComponent.h
+++ b/Engine/Plugins/Experimental/LandscapePatch/Source/LandscapePatch/Public/LandscapePatchComponent.h
@@ -116,11 +116,11 @@ public:
 	UFUNCTION(BlueprintCallable, Category = "LandscapePatch")
 	LANDSCAPEPATCH_API virtual void SetLandscape(ALandscape* NewLandscape);
 
-	UE_DEPRECATED(5.5, "Patch manager use is deprecated. Patches should point to an edit layer using a Guid, and order themselves using priority.")
+	UE_DEPRECATED(all, "Patch manager use is deprecated. Patches should point to an edit layer using a Guid, and order themselves using priority.")
 	UFUNCTION(BlueprintCallable, Category = "LandscapePatch")
 	LANDSCAPEPATCH_API virtual void SetPatchManager(ALandscapePatchManager* NewPatchManager);
 
-	UE_DEPRECATED(5.5, "Patch manager use is deprecated. Patches should point to an edit layer using a Guid, and order themselves using priority.")
+	UE_DEPRECATED(all, "Patch manager use is deprecated. Patches should point to an edit layer using a Guid, and order themselves using priority.")
 	UFUNCTION(BlueprintCallable, Category = "LandscapePatch")
 	LANDSCAPEPATCH_API virtual ALandscapePatchManager* GetPatchManager() const;
 
diff --git a/Engine/Source/Editor/LandscapeEditor/Private/SLandscapeEditor.cpp b/Engine/Source/Editor/LandscapeEditor/Private/SLandscapeEditor.cpp
--- a/Engine/Source/Editor/LandscapeEditor/Private/SLandscapeEditor.cpp
+++ b/Engine/Source/Editor/LandscapeEditor/Private/SLandscapeEditor.cpp
@@ -31,7 +31,7 @@ namespace UE::Landscape::Private {
 	static FAutoConsoleVariableRef CVarEnableRetopoTool(
 		TEXT("landscape.EnableRetopologizeTool"),
 		bEnableRetopoTool,
-		TEXT("Enable the Retopologize tool.  The tool will be fully deprecated in UE5.6, but this cvar will enable it for 5.5"),
+		TEXT("Enable the Retopologize tool.  The tool will be fully deprecated in UE5.7, but this cvar will enable it for 5.6"),
 		ECVF_Default
 	);
 }
diff --git a/Engine/Source/Runtime/Landscape/Classes/Landscape.h b/Engine/Source/Runtime/Landscape/Classes/Landscape.h
--- a/Engine/Source/Runtime/Landscape/Classes/Landscape.h
+++ b/Engine/Source/Runtime/Landscape/Classes/Landscape.h
@@ -384,14 +384,6 @@ public:
 	* @return the index of the newly-created layer
 	*/
 	LANDSCAPE_API int32 CreateLayer(FName InName = NAME_None, const TSubclassOf<ULandscapeEditLayerBase>& InEditLayerClass = TSubclassOf<ULandscapeEditLayerBase>(), bool bInIgnoreLayerCountLimit = false);
-
-	/** 
-	* Creates a new edit layer by copying all the settings from the layer passed in parameter
-	* @Param InLayer edit layer to copy from
-	* @return the index of the newly-created layer
-	*/
-	UE_DEPRECATED(5.5, "Use DuplicateLayerAndMoveBrushes")
-	LANDSCAPE_API int32 CreateLayerFrom(const FLandscapeLayer& InLayer);
 	LANDSCAPE_API void CreateDefaultLayer();
 
 	LANDSCAPE_API void CopyOldDataToDefaultLayer();
@@ -445,18 +437,7 @@ public:
 	LANDSCAPE_API uint8 GetLayerCount() const;
 	UE_DEPRECATED(5.6, "Use GetLayersConst")
 	TArrayView<const FLandscapeLayer> GetLayers() const { return MakeArrayView(LandscapeEditLayers); }
-	UE_DEPRECATED(5.5, "This has moved to private (GetLayerInternal), use the appropriate setters to mutate the edit layer")
-	FLandscapeLayer* GetLayer(int32 InLayerIndex) { return nullptr; }
-	UE_DEPRECATED(5.5, "Use GetLayerConst")
-	LANDSCAPE_API const FLandscapeLayer* GetLayer(int32 InLayerIndex) const;
-	UE_DEPRECATED(5.5, "Use GetLayerConst")
-	LANDSCAPE_API const FLandscapeLayer* GetLayer(const FGuid& InLayerGuid) const;
-	UE_DEPRECATED(5.5, "Use GetLayerConst")
-	LANDSCAPE_API const FLandscapeLayer* GetLayer(const FName& InLayerName) const;
-
-	UE_DEPRECATED(5.5, "Use ForEachLayerConst")
-	LANDSCAPE_API void ForEachLayer(TFunctionRef<void(FLandscapeLayer&)> Fn);
-	
+
 	/**
 	 * Runs the given function on each edit layer, with the possibility of early exit
 	 * Most easily used with a lambda as follows:
@@ -487,15 +468,6 @@ public:
 	LANDSCAPE_API void SetSelectedEditLayerIndex(const int32 InEditLayerIndex);
 	LANDSCAPE_API const int32 GetSelectedEditLayerIndex() const;
 
-	UE_DEPRECATED(5.5, "Use CreateLayer with a ULandscapeEditLayerSplines instead")
-	LANDSCAPE_API void SetLandscapeSplinesReservedLayer(int32 InLayerIndex);
-	UE_DEPRECATED(5.5, "Use FindLayerOfType(ULandscapeEditLayerSplines::StaticClass()) instead")
-	LANDSCAPE_API struct FLandscapeLayer* GetLandscapeSplinesReservedLayer();
-	UE_DEPRECATED(5.5, "Use FindLayerOfType(ULandscapeEditLayerSplines::StaticClass()) instead")
-	LANDSCAPE_API const struct FLandscapeLayer* GetLandscapeSplinesReservedLayer() const;
-	UE_DEPRECATED(5.5, "Use FindLayerOfType(ULandscapeEditLayerSplines::StaticClass()).Guid == GetEditingLayer() instead")
-	LANDSCAPE_API bool IsEditingLayerReservedForSplines() const;
-
 	UE_DEPRECATED(5.6, "Use ULandscapeEditLayerSplines GetWeightmapLayerAllocationBlend().Find instead")
 	LANDSCAPE_API bool IsLayerBlendSubstractive(int32 InLayerIndex, const TWeakObjectPtr<ULandscapeLayerInfoObject>& InLayerInfoObj) const;
 	UE_DEPRECATED(5.6, "Use ULandscapeEditLayerSplines AddOrUpdateWeightmapAllocationLayerBlend instead")
@@ -682,7 +654,7 @@ public:
 	FOnEditLayersMergedDelegate::RegistrationType& OnEditLayersMerged() const { return OnEditLayersMergedDelegate; }
 
 	/** Target Landscape Layer for Landscape Splines */
-	UE_DEPRECATED(5.5, "This has been refactored into the generic ULandscapeEditLayerBase system. Please check for the presence of a ULandscapeEditLayerSplines layer instead")
+	UE_DEPRECATED(all, "This has been refactored into the generic ULandscapeEditLayerBase system. Please check for the presence of a ULandscapeEditLayerSplines layer instead")
 	UPROPERTY()
 	FGuid LandscapeSplinesTargetLayerGuid_DEPRECATED;
 	
@@ -703,7 +675,7 @@ public:
 	UPROPERTY(Transient, DuplicateTransient, TextExportTransient, NonPIEDuplicateTransient)
 	bool bWarnedGlobalMergeDimensionsExceeded = false;
 
-	UE_DEPRECATED(5.5, "This property has moved to private. Use the public accessors instead")
+	UE_DEPRECATED(all, "This property has moved to private. Use the public accessors instead")
 	UPROPERTY()
 	TArray<FLandscapeLayer> LandscapeLayers_DEPRECATED;
 
diff --git a/Engine/Source/Runtime/Landscape/Classes/LandscapeComponent.h b/Engine/Source/Runtime/Landscape/Classes/LandscapeComponent.h
--- a/Engine/Source/Runtime/Landscape/Classes/LandscapeComponent.h
+++ b/Engine/Source/Runtime/Landscape/Classes/LandscapeComponent.h
@@ -211,11 +211,11 @@ struct FLandscapeComponentGrassData
 
 	// Guid per material instance in the hierarchy between the assigned landscape material (instance) and the root UMaterial
 	// used to detect changes to material instance parameters or the root material that could affect the grass maps
-	UE_DEPRECATED(5.5, "GenerationHash is now used") 
+	UE_DEPRECATED(all, "GenerationHash is now used") 
 	TArray<FGuid, TInlineAllocator<2>> MaterialStateIds_DEPRECATED;
 	// cached component rotation when material world-position-offset is used,
 	// as this will affect the direction of world-position-offset deformation (included in the HeightData below)
-	UE_DEPRECATED(5.5, "GenerationHash is now used") 
+	UE_DEPRECATED(all, "GenerationHash is now used") 
 	FQuat RotationForWPO_DEPRECATED;
 
 	// Variable used to detect when grass data needs to be regenerated:
@@ -448,7 +448,7 @@ class ULandscapeComponent : public UPrimitiveComponent
 	TArray<FLandscapeComponentMaterialOverride> OverrideMaterials_DEPRECATED;
 	PRAGMA_ENABLE_DEPRECATION_WARNINGS
 
-	UE_DEPRECATED(5.5, "MaterialInstance has been deprecated, use MaterialInstances instead.")
+	UE_DEPRECATED(all, "MaterialInstance has been deprecated, use MaterialInstances instead.")
 	UPROPERTY()
 	TObjectPtr<UMaterialInstanceConstant> MaterialInstance_DEPRECATED;
 #endif // WITH_EDITORONLY_DATA
@@ -495,7 +495,7 @@ class ULandscapeComponent : public UPrimitiveComponent
 	TArray<double> MipToMipMaxDeltas;
 
 #if WITH_EDITORONLY_DATA
-	UE_DEPRECATED(5.5, "CollisionComponent has been deprecated and will be removed in a future version")
+	UE_DEPRECATED(all, "CollisionComponent has been deprecated")
 	UPROPERTY()
 	TLazyObjectPtr<ULandscapeHeightfieldCollisionComponent> CollisionComponent_DEPRECATED;
 #endif // !WITH_EDITORONLY_DATA
@@ -619,10 +619,6 @@ public:
 	UPROPERTY(EditAnywhere, BlueprintReadOnly, Category=LandscapeComponent)
 	int32 LODBias;
 
-	UE_DEPRECATED(5.5, "StateId is unused and will be removed in a future version")
-	UPROPERTY()
-	FGuid StateId;
-
 	/** When the component has registered a heightmap texture edge fixup,
 	  * a reference to the edge fixup and the group are stored here for easy access
 	  *  (the fixup is also referenced via the heightmap AssetUserDatas,
@@ -634,7 +630,7 @@ public:
 
 #if WITH_EDITORONLY_DATA
 	/**	Legacy irrelevant lights */
-	UE_DEPRECATED(5.5, "IrrelevantLights is officially deprecated now and will be removed in a future version")
+	UE_DEPRECATED(all, "IrrelevantLights is officially deprecated now")
 	UPROPERTY()
 	TArray<FGuid> IrrelevantLights_DEPRECATED;
 
@@ -672,11 +668,11 @@ public:
 	UPROPERTY(Transient)
 	uint32 LastSavedPhysicalMaterialHash;
 
-	UE_DEPRECATED(5.5, "MobileMaterialInterface has been deprecated and will be removed in a future version")
+	UE_DEPRECATED(all, "MobileMaterialInterface has been deprecated and will be removed in a future version")
 	UPROPERTY(NonPIEDuplicateTransient)
 	TObjectPtr<UMaterialInterface> MobileMaterialInterface_DEPRECATED;
 
-    UE_DEPRECATED(5.5, "MobileCombinationMaterialInstance has been deprecated and will be removed in a future version")
+    UE_DEPRECATED(all, "MobileCombinationMaterialInstance has been deprecated and will be removed in a future version")
 	UPROPERTY(NonPIEDuplicateTransient)
 	TObjectPtr<UMaterialInstanceConstant> MobileCombinationMaterialInstance_DEPRECATED;
 #endif // WITH_EDITORONLY_DATA
@@ -1177,12 +1173,6 @@ public:
 	/** Update collision component physical materials from render task results. */
 	void UpdateCollisionPhysicalMaterialData(TArray<UPhysicalMaterial*> const& InPhysicalMaterials, TArray<uint8> const& InMaterialIds);
 
-	/**
-	 * Create weightmaps for this component for the layers specified in the WeightmapLayerAllocations array, works in the landscape current edit layer when InCanUseEditingWeightmap is true
-	 */
-	UE_DEPRECATED(5.5, "Use the new prototype of ReallocateWeightmaps : note : this is very internal stuff, normally, you shouldn't have to use this function at all")
-	LANDSCAPE_API void ReallocateWeightmaps(FLandscapeEditDataInterface* DataInterface = nullptr, bool InCanUseEditingWeightmap = true, bool InSaveToTransactionBuffer = true, bool InForceReallocate = false, ALandscapeProxy* InTargetProxy = nullptr, TArray<UTexture*>* OutNewCreatedTextures = nullptr);
-
 	/**
 	 * Create weightmaps for this component for the layers specified in the WeightmapLayerAllocations array, works in the specified edit layer
 	 */
diff --git a/Engine/Source/Runtime/Landscape/Classes/LandscapeInfo.h b/Engine/Source/Runtime/Landscape/Classes/LandscapeInfo.h
--- a/Engine/Source/Runtime/Landscape/Classes/LandscapeInfo.h
+++ b/Engine/Source/Runtime/Landscape/Classes/LandscapeInfo.h
@@ -95,20 +95,13 @@ struct FLandscapeInfoLayerSettings
 	}
 
 	LANDSCAPE_API FLandscapeInfoLayerSettings(ULandscapeLayerInfoObject* InLayerInfo, ALandscapeProxy* InProxy);
-
 	LANDSCAPE_API FLandscapeInfoLayerSettings(FName InPlaceholderLayerName, ALandscapeProxy* InProxy);
 
 	LANDSCAPE_API FName GetLayerName() const;
 
 #if WITH_EDITORONLY_DATA
-PRAGMA_DISABLE_DEPRECATION_WARNINGS	
-	UE_DEPRECATED(5.5, "This property has been deprecated, please use the GetTargetLayerSettings instead")
-	LANDSCAPE_API FLandscapeEditorLayerSettings& GetEditorSettings() const;
-PRAGMA_ENABLE_DEPRECATION_WARNINGS
-	
 	LANDSCAPE_API const FLandscapeTargetLayerSettings& GetTargetLayerSettings() const;
-
-#endif
+#endif // WITH_EDITORONLY_DATA
 };
 
 UCLASS(Transient)
@@ -233,14 +226,8 @@ public:
 	LANDSCAPE_API void UpdateAllAddCollisions();
 	LANDSCAPE_API void UpdateAddCollision(FIntPoint LandscapeKey);
 
-	UE_DEPRECATED(5.5, "This property has been deprecated, please use the GetTargetLayerSettings(ULandscapeLayerInfoObject* LayerInfo) instead")
-	LANDSCAPE_API FLandscapeTargetLayerSettings& GetLayerEditorSettings(ULandscapeLayerInfoObject* LayerInfo) const;
-	
 	LANDSCAPE_API const FLandscapeTargetLayerSettings& GetTargetLayerSettings(ULandscapeLayerInfoObject* LayerInfo) const;
 
-	UE_DEPRECATED(5.5, "This property has been deprecated, please use the CreateTargetLayerSettingsFor instead")
-	void CreateLayerEditorSettingsFor(ULandscapeLayerInfoObject* LayerInfo) {}
-	
 	LANDSCAPE_API void CreateTargetLayerSettingsFor(ULandscapeLayerInfoObject* LayerInfo);
 
 	LANDSCAPE_API ULandscapeLayerInfoObject* GetLayerInfoByName(FName LayerName, ALandscapeProxy* Owner = nullptr) const;
diff --git a/Engine/Source/Runtime/Landscape/Classes/LandscapeProxy.h b/Engine/Source/Runtime/Landscape/Classes/LandscapeProxy.h
--- a/Engine/Source/Runtime/Landscape/Classes/LandscapeProxy.h
+++ b/Engine/Source/Runtime/Landscape/Classes/LandscapeProxy.h
@@ -113,7 +113,7 @@ struct FLandscapeTargetLayerSettings
 };
 
 USTRUCT()
-struct UE_DEPRECATED(5.5, "FLandscapeEditorLayerSettings is deprecated; please use FLandscapeTargetLayerSettings instead")  FLandscapeEditorLayerSettings
+struct UE_DEPRECATED(all, "FLandscapeEditorLayerSettings is deprecated; please use FLandscapeTargetLayerSettings instead")  FLandscapeEditorLayerSettings
 {
 	GENERATED_USTRUCT_BODY()
 
@@ -498,11 +498,11 @@ public:
 	int32 MaxLODLevel;
 
 #if WITH_EDITORONLY_DATA
-	UE_DEPRECATED(5.5, "Unused property : will be removed in a future version")
+	UE_DEPRECATED(all, "Unused property, replaced by LOD0DistributionSetting and LODDistributionSetting")
 	UPROPERTY()
 	float LODDistanceFactor_DEPRECATED;
 
-	UE_DEPRECATED(5.5, "Unused property : will be removed in a future version")
+	UE_DEPRECATED(all, "Unused property, replaced by LOD0DistributionSetting and LODDistributionSetting")
 	UPROPERTY()
 	TEnumAsByte<ELandscapeLODFalloff::Type> LODFalloff_DEPRECATED;
 #endif // WITH_EDITORONLY_DATA
@@ -844,7 +844,7 @@ PRAGMA_ENABLE_DEPRECATION_WARNINGS
 	uint32 bBakeMaterialPositionOffsetIntoCollision:1;
 
 #if WITH_EDITORONLY_DATA
-	UE_DEPRECATED(5.5, "Unused property : will be removed in a future version")
+	UE_DEPRECATED(all, "TargetLayers will be used instead")
 	UPROPERTY()
 	TArray<TObjectPtr<ULandscapeLayerInfoObject>> EditorCachedLayerInfos_DEPRECATED;
 
@@ -856,7 +856,7 @@ PRAGMA_ENABLE_DEPRECATION_WARNINGS
 	FGuid ReimportDestinationLayerGuid;
 
 PRAGMA_DISABLE_DEPRECATION_WARNINGS
-	UE_DEPRECATED(5.5, "TargetLayers will be used instead")
+	UE_DEPRECATED(all, "TargetLayers will be used instead")
 	UPROPERTY()
 	TArray<FLandscapeEditorLayerSettings> EditorLayerSettings_DEPRECATED;
 PRAGMA_ENABLE_DEPRECATION_WARNINGS
@@ -1076,9 +1076,6 @@ PRAGMA_ENABLE_DEPRECATION_WARNINGS
 	int32 NumNaniteRequiredComponents() const { return FMath::DivideAndRoundUp(LandscapeComponents.Num(), NaniteMaxComponents); }
 
 	TArray<ULandscapeComponent*> GatherSourceComponentsForNaniteComponent(int32 NaniteComponentIndex) const;
-
-	UE_DEPRECATED(5.5, "This version of HasLayer is deprecated.  Use HasTargetLayer(ULandscapeLayerInfoObject* LayerInfoObject) instead.")
-	bool HasLayer(ULandscapeLayerInfoObject* LayerInfoObject) const;
 #endif	//WITH_EDITOR
 
 	bool AreNaniteComponentsValid(const FGuid& InProxyContentId) const;
@@ -1390,11 +1387,6 @@ PRAGMA_ENABLE_DEPRECATION_WARNINGS
 	LANDSCAPE_API static ULandscapeMaterialInstanceConstant* GetLayerThumbnailMIC(UMaterialInterface* LandscapeMaterial, FName LayerName, UTexture2D* ThumbnailWeightmap, UTexture2D* ThumbnailHeightmap, ALandscapeProxy* Proxy);
 
 	/** Import the given Height/Weight data into this landscape */
-	UE_DEPRECATED(5.5, "Use the new method prototype which takes a TArrayView parameter for import layers")
-	LANDSCAPE_API void Import(const FGuid& InGuid, int32 InMinX, int32 InMinY, int32 InMaxX, int32 InMaxY, int32 InNumSubsections, int32 InSubsectionSizeQuads, const TMap<FGuid, TArray<uint16>>& InImportHeightData,
-		const TCHAR* const InHeightmapFileName, const TMap<FGuid, TArray<FLandscapeImportLayerInfo>>& InImportMaterialLayerInfos, ELandscapeImportAlphamapType InImportMaterialLayerType, 
-		const TArray<struct FLandscapeLayer>* InImportLayers = nullptr);
-
 	LANDSCAPE_API void Import(const FGuid& InGuid, int32 InMinX, int32 InMinY, int32 InMaxX, int32 InMaxY, int32 InNumSubsections, int32 InSubsectionSizeQuads, const TMap<FGuid, TArray<uint16>>& InImportHeightData,
 		const TCHAR* const InHeightmapFileName, const TMap<FGuid, TArray<FLandscapeImportLayerInfo>>& InImportMaterialLayerInfos, ELandscapeImportAlphamapType InImportMaterialLayerType,
 		const TArrayView<const struct FLandscapeLayer>& InImportLayers);
diff --git a/Engine/Source/Runtime/Landscape/Private/Landscape.cpp b/Engine/Source/Runtime/Landscape/Private/Landscape.cpp
--- a/Engine/Source/Runtime/Landscape/Private/Landscape.cpp
+++ b/Engine/Source/Runtime/Landscape/Private/Landscape.cpp
@@ -4015,14 +4015,6 @@ FName FLandscapeInfoLayerSettings::GetLayerName() const
 	return LayerName;
 }
 
-PRAGMA_DISABLE_DEPRECATION_WARNINGS
-FLandscapeEditorLayerSettings& FLandscapeInfoLayerSettings::GetEditorSettings() const
-{
-	static FLandscapeEditorLayerSettings DeprecatedSettings;
-	return DeprecatedSettings;
-}
-PRAGMA_ENABLE_DEPRECATION_WARNINGS
-
 const FLandscapeTargetLayerSettings& FLandscapeInfoLayerSettings::GetTargetLayerSettings() const
 {
 	check(LayerInfoObj);
@@ -4031,12 +4023,6 @@ const FLandscapeTargetLayerSettings& FLandscapeInfoLayerSettings::GetTargetLayer
 	return LandscapeInfo->GetTargetLayerSettings(LayerInfoObj);
 }
 
-FLandscapeTargetLayerSettings& ULandscapeInfo::GetLayerEditorSettings(ULandscapeLayerInfoObject* LayerInfo) const
-{
-	static FLandscapeTargetLayerSettings DeprecatedSettings;
-	return DeprecatedSettings;
-}
-
 const FLandscapeTargetLayerSettings& ULandscapeInfo::GetTargetLayerSettings(ULandscapeLayerInfoObject* LayerInfo) const
 {
 	ALandscapeProxy* Proxy = GetLandscapeProxy();
@@ -6911,11 +6897,6 @@ void ALandscapeProxy::EnableNaniteComponents(bool bInNaniteActive)
 }
 
 #if WITH_EDITOR
-bool ALandscapeProxy::HasLayer(ULandscapeLayerInfoObject* LayerInfoObject) const
-{
-	return TargetLayers.FindKey(FLandscapeTargetLayerSettings(LayerInfoObject)) == nullptr;
-}
-
 bool ALandscapeProxy::RemoveTargetLayer(const FName& Name,  bool bPostEditChange)
 {
 	Modify();
diff --git a/Engine/Source/Runtime/Landscape/Private/LandscapeEdit.cpp b/Engine/Source/Runtime/Landscape/Private/LandscapeEdit.cpp
--- a/Engine/Source/Runtime/Landscape/Private/LandscapeEdit.cpp
+++ b/Engine/Source/Runtime/Landscape/Private/LandscapeEdit.cpp
@@ -3259,17 +3259,6 @@ ULandscapeLayerInfoObject* ALandscapeProxy::CreateLayerInfo(const TCHAR* InLayer
 #define HEIGHTDATA(X,Y) (HeightData.Num() == 0 ? LandscapeDataAccess::GetTexHeight(0.0f) : HeightData[ FMath::Clamp<int32>(Y,0,VertsY) * VertsX + FMath::Clamp<int32>(X,0,VertsX) ])
 ENGINE_API extern bool GDisableAutomaticTextureMaterialUpdateDependencies;
 
-void ALandscapeProxy::Import(const FGuid& InGuid, int32 InMinX, int32 InMinY, int32 InMaxX, int32 InMaxY, int32 InNumSubsections, int32 InSubsectionSizeQuads, const TMap<FGuid, TArray<uint16>>& InImportHeightData, 
-										   const TCHAR* const InHeightmapFileName, const TMap<FGuid, TArray<FLandscapeImportLayerInfo>>& InImportMaterialLayerInfos, ELandscapeImportAlphamapType InImportMaterialLayerType, const TArray<FLandscapeLayer>* InImportLayers)
-{
-	TArrayView<const FLandscapeLayer> ImportLayers;
-	if (InImportLayers != nullptr)
-	{
-		ImportLayers = MakeArrayView(static_cast<const FLandscapeLayer*>(InImportLayers->GetData()), InImportLayers->Num());
-	}
-	Import(InGuid, InMinX, InMinY, InMaxX, InMaxY, InNumSubsections, InSubsectionSizeQuads, InImportHeightData, InHeightmapFileName, InImportMaterialLayerInfos, InImportMaterialLayerType, ImportLayers);
-}
-
 void ALandscapeProxy::Import(const FGuid& InGuid, int32 InMinX, int32 InMinY, int32 InMaxX, int32 InMaxY, int32 InNumSubsections, int32 InSubsectionSizeQuads, const TMap<FGuid, TArray<uint16>>& InImportHeightData,
 	const TCHAR* const InHeightmapFileName, const TMap<FGuid, TArray<FLandscapeImportLayerInfo>>& InImportMaterialLayerInfos, ELandscapeImportAlphamapType InImportMaterialLayerType, const TArrayView<const FLandscapeLayer>& InImportLayers)
 {
@@ -7279,17 +7268,6 @@ void ULandscapeInfo::ClearSelectedRegion(bool bIsComponentwise /*= true*/)
 	LandscapeModule.GetLandscapeEditorServices()->RefreshDetailPanel();
 }
 
-void ULandscapeComponent::ReallocateWeightmaps(FLandscapeEditDataInterface* DataInterface, bool InCanUseEditingWeightmap, bool InSaveToTransactionBuffer, bool InForceReallocate, 
-	ALandscapeProxy* InTargetProxy, TArray<UTexture*>* OutNewCreatedTextures)
-{
-	FGuid TargetLayerGuid = InCanUseEditingWeightmap ? GetEditingLayerGUID() : FGuid();
-	TArray<UTexture*> CreatedTextures = ReallocateWeightmaps(DataInterface, TargetLayerGuid, InSaveToTransactionBuffer, InForceReallocate, InTargetProxy, /*InRestrictSharingToComponents = */nullptr);
-	if (OutNewCreatedTextures != nullptr)
-	{
-		*OutNewCreatedTextures = CreatedTextures;
-	}
-}
-
 TArray<UTexture*> ULandscapeComponent::ReallocateWeightmaps(FLandscapeEditDataInterface* DataInterface, const FGuid& InEditLayerGuid, bool bInSaveToTransactionBuffer, bool bInForceReallocate, 
 	ALandscapeProxy* InTargetProxy, TSet<ULandscapeComponent*>* InRestrictSharingToComponents)
 {
diff --git a/Engine/Source/Runtime/Landscape/Private/LandscapeEditLayers.cpp b/Engine/Source/Runtime/Landscape/Private/LandscapeEditLayers.cpp
--- a/Engine/Source/Runtime/Landscape/Private/LandscapeEditLayers.cpp
+++ b/Engine/Source/Runtime/Landscape/Private/LandscapeEditLayers.cpp
@@ -12520,24 +12520,6 @@ ULandscapeEditLayerBase* ALandscape::GetEditLayerInternal(int32 InLayerIndex)
 	return nullptr;
 }
 
-// Deprecated
-const FLandscapeLayer* ALandscape::GetLayer(int32 InLayerIndex) const
-{
-	return GetLayerConst(InLayerIndex);
-}
-
-// Deprecated
-const FLandscapeLayer* ALandscape::GetLayer(const FGuid& InLayerGuid) const
-{
-	return GetLayerConst(InLayerGuid);
-}
-
-// Deprecated
-const FLandscapeLayer* ALandscape::GetLayer(const FName& InLayerName) const
-{
-	return GetLayerConst(InLayerName);
-}
-
 TArrayView<const FLandscapeLayer> ALandscape::GetLayersConst() const
 { 
 	return MakeArrayView(LandscapeEditLayers); 
@@ -12689,15 +12671,6 @@ int32 ALandscape::GetLayerIndex(FName InLayerName) const
 	return LandscapeEditLayers.IndexOfByPredicate([InLayerName](const FLandscapeLayer& Layer) { return Layer.EditLayer->GetName() == InLayerName; });
 }
 
-// Deprecated
-void ALandscape::ForEachLayer(TFunctionRef<void(struct FLandscapeLayer&)> Fn)
-{
-	for (FLandscapeLayer& Layer : LandscapeEditLayers)
-	{
-		Fn(Layer);
-	}
-}
-
 void ALandscape::ForEachLayerConst(TFunctionRef<bool(const FLandscapeLayer&)> Fn)
 {
 	for (FLandscapeLayer& Layer : LandscapeEditLayers)
@@ -13064,23 +13037,6 @@ void ALandscape::ShowAllLayers()
 	}
 }
 
-// Deprecated
-void ALandscape::SetLandscapeSplinesReservedLayer(int32 InLayerIndex)
-{
-}
-
-// Deprecated
-const FLandscapeLayer* ALandscape::GetLandscapeSplinesReservedLayer() const
-{
-	return FindLayerOfTypeConst(ULandscapeEditLayerSplines::StaticClass());
-}
-
-// Deprecated
-FLandscapeLayer* ALandscape::GetLandscapeSplinesReservedLayer()
-{
-	return const_cast<FLandscapeLayer*>(FindLayerOfTypeConst(ULandscapeEditLayerSplines::StaticClass()));
-}
-
 LANDSCAPE_API extern bool GDisableUpdateLandscapeMaterialInstances;
 
 uint32 ULandscapeComponent::ComputeLayerHash(bool InReturnEditingHash) const
@@ -13231,12 +13187,6 @@ FScopedSetLandscapeEditingLayer::~FScopedSetLandscapeEditingLayer()
 	}
 }
 
-// Deprecated
-bool ALandscape::IsEditingLayerReservedForSplines() const
-{
-	return false;
-}
-
 void ALandscape::SetEditingLayer(const FGuid& InLayerGuid)
 {
 	ensure(CanHaveLayersContent());
@@ -13365,29 +13315,6 @@ int32 ALandscape::CreateLayer(FName InName, const TSubclassOf<ULandscapeEditLaye
 	return LayerIndex;
 }
 
-// Deprecated : use DuplicateLayerAndMoveBrushes 
-int32 ALandscape::CreateLayerFrom(const FLandscapeLayer& InLayer)
-{
-	if (IsMaxLayersReached() || !CanHaveLayersContent())
-	{
-		return INDEX_NONE;
-	}
-
-	// Cannot create a layer whose Guid is already taken : 
-	if (GetLayerConst(InLayer.EditLayer->GetGuid()) != nullptr)
-	{
-		return INDEX_NONE;
-	}
-
-	check(InLayer.EditLayer != nullptr);
-
-	int32 LayerIndex = LandscapeEditLayers.Add(InLayer);
-	FLandscapeLayer& NewLayer = LandscapeEditLayers[LayerIndex];
-	NewLayer.EditLayer = DuplicateObject<ULandscapeEditLayerBase>(InLayer.EditLayer, this, MakeUniqueObjectName(this, InLayer.EditLayer->GetClass()));
-	OnLayerCreatedInternal(NewLayer.EditLayer);
-	return LayerIndex;
-}
-
 void ALandscape::OnLayerCreatedInternal(ULandscapeEditLayerBase* EditLayer)
 {
 	// TODO: Might not be necessary eventually, if EditLayer has ability to trigger landscape updates
diff --git a/Engine/Source/Runtime/Landscape/Private/LandscapeSubsystem.cpp b/Engine/Source/Runtime/Landscape/Private/LandscapeSubsystem.cpp
--- a/Engine/Source/Runtime/Landscape/Private/LandscapeSubsystem.cpp
+++ b/Engine/Source/Runtime/Landscape/Private/LandscapeSubsystem.cpp
@@ -1090,31 +1090,6 @@ TSet<UPackage*> ULandscapeSubsystem::GetDirtyLandscapeProxyPackages() const
 	return DirtyPackages;
 }
 
-// Deprecated :
-TArray<ALandscapeProxy*> ULandscapeSubsystem::GetOutdatedProxies(UE::Landscape::EOutdatedDataFlags InMatchingOutdatedDataFlags, bool bInMustMatchAllFlags) const
-{
-	TRACE_CPUPROFILER_EVENT_SCOPE(ULandscapeSubsystem::GetOutdatedProxies);
-
-	UWorld* World = GetWorld();
-	if (!World || World->IsGameWorld())
-	{
-		return {};
-	}
-
-	TArray<ALandscapeProxy*> FinalProxiesToBuild;
-	Algo::TransformIf(Proxies, FinalProxiesToBuild, 
-		[InMatchingOutdatedDataFlags, bInMustMatchAllFlags](const TObjectPtr<ALandscapeProxy>& InProxyPtr)
-		{ 
-			UE::Landscape::EOutdatedDataFlags ProxyOutdatedDataFlags = InProxyPtr->GetOutdatedDataFlags();
-			return bInMustMatchAllFlags
-				? EnumHasAllFlags(ProxyOutdatedDataFlags, InMatchingOutdatedDataFlags)
-				: EnumHasAnyFlags(ProxyOutdatedDataFlags, InMatchingOutdatedDataFlags);
-		}, 
-		[](const TWeakObjectPtr<ALandscapeProxy>& InProxyPtr) { return InProxyPtr.Get(); });
-
-	return FinalProxiesToBuild;
-}
-
 TArray<TTuple<ALandscapeProxy*, UE::Landscape::EOutdatedDataFlags>> ULandscapeSubsystem::GetOutdatedProxyDetails(UE::Landscape::EOutdatedDataFlags InMatchingOutdatedDataFlags, bool bInMustMatchAllFlags) const
 {
 	TRACE_CPUPROFILER_EVENT_SCOPE(ULandscapeSubsystem::GetOutdatedProxyDetails);
@@ -1377,11 +1352,6 @@ ALandscapeProxy* ULandscapeSubsystem::FindOrAddLandscapeProxy(ULandscapeInfo* La
 	return FLandscapeConfigHelper::FindOrAddLandscapeStreamingProxy(LandscapeInfo, SectionBase);
 }
 
-void ULandscapeSubsystem::DisplayMessages(FCanvas* Canvas, float& XPos, float& YPos)
-{
-	
-}
-
 bool ULandscapeSubsystem::GetActionableMessage(FActionableMessage& OutActionableMessage)
 {
 	using namespace UE::Landscape;
diff --git a/Engine/Source/Runtime/Landscape/Public/LandscapeNotification.h b/Engine/Source/Runtime/Landscape/Public/LandscapeNotification.h
--- a/Engine/Source/Runtime/Landscape/Public/LandscapeNotification.h
+++ b/Engine/Source/Runtime/Landscape/Public/LandscapeNotification.h
@@ -17,14 +17,6 @@ enum class ELandscapeNotificationType
 
 struct FLandscapeNotification
 {
-	enum class UE_DEPRECATED(5.5, "Use ELandscapeNotificationType") EType : uint8
-	{
-		ShadingModelInvalid = 0,
-		LandscapeTextureResourcesNotReady,
-		LandscapeBrushResourcesNotReady,
-		GrassRendering,
-	};
-
 	using FConditionCallback = TFunction<bool(void)>;
 	using FUpdateTextCallback = TFunction<void(FText&)>;
 
diff --git a/Engine/Source/Runtime/Landscape/Public/LandscapeSubsystem.h b/Engine/Source/Runtime/Landscape/Public/LandscapeSubsystem.h
--- a/Engine/Source/Runtime/Landscape/Public/LandscapeSubsystem.h
+++ b/Engine/Source/Runtime/Landscape/Public/LandscapeSubsystem.h
@@ -164,17 +164,12 @@ public:
 	 */
 	LANDSCAPE_API void BuildNanite(UE::Landscape::EBuildFlags InBuildFlags, TArrayView<ALandscapeProxy*> InProxiesToBuild = TArrayView<ALandscapeProxy*>());
 	
-	UE_DEPRECATED(5.5, "Use GetOutdatedProxyDetails")
-	LANDSCAPE_API TArray<ALandscapeProxy*> GetOutdatedProxies(UE::Landscape::EOutdatedDataFlags InMatchingOutdatedDataFlags, bool bInMustMatchAllFlags) const;
 	LANDSCAPE_API TArray<TTuple<ALandscapeProxy*, UE::Landscape::EOutdatedDataFlags>> GetOutdatedProxyDetails(UE::Landscape::EOutdatedDataFlags InMatchingOutdatedDataFlags, bool bInMustMatchAllFlags) const;
 	
 	LANDSCAPE_API bool IsGridBased() const;
 	LANDSCAPE_API void ChangeGridSize(ULandscapeInfo* LandscapeInfo, uint32 NewGridSizeInComponents);
 	LANDSCAPE_API ALandscapeProxy* FindOrAddLandscapeProxy(ULandscapeInfo* LandscapeInfo, const FIntPoint& SectionBase);
 
-	UE_DEPRECATED(5.5, "DisplayMessages is now deprecated.")
-	LANDSCAPE_API void DisplayMessages(class FCanvas* Canvas, float& XPos, float& YPos);
-
 	LANDSCAPE_API bool GetActionableMessage(FActionableMessage& OutActionableMessage);
 	UE_DEPRECATED(5.6, "Use the function with the EBuildFlags param")
 	LANDSCAPE_API void MarkModifiedLandscapesAsDirty();
