Landscape Editor -  Fix Landscape components not updating nanite materials on Undo/Redo paint operations

diff --git a/Engine/Source/Runtime/Landscape/Private/LandscapeEdit.cpp b/Engine/Source/Runtime/Landscape/Private/LandscapeEdit.cpp
--- a/Engine/Source/Runtime/Landscape/Private/LandscapeEdit.cpp
+++ b/Engine/Source/Runtime/Landscape/Private/LandscapeEdit.cpp
@@ -812,11 +812,6 @@ void ULandscapeComponent::UpdateMaterialInstances()
 	// Recreate the render state for this component, needed to update the static drawlist which has cached the MaterialRenderProxies
 	// Must be after the FMaterialUpdateContext is destroyed
 	RecreateRenderStateContext.Reset();
-
-	// If this component is within a Nanite Enabled Landscape we have to update the materials on the
-	// ULandscapeNaniteComponents within the Parent Actor. 
-	//
-	// todo [don.boogert] reinstate update of landscape nanite materials GetLandscapeProxy()->UpdateNaniteMaterials(); 
 }
 
 void ULandscapeComponent::UpdateMaterialInstances(FMaterialUpdateContext& InOutMaterialContext, TArray<FComponentRecreateRenderStateContext>& InOutRecreateRenderStateContext)
@@ -1076,6 +1071,12 @@ void ULandscapeComponent::PostEditUndo()
 			}
 		}
 
+		// After updating the component materials, update the nanite materials on each proxy
+		for (ALandscapeProxy* ComponentProxy : Proxies)
+		{
+			ComponentProxy->UpdateNaniteMaterials();
+		}
+
 		// Clear out the temporarily recorded undo info
 		UndoRedoModifiedComponentCount = 0;
 		UndoRedoModifiedComponents.Empty();
