[Lumen card capture] Disable capturing landscape cards using Nanite by default for now. There is a report of increased capturing cost on PC which needs to be investigated

diff --git a/Engine/Source/Runtime/Landscape/Private/LandscapeRender.cpp b/Engine/Source/Runtime/Landscape/Private/LandscapeRender.cpp
--- a/Engine/Source/Runtime/Landscape/Private/LandscapeRender.cpp
+++ b/Engine/Source/Runtime/Landscape/Private/LandscapeRender.cpp
@@ -196,7 +196,7 @@ static FAutoConsoleVariableRef CVarDisableLandscapeNaniteGI(
 
 static TAutoConsoleVariable<int32> CVarLumenSceneSurfaceCacheCaptureNaniteLandscape(
 	TEXT("r.LumenScene.SurfaceCache.CaptureNaniteLandscapes"),
-	1,
+	0,
 	TEXT("Whether to use Nanite landscape meshes when capturing landscape Lumen cards."),
 	FConsoleVariableDelegate::CreateStatic(&InvalidateRenderStateForLandscapeNaniteComponents),
 	ECVF_RenderThreadSafe | ECVF_Scalability
