Landscape Editor - Fix to set selected edit layer index on landscape actor post duplicate

diff --git a/Engine/Source/Runtime/Landscape/Private/LandscapeEdit.cpp b/Engine/Source/Runtime/Landscape/Private/LandscapeEdit.cpp
--- a/Engine/Source/Runtime/Landscape/Private/LandscapeEdit.cpp
+++ b/Engine/Source/Runtime/Landscape/Private/LandscapeEdit.cpp
@@ -5467,6 +5467,9 @@ void ALandscape::PostDuplicate(bool bDuplicateForPIE)
 		// Some edit layers could be affected by BP brushes, which might need to be updated when the landscape is transformed :
 		RequestLayersContentUpdate(ELandscapeLayerUpdateMode::Update_All);
 	}
+
+	// Ensure the first edit layer is selected after duplicating
+	SetSelectedEditLayerIndex(0);
 }
 #endif	//WITH_EDITOR
 
