Landscape Editor - Fix to ensure selected edit layer index only gets updated for edit-layer based landscapes on duplication

diff --git a/Engine/Source/Runtime/Landscape/Private/LandscapeEdit.cpp b/Engine/Source/Runtime/Landscape/Private/LandscapeEdit.cpp
--- a/Engine/Source/Runtime/Landscape/Private/LandscapeEdit.cpp
+++ b/Engine/Source/Runtime/Landscape/Private/LandscapeEdit.cpp
@@ -5468,8 +5468,11 @@ void ALandscape::PostDuplicate(bool bDuplicateForPIE)
 		RequestLayersContentUpdate(ELandscapeLayerUpdateMode::Update_All);
 	}
 
-	// Ensure the first edit layer is selected after duplicating
-	SetSelectedEditLayerIndex(0);
+	// If edit layers are supported, ensure the first edit layer is selected after duplicating
+	if (CanHaveLayersContent())
+	{
+		SetSelectedEditLayerIndex(0);
+	}
 }
 #endif	//WITH_EDITOR
 
