[Landscape] Only use LandscapeComponents/CollisionComponents bounds in new GetProxyBounds() which is now used in GetLoadedBounds/GetCompleteBounds
[PCG] Make use of new GetProxyBounds
- Fix issues where grass components with invalid bounds would break the landscape bounds

#rnx

diff --git a/Engine/Plugins/PCG/Source/PCG/Private/Helpers/PCGHelpers.cpp b/Engine/Plugins/PCG/Source/PCG/Private/Helpers/PCGHelpers.cpp
--- a/Engine/Plugins/PCG/Source/PCG/Private/Helpers/PCGHelpers.cpp
+++ b/Engine/Plugins/PCG/Source/PCG/Private/Helpers/PCGHelpers.cpp
@@ -14,6 +14,7 @@
 
 #include "Components/BillboardComponent.h"
 #include "Landscape.h"
+#include "LandscapeProxy.h"
 #include "Algo/AnyOf.h"
 #include "UObject/UObjectIterator.h"
 
@@ -131,7 +132,7 @@ namespace PCGHelpers
 			}
 		}
 		// TODO: verify this works as expected in non-editor builds
-		else if (const ALandscapeProxy* LandscapeActor = Cast<const ALandscape>(Actor))
+		else if (const ALandscapeProxy* LandscapeActor = Cast<const ALandscapeProxy>(Actor))
 		{
 			Bounds = GetLandscapeBounds(LandscapeActor);
 		}
@@ -286,7 +287,7 @@ namespace PCGHelpers
 		}
 		else
 		{
-			return GetActorBounds(InLandscape);
+			return InLandscape->GetProxyBounds();
 		}
 	}
 
diff --git a/Engine/Source/Runtime/Landscape/Classes/LandscapeProxy.h b/Engine/Source/Runtime/Landscape/Classes/LandscapeProxy.h
--- a/Engine/Source/Runtime/Landscape/Classes/LandscapeProxy.h
+++ b/Engine/Source/Runtime/Landscape/Classes/LandscapeProxy.h
@@ -1227,6 +1227,9 @@ PRAGMA_ENABLE_DEPRECATION_WARNINGS
 	/** Get the LandcapeActor-to-world transform with respect to landscape section offset*/
 	virtual LANDSCAPE_API FTransform LandscapeActorToWorld() const override;
 
+	/** Differs from GetActorBounds as it only considers LandscapeComponents/CollisionComponents bounds */
+	LANDSCAPE_API FBox GetProxyBounds() const;
+
 	/**
 	* Output a landscape heightmap to a render target
 	* @param InRenderTarget - Valid render target with a format of RTF_RGBA16f, RTF_RGBA32f or RTF_RGBA8
diff --git a/Engine/Source/Runtime/Landscape/Private/Landscape.cpp b/Engine/Source/Runtime/Landscape/Private/Landscape.cpp
--- a/Engine/Source/Runtime/Landscape/Private/Landscape.cpp
+++ b/Engine/Source/Runtime/Landscape/Private/Landscape.cpp
@@ -6383,23 +6383,45 @@ void ULandscapeInfo::UnregisterActorComponent(ULandscapeComponent* Component)
 
 namespace LandscapeInfoBoundsHelper
 {
-	void AccumulateBounds(ALandscapeProxy* Proxy, FBox& Bounds)
+	template<class ComponentClass>
+	void AccumulateComponentBounds(TConstArrayView<TObjectPtr<ComponentClass>> InComponents, FBox& InOutBounds)
 	{
-		const bool bOnlyCollidingComponents = false;
-		const bool bIncludeChildActors = false;
-		FVector Origin;
-		FVector BoxExtents;
-
-		Proxy->GetActorBounds(bOnlyCollidingComponents, Origin, BoxExtents, bIncludeChildActors);
+		for (TObjectPtr<ComponentClass> Component : InComponents)
+		{
+			if ((Component != nullptr) && Component->IsRegistered())
+			{
+				// Reject invalid bounds
+				const FBox& ComponentBox = Component->Bounds.GetBox();
+				if (ComponentBox.IsValid && ComponentBox.GetSize() != FVector::Zero())
+				{
+					InOutBounds += ComponentBox;
+				}
+			}
+		};
+	}
 
-		// Reject invalid bounds
-		if (BoxExtents != FVector::Zero())
+	void AccumulateBounds(const ALandscapeProxy* InLandscapeProxy, FBox& InOutBounds)
+	{
+		FBox LandscapeProxyBounds = InLandscapeProxy->GetProxyBounds();
+		if (LandscapeProxyBounds.IsValid && LandscapeProxyBounds.GetSize() != FVector::Zero())
 		{
-			Bounds += FBox::BuildAABB(Origin, BoxExtents);
+			InOutBounds += LandscapeProxyBounds;
 		}
 	}
 }
 
+FBox ALandscapeProxy::GetProxyBounds() const
+{
+	FBox Bounds(EForceInit::ForceInit);
+	
+	// Only accumulate bounds from landscape components, so that bounds don't get artificially inflated by other components like spawned HISM (grass) components.
+	//  We include collision components, so that if this ever gets called on the server, where landscape components are missing, this is consistent
+	LandscapeInfoBoundsHelper::AccumulateComponentBounds(MakeConstArrayView(LandscapeComponents), Bounds);
+	LandscapeInfoBoundsHelper::AccumulateComponentBounds(MakeConstArrayView(CollisionComponents), Bounds);
+
+	return Bounds;
+}
+
 FBox ULandscapeInfo::GetLoadedBounds() const
 {
 	FBox Bounds(EForceInit::ForceInit);
