Removed useless LANDSCAPE_API and FORCEINLINE

#rnx

diff --git a/Engine/Source/Runtime/Landscape/Public/LandscapeEdit.h b/Engine/Source/Runtime/Landscape/Public/LandscapeEdit.h
--- a/Engine/Source/Runtime/Landscape/Public/LandscapeEdit.h
+++ b/Engine/Source/Runtime/Landscape/Public/LandscapeEdit.h
@@ -310,11 +310,11 @@ private:
 	bool bUseSharedLandscapeEditLayer;
 	FGuid LocalEditLayerGUID;
 
-	LANDSCAPE_API void FillLayer(ULandscapeLayerInfoObject* LayerInfo, bool bEmptyLayersOnly);
+	void FillLayer(ULandscapeLayerInfoObject* LayerInfo, bool bEmptyLayersOnly);
 
 	// Only for Missing Data interpolation... only internal usage
 	template<typename TData, typename TStoreData, typename FType>
-	FORCEINLINE void CalcMissingValues(const int32 X1, const int32 X2, const int32 Y1, const int32 Y2,
+	void CalcMissingValues(const int32 X1, const int32 X2, const int32 Y1, const int32 Y2,
 		const int32 ComponentIndexX1, const int32 ComponentIndexX2, const int32 ComponentIndexY1, const int32 ComponentIndexY2,
 		const int32 ComponentSizeX, const int32 ComponentSizeY, TData* CornerValues,
 		TArray<bool>& NoBorderY1, TArray<bool>& NoBorderY2, TArray<bool>& ComponentDataExist, TStoreData& StoreData);
@@ -323,7 +323,7 @@ private:
 	template<typename TStoreData, typename TGetHeightMapDataFunction>
 	void GetHeightDataInternal(int32& ValidX1, int32& ValidY1, int32& ValidX2, int32& ValidY2, TStoreData& StoreData, TGetHeightMapDataFunction GetHeightMapDataFunction);
 
-	LANDSCAPE_API FORCEINLINE FColor& GetHeightMapColor(const ULandscapeComponent* Component, int32 TexU, int32 TexV, FColor* TextureData);
+	FORCEINLINE FColor& GetHeightMapColor(const ULandscapeComponent* Component, int32 TexU, int32 TexV, FColor* TextureData);
 
 	// test if layer is allowed for a given texel
 	inline bool IsLayerAllowed(const ULandscapeLayerInfoObject* LayerInfo,
@@ -334,7 +334,7 @@ private:
 	inline TMap<const ULandscapeLayerInfoObject*, uint32> CountWeightBlendedLayerInfluence(int32 ComponentIndexX, int32 ComponentIndexY, TOptional<TArrayView<const uint8* const>> LayerDataPtrs);
 
 	// chooses a replacement layer to use when erasing from 100% influence on a texel
-	LANDSCAPE_API const ULandscapeLayerInfoObject* ChooseReplacementLayer(const ULandscapeLayerInfoObject* LayerInfo, int32 ComponentIndexX, int32 SubIndexX, int32 SubX, int32 ComponentIndexY, int32 SubIndexY, int32 SubY, TMap<FIntPoint, TMap<const ULandscapeLayerInfoObject*, uint32>>& LayerInfluenceCache, TArrayView<const uint8* const> LayerDataPtrs);
+	const ULandscapeLayerInfoObject* ChooseReplacementLayer(const ULandscapeLayerInfoObject* LayerInfo, int32 ComponentIndexX, int32 SubIndexX, int32 SubX, int32 ComponentIndexY, int32 SubIndexY, int32 SubY, TMap<FIntPoint, TMap<const ULandscapeLayerInfoObject*, uint32>>& LayerInfluenceCache, TArrayView<const uint8* const> LayerDataPtrs);
 };
 
 struct FLandscapeDoNotDirtyScope
