Fix assert in UpdatePhysicalMaterialTasks

diff --git a/Engine/Source/Runtime/Landscape/Classes/Landscape.h b/Engine/Source/Runtime/Landscape/Classes/Landscape.h
--- a/Engine/Source/Runtime/Landscape/Classes/Landscape.h
+++ b/Engine/Source/Runtime/Landscape/Classes/Landscape.h
@@ -515,6 +515,7 @@ PRAGMA_ENABLE_DEPRECATION_WARNINGS
 	LANDSCAPE_API void SetUseGeneratedLandscapeSplineMeshesActors(bool bInEnabled);
 	LANDSCAPE_API bool GetUseGeneratedLandscapeSplineMeshesActors() const;
 	LANDSCAPE_API bool PrepareTextureResources(bool bInWaitForStreaming);
+	bool PrepareTextureResourcesLimited(bool bInWaitForStreaming);
 
 	bool GetVisibilityLayerAllocationIndex() const { return 0; }
 	
@@ -786,6 +787,7 @@ private:
 	class FLandscapeTexture2DResource* WeightmapScratchPackLayerTextureResource;
 
 	uint32 LastFlushedLayerUpdateFrame = 0;
+	uint32 LastPrepareTextureResourcesCalled = 0;
 
 	// Counter to detect re-entrance
 	uint32 InLayerUpdateCount = 0;
diff --git a/Engine/Source/Runtime/Landscape/Private/Landscape.cpp b/Engine/Source/Runtime/Landscape/Private/Landscape.cpp
--- a/Engine/Source/Runtime/Landscape/Private/Landscape.cpp
+++ b/Engine/Source/Runtime/Landscape/Private/Landscape.cpp
@@ -6875,6 +6875,12 @@ void ALandscapeProxy::UpdatePhysicalMaterialTasksStatus(TSet<ULandscapeComponent
 
 bool ALandscapeProxy::UpdatePhysicalMaterialTasks(bool bInShouldMarkDirty)
 {
+	if (ALandscape* Landscape = GetLandscapeActor())
+	{
+		// Physical material depends on the weightmaps, so make sure they are controlled by the texture streaming manager.
+		Landscape->PrepareTextureResourcesLimited(false);
+	}
+
 	TSet<ULandscapeComponent*> OutdatedComponents;
 	int32 PendingComponentsToBeSaved = 0;
 	UpdatePhysicalMaterialTasksStatus(&OutdatedComponents, &PendingComponentsToBeSaved);
diff --git a/Engine/Source/Runtime/Landscape/Private/LandscapeEditLayers.cpp b/Engine/Source/Runtime/Landscape/Private/LandscapeEditLayers.cpp
--- a/Engine/Source/Runtime/Landscape/Private/LandscapeEditLayers.cpp
+++ b/Engine/Source/Runtime/Landscape/Private/LandscapeEditLayers.cpp
@@ -4004,6 +4004,17 @@ void ALandscape::PrintLayersDebugTextureResource(const FString& InContext, FText
 	}
 }
 
+bool ALandscape::PrepareTextureResourcesLimited(bool bInWaitForStreaming)
+{
+	// Rate-limit to once per frame
+	if (GFrameNumber != LastPrepareTextureResourcesCalled)
+	{
+		LastPrepareTextureResourcesCalled = GFrameNumber;
+		return PrepareTextureResources(bInWaitForStreaming);
+	}
+	return true;
+}
+
 bool ALandscape::PrepareTextureResources(bool bInWaitForStreaming)
 {
 	TRACE_CPUPROFILER_EVENT_SCOPE(LandscapeLayers_PrepareTextureResources);
@@ -11555,7 +11566,7 @@ void ALandscape::UpdateLayersContent(bool bInWaitForStreaming, bool bInSkipMonit
 	};
 
 	// Note : no early-out allowed before this : even if not actually updating edit layers, we need to poll our resources in order to make sure we register to streaming events when needed: 
-	bool bResourcesReady = PrepareTextureResources(bInWaitForStreaming);
+	bool bResourcesReady = PrepareTextureResourcesLimited(bInWaitForStreaming);
 
 	ULandscapeInfo* LandscapeInfo = GetLandscapeInfo();
 	if (!CanUpdateLayersContent())
