Fix for crash in landscape texture hash code when textures are corrupted

diff --git a/Engine/Source/Runtime/Landscape/Classes/LandscapeTextureHash.h b/Engine/Source/Runtime/Landscape/Classes/LandscapeTextureHash.h
--- a/Engine/Source/Runtime/Landscape/Classes/LandscapeTextureHash.h
+++ b/Engine/Source/Runtime/Landscape/Classes/LandscapeTextureHash.h
@@ -61,8 +61,8 @@ public:
 	static void SetHash64(UTexture2D* LandscapeTexture, uint64 NewHash, ELandscapeTextureUsage TextureUsage, ELandscapeTextureType TextureType);
 
 	// calculate the texture hash directly from the source mip0 contents
-	static uint64 CalculateTextureHash64(UTexture2D* LandscapeTexture, ELandscapeTextureType TextureType);
-	static FGuid CalculateTextureHashGUID(UTexture2D* LandscapeTexture, ELandscapeTextureType TextureType);
+	static uint64 CalculateTextureHash64(UTexture2D* LandscapeTexture, ELandscapeTextureType TextureType, bool& bOutHashIsValid);
+	static FGuid CalculateTextureHashGUID(UTexture2D* LandscapeTexture, ELandscapeTextureType TextureType, bool& bOutHashIsValid);
 
 	// calculate the texture hash from the Mip0Data buffer
 	static uint64 CalculateTextureHash64(const FColor* Mip0Data, int32 PixelCount, ELandscapeTextureType TextureType);
diff --git a/Engine/Source/Runtime/Landscape/Private/LandscapeEdit.cpp b/Engine/Source/Runtime/Landscape/Private/LandscapeEdit.cpp
--- a/Engine/Source/Runtime/Landscape/Private/LandscapeEdit.cpp
+++ b/Engine/Source/Runtime/Landscape/Private/LandscapeEdit.cpp
@@ -6102,7 +6102,8 @@ void ALandscape::SplitHeightmap(ULandscapeComponent* Comp, ALandscapeProxy* Targ
 	if (bWithLayersContent)
 	{
 		FLandscapeEditLayerReadback* NewCPUReadback = new FLandscapeEditLayerReadback();
-		const uint64 Hash = ULandscapeTextureHash::CalculateTextureHash64(NewHeightmapTexture, ELandscapeTextureType::Heightmap);
+		bool bHashIsValid = true;
+		const uint64 Hash = ULandscapeTextureHash::CalculateTextureHash64(NewHeightmapTexture, ELandscapeTextureType::Heightmap, bHashIsValid);
 		NewCPUReadback->SetHash(Hash);
 		DstProxy->HeightmapsCPUReadback.Add(NewHeightmapTexture, NewCPUReadback);
 
diff --git a/Engine/Source/Runtime/Landscape/Private/LandscapeEditLayers.cpp b/Engine/Source/Runtime/Landscape/Private/LandscapeEditLayers.cpp
--- a/Engine/Source/Runtime/Landscape/Private/LandscapeEditLayers.cpp
+++ b/Engine/Source/Runtime/Landscape/Private/LandscapeEditLayers.cpp
@@ -7857,7 +7857,8 @@ int32 ALandscape::RegenerateLayersHeightmaps(const FUpdateLayersContentContext&
 			{
 				FLandscapeEditLayerReadback* NewCPUReadback = new FLandscapeEditLayerReadback();
 				// gather the existing hash, pre-readback
-				const uint64 Hash = ULandscapeTextureHash::CalculateTextureHash64(ComponentHeightmap, ELandscapeTextureType::Heightmap);
+				bool bHashIsValid = true;
+				const uint64 Hash = ULandscapeTextureHash::CalculateTextureHash64(ComponentHeightmap, ELandscapeTextureType::Heightmap, bHashIsValid);
 				NewCPUReadback->SetHash(Hash);
 				Proxy->HeightmapsCPUReadback.Add(ComponentHeightmap, NewCPUReadback);
 			}
@@ -9801,7 +9802,8 @@ void ALandscape::PrepareLayersWeightmapsLocalMergeRenderThreadData(const FUpdate
 						{
 							// Lazily create the readback objects as required (ReallocateLayersWeightmaps might have created new weightmaps)
 							FLandscapeEditLayerReadback* NewCPUReadback = new FLandscapeEditLayerReadback();
-							const uint64 Hash = ULandscapeTextureHash::CalculateTextureHash64(ComponentWeightmap, ELandscapeTextureType::Weightmap);
+							bool bHashIsValid = true;
+							const uint64 Hash = ULandscapeTextureHash::CalculateTextureHash64(ComponentWeightmap, ELandscapeTextureType::Weightmap, bHashIsValid);
 							NewCPUReadback->SetHash(Hash);
 							CPUReadback = &Proxy->WeightmapsCPUReadback.Add(ComponentWeightmap, NewCPUReadback);
 						}
@@ -10305,7 +10307,8 @@ int32 ALandscape::PerformLayersWeightmapsGlobalMerge(FUpdateLayersContentContext
 					if (CPUReadback == nullptr)
 					{
 						FLandscapeEditLayerReadback* NewCPUReadback = new FLandscapeEditLayerReadback();
-						const uint64 Hash = ULandscapeTextureHash::CalculateTextureHash64(WeightmapTexture, ELandscapeTextureType::Weightmap);
+						bool bHashIsValid = true;
+						const uint64 Hash = ULandscapeTextureHash::CalculateTextureHash64(WeightmapTexture, ELandscapeTextureType::Weightmap, bHashIsValid);
 						NewCPUReadback->SetHash(Hash);
 						Proxy->WeightmapsCPUReadback.Add(WeightmapTexture, NewCPUReadback);
 					}
@@ -10639,7 +10642,8 @@ int32 ALandscape::PerformLayersWeightmapsBatchedMerge(FUpdateLayersContentContex
 							{
 								// Lazily create the readback objects as required (ReallocateLayersWeightmaps might have created new weightmaps)
 								FLandscapeEditLayerReadback* NewCPUReadback = new FLandscapeEditLayerReadback();
-								const uint64 Hash = ULandscapeTextureHash::CalculateTextureHash64(Weightmap, ELandscapeTextureType::Weightmap);
+								bool bHashIsValid = true;
+								const uint64 Hash = ULandscapeTextureHash::CalculateTextureHash64(Weightmap, ELandscapeTextureType::Weightmap, bHashIsValid);
 								NewCPUReadback->SetHash(Hash);
 								CPUReadback = &Proxy->WeightmapsCPUReadback.Add(Weightmap, NewCPUReadback);
 							}
diff --git a/Engine/Source/Runtime/Landscape/Private/LandscapeTextureHash.cpp b/Engine/Source/Runtime/Landscape/Private/LandscapeTextureHash.cpp
--- a/Engine/Source/Runtime/Landscape/Private/LandscapeTextureHash.cpp
+++ b/Engine/Source/Runtime/Landscape/Private/LandscapeTextureHash.cpp
@@ -61,20 +61,31 @@ void ULandscapeTextureHash::CheckHashIsUpToDate(UTexture2D* LandscapeTexture)
 	check(LandscapeTexture->Source.GetId() == TextureHash->LastSourceID);
 }
 
-FGuid ULandscapeTextureHash::CalculateTextureHashGUID(UTexture2D* LandscapeTexture, ELandscapeTextureType TextureType)
+FGuid ULandscapeTextureHash::CalculateTextureHashGUID(UTexture2D* LandscapeTexture, ELandscapeTextureType TextureType, bool& bOutHashIsValid)
 {
-	uint64 Hash64 = CalculateTextureHash64(LandscapeTexture, TextureType);
+	uint64 Hash64 = CalculateTextureHash64(LandscapeTexture, TextureType, bOutHashIsValid);
 	return UE::Landscape::Private::Hash64ToGUID(Hash64);
 }
 
-uint64 ULandscapeTextureHash::CalculateTextureHash64(UTexture2D* LandscapeTexture, ELandscapeTextureType TextureType)
+uint64 ULandscapeTextureHash::CalculateTextureHash64(UTexture2D* LandscapeTexture, ELandscapeTextureType TextureType, bool& bOutHashIsValid)
 {
 	const int32 MipIndex = 0;
 	const int64 MipSizeInBytes = LandscapeTexture->Source.CalcMipSize(MipIndex);
 	const int64 MipSizeInPixels = MipSizeInBytes / 4;
+	uint64 Hash;
 	const FColor* MipData = (const FColor*)LandscapeTexture->Source.LockMipReadOnly(MipIndex);
-	uint64 Hash = CalculateTextureHash64(MipData, MipSizeInPixels, TextureType);
-	LandscapeTexture->Source.UnlockMip(MipIndex);
+	if (MipData) // this can be null if the texture data is corrupted and fails to decompress
+	{
+		bOutHashIsValid = true;
+		Hash = CalculateTextureHash64(MipData, MipSizeInPixels, TextureType);
+		LandscapeTexture->Source.UnlockMip(MipIndex);
+	}
+	else
+	{
+		UE_LOG(LogLandscape, Warning, TEXT("Could not calculate the texture hash for landscape texture '%s' because the texture could not be locked (it may be corrupt or being modified incorrectly)"), *LandscapeTexture->GetName());
+		Hash = 0;
+		bOutHashIsValid = false;
+	}
 	return Hash;
 }
 
@@ -287,6 +298,7 @@ void ULandscapeTextureHash::UpdateHash(UTexture2D* LandscapeTexture, ELandscapeT
 	}
 
 	FGuid NewHash;
+	bool bHashIsValid = true;
 	if ((TextureUsage != ELandscapeTextureUsage::FinalData) ||
 		(TextureType == ELandscapeTextureType::Unknown))
 	{
@@ -303,18 +315,21 @@ void ULandscapeTextureHash::UpdateHash(UTexture2D* LandscapeTexture, ELandscapeT
 		else
 		{
 			// otherwise compute a new one
-			NewHash = CalculateTextureHashGUID(LandscapeTexture, TextureType);
+			NewHash = CalculateTextureHashGUID(LandscapeTexture, TextureType, bHashIsValid);
 		}
 	}
 
-	TextureHash->TextureHashGUID = NewHash;
-	TextureHash->LastSourceID = LandscapeTextureSourceId;
-	TextureHash->TextureType = TextureType;
-	TextureHash->TextureUsage = TextureUsage;
-
-	if (bNewlyCreated)
+	if (bHashIsValid)
 	{
-		TextureHash->RecentlySerializedHashes.Add(LandscapeTextureSourceId, NewHash);
+		TextureHash->TextureHashGUID = NewHash;
+		TextureHash->LastSourceID = LandscapeTextureSourceId;
+		TextureHash->TextureType = TextureType;
+		TextureHash->TextureUsage = TextureUsage;
+
+		if (bNewlyCreated)
+		{
+			TextureHash->RecentlySerializedHashes.Add(LandscapeTextureSourceId, NewHash);
+		}
 	}
 }
 
