Fix for crash updating landscape edge snapshots when heightmaps are corrupted (or otherwise not readable)

diff --git a/Engine/Source/Runtime/Landscape/Classes/LandscapeEdgeFixup.h b/Engine/Source/Runtime/Landscape/Classes/LandscapeEdgeFixup.h
--- a/Engine/Source/Runtime/Landscape/Classes/LandscapeEdgeFixup.h
+++ b/Engine/Source/Runtime/Landscape/Classes/LandscapeEdgeFixup.h
@@ -167,7 +167,7 @@ public:
 
 #if WITH_EDITOR
 	// Create a snapshot from the heightmap source
-	static TSharedRef<FHeightmapTextureEdgeSnapshot> CreateEdgeSnapshotFromHeightmapSource(UTexture2D* InHeightmap, const FVector& LandscapeGridScale);
+	static TSharedPtr<FHeightmapTextureEdgeSnapshot> CreateEdgeSnapshotFromHeightmapSource(UTexture2D* InHeightmap, const FVector& LandscapeGridScale);
 #endif // WITH_EDITOR
 	// Create a snapshot from an explicit byte array (in heightmap source standard layout)
 	static TSharedRef<FHeightmapTextureEdgeSnapshot> CreateEdgeSnapshotFromTextureData(const TArrayView64<UE::Landscape::FHeightmapTexel>& InHeightmapTextureData, int32 InEdgeLength, const FVector& LandscapeGridScale);
@@ -189,11 +189,8 @@ public:
 private:
 	void ResizeForEdgeLength(const int32 InEdgeLength);
 
-#if WITH_EDITOR
-	// These CaptureEdgeData* functions are NOT THREAD SAFE --
-	// they stomp existing data and should only be used on newly allocated FHeightmapTextureEdgeSnapshot
-	void CaptureEdgeDataFromHeightmapSource_Internal(UTexture2D* InHeightmap, const FVector& LandscapeGridScale);
-#endif // WITH_EDITOR
+	// This function is not thread safe to call on existing snapshots
+	// it should only be used on a newly allocated FHeightmapTextureEdgeSnapshot
 	void CaptureEdgeDataFromTextureData_Internal(const TArrayView64<UE::Landscape::FHeightmapTexel>& InHeightmapTextureData, int32 InEdgeLength, const FVector& LandscapeGridScale);
 	
 	// must call ResizeForEdgeLength to set the texture size before calling CopyTextureEdgeDataAndRecomputeNormals
@@ -277,7 +274,7 @@ public:
 	void RequestEdgeSnapshotUpdateFromHeightmapSource(bool bUpdateGPUEdgeHashes = false);
 
 	// Update the edge snapshot from heightmap source.  Returns the set of edges that changed since the previous snapshot.
-	UE::Landscape::EEdgeFlags UpdateEdgeSnapshotFromHeightmapSource(const FVector& LandscapeGridScale, bool bForceUpdate = false);
+	UE::Landscape::EEdgeFlags UpdateEdgeSnapshotFromHeightmapSource(const FVector& LandscapeGridScale, bool bForceUpdate = false, bool* bOutSuccess = nullptr);
 #endif // WITH_EDITOR
 
 	// Patch the GPU texture edges if needed, using the current snapshot and corresponding neighbor snapshots as source data.
diff --git a/Engine/Source/Runtime/Landscape/Private/LandscapeEdgeFixup.cpp b/Engine/Source/Runtime/Landscape/Private/LandscapeEdgeFixup.cpp
--- a/Engine/Source/Runtime/Landscape/Private/LandscapeEdgeFixup.cpp
+++ b/Engine/Source/Runtime/Landscape/Private/LandscapeEdgeFixup.cpp
@@ -302,11 +302,27 @@ TSharedRef<FHeightmapTextureEdgeSnapshot> FHeightmapTextureEdgeSnapshot::CreateE
 }
 
 #if WITH_EDITOR
-TSharedRef<FHeightmapTextureEdgeSnapshot> FHeightmapTextureEdgeSnapshot::CreateEdgeSnapshotFromHeightmapSource(UTexture2D* InHeightmap, const FVector& LandscapeGridScale)
+TSharedPtr<FHeightmapTextureEdgeSnapshot> FHeightmapTextureEdgeSnapshot::CreateEdgeSnapshotFromHeightmapSource(UTexture2D* InHeightmap, const FVector& LandscapeGridScale)
 {
 	TRACE_CPUPROFILER_EVENT_SCOPE(ULandscapeHeightmapTextureEdgeSnapshot::CreateEdgeSnapshotFromHeightmapSource);	
-	TSharedRef<FHeightmapTextureEdgeSnapshot> NewEdgeSnapshot = MakeShared<FHeightmapTextureEdgeSnapshot>();
-	NewEdgeSnapshot->CaptureEdgeDataFromHeightmapSource_Internal(InHeightmap, LandscapeGridScale);
+	using namespace UE::Landscape;
+
+	FTextureSource& Source = InHeightmap->Source;
+	check(Source.IsValid());
+	check(Source.GetSizeX() == Source.GetSizeY());
+
+	TArray64<uint8> MipData;
+	if (!Source.GetMipData(MipData, 0))
+	{
+		return nullptr;
+	}
+	FGuid TextureSourceID = ULandscapeTextureHash::GetHash(InHeightmap);
+
+	FIXUP_DEBUG_LOG_DETAIL(TEXT("    CaptureEdgeDataFromHeightmapSource_Internal %p (%s) - ID %s"), InHeightmap, *LandscapeGridScale.ToString(), *TextureSourceID.ToString());
+
+	TArrayView64<FHeightmapTexel> TexelData((FHeightmapTexel*)MipData.GetData(), MipData.Num() / sizeof(FHeightmapTexel));
+	TSharedPtr<FHeightmapTextureEdgeSnapshot> NewEdgeSnapshot = MakeShared<FHeightmapTextureEdgeSnapshot>();
+	NewEdgeSnapshot->CaptureEdgeDataFromTextureData_Internal(TexelData, Source.GetSizeY(), LandscapeGridScale);
 	return NewEdgeSnapshot;
 }
 #endif // WITH_EDITOR
@@ -341,24 +357,6 @@ void FHeightmapTextureEdgeSnapshot::ResizeForEdgeLength(const int32 InEdgeLength
 	}
 }
 
-#if WITH_EDITOR
-void FHeightmapTextureEdgeSnapshot::CaptureEdgeDataFromHeightmapSource_Internal(UTexture2D* InHeightmap, const FVector& LandscapeGridScale)
-{
-	using namespace UE::Landscape;
-
-	FTextureSource& Source = InHeightmap->Source;
-	check(Source.IsValid());
-	check(Source.GetSizeX() == Source.GetSizeY());
-
-	TArray64<uint8> MipData;
-	verify(Source.GetMipData(MipData, 0));
-	TextureSourceID = ULandscapeTextureHash::GetHash(InHeightmap);
-	FIXUP_DEBUG_LOG_DETAIL(TEXT("    CaptureEdgeDataFromHeightmapSource_Internal %p (%s) - ID %s"), InHeightmap, *LandscapeGridScale.ToString(), *TextureSourceID.ToString());
-
-	TArrayView64<FHeightmapTexel> TexelData((FHeightmapTexel*) MipData.GetData(), MipData.Num() / sizeof(FHeightmapTexel));
-	CaptureEdgeDataFromTextureData_Internal(TexelData, Source.GetSizeY(), LandscapeGridScale);
-}
-#endif // WITH_EDITOR
 
 void FHeightmapTextureEdgeSnapshot::CaptureSingleEdgeDataAndComputeNormalsAndHashes(
 	const UE::Landscape::FHeightmapTexel* TextureData, const UE::Landscape::EEdgeIndex EdgeOrCorner, const FVector& LandscapeGridScale)
@@ -1372,7 +1370,7 @@ void ULandscapeHeightmapTextureEdgeFixup::RequestEdgeSnapshotUpdateFromHeightmap
 	ActiveGroup->HeightmapsNeedingEdgeSnapshotCapture.Add(this);
 }
 
-UE::Landscape::ENeighborFlags ULandscapeHeightmapTextureEdgeFixup::UpdateEdgeSnapshotFromHeightmapSource(const FVector& LandscapeGridScale, bool bForceUpdate)
+UE::Landscape::ENeighborFlags ULandscapeHeightmapTextureEdgeFixup::UpdateEdgeSnapshotFromHeightmapSource(const FVector& LandscapeGridScale, bool bForceUpdate, bool* bOutSuccess)
 {
 	TRACE_CPUPROFILER_EVENT_SCOPE(UpdateEdgeSnapshotFromHeightmapSource);
 
@@ -1383,7 +1381,6 @@ UE::Landscape::ENeighborFlags ULandscapeHeightmapTextureEdgeFixup::UpdateEdgeSna
 	if (bForceUpdateSnapshot)
 	{
 		bForceUpdate = true;
-		bForceUpdateSnapshot = false;
 	}
 
 	// no need to update if source ID matches (we assume texture has not actually changed)		
@@ -1394,28 +1391,49 @@ UE::Landscape::ENeighborFlags ULandscapeHeightmapTextureEdgeFixup::UpdateEdgeSna
 			*EdgeSnapshot->TextureSourceID.ToString(), *HeightmapSourceID.ToString());
 			
 		// create a new edge data
-		TSharedRef<FHeightmapTextureEdgeSnapshot> NewEdgeSnapshot = FHeightmapTextureEdgeSnapshot::CreateEdgeSnapshotFromHeightmapSource(HeightmapTexture, LandscapeGridScale);
-
-		// compare against previous edge data to see what edges have changed (and might cause neighbors to need to patch)
-		ChangedEdges = NewEdgeSnapshot->CompareEdges(EdgeSnapshot.Get());
+		TSharedPtr<FHeightmapTextureEdgeSnapshot> NewEdgeSnapshot = FHeightmapTextureEdgeSnapshot::CreateEdgeSnapshotFromHeightmapSource(HeightmapTexture, LandscapeGridScale);
+		if (!NewEdgeSnapshot.IsValid())
+		{
+			// failed to capture an edge snapshot (for example HeightmapTexture could have corrupt data)
+			if (bOutSuccess)
+			{
+				UE_LOG(LogLandscape, Warning, TEXT("Failed to capture edge snapshot from landscape heightmap, landscape edge patching may not work correctly."));
+				*bOutSuccess = false;
+			}
+			return ENeighborFlags::None;
+		}
+		else
+		{
+			// compare against previous edge data to see what edges have changed (and might cause neighbors to need to patch)
+			ChangedEdges = NewEdgeSnapshot->CompareEdges(EdgeSnapshot.Get());
 
-		// assign the new edge data
-		EdgeSnapshot = NewEdgeSnapshot;
+			// assign the new edge data
+			EdgeSnapshot = NewEdgeSnapshot.ToSharedRef();
 
-		if (bUpdateGPUEdgeHashes)
-		{
-			GPUEdgeHashes = EdgeSnapshot->InitialEdgeHashes;
-			GPUEdgeModifiedFlags = EEdgeFlags::None;
+			if (bUpdateGPUEdgeHashes)
+			{
+				GPUEdgeHashes = EdgeSnapshot->InitialEdgeHashes;
+				GPUEdgeModifiedFlags = EEdgeFlags::None;
 
-			FIXUP_DEBUG_LOG_DETAIL(TEXT("    Reset GPU Edge Hashes: [B:%x BR:%x R:%x TR:%x T:%x TL:%x L:%x BL:%x]"),
-				GPUEdgeHashes[0], GPUEdgeHashes[1], GPUEdgeHashes[2], GPUEdgeHashes[3], GPUEdgeHashes[4], GPUEdgeHashes[5], GPUEdgeHashes[6], GPUEdgeHashes[7]);
+				FIXUP_DEBUG_LOG_DETAIL(TEXT("    Reset GPU Edge Hashes: [B:%x BR:%x R:%x TR:%x T:%x TL:%x L:%x BL:%x]"),
+					GPUEdgeHashes[0], GPUEdgeHashes[1], GPUEdgeHashes[2], GPUEdgeHashes[3], GPUEdgeHashes[4], GPUEdgeHashes[5], GPUEdgeHashes[6], GPUEdgeHashes[7]);
+			}
+			if (bOutSuccess)
+			{
+				*bOutSuccess = true;
+			}
 		}
 	}
 	else
 	{
 		FIXUP_DEBUG_LOG_DETAIL(TEXT("  UpdateEdgeSnapshotFromHeightmapSource %p HeightMap: %p -- SourceID matches (%s)"), ActiveComponent.Get(), HeightmapTexture, *EdgeSnapshot->TextureSourceID.ToString());
+		if (bOutSuccess)
+		{
+			*bOutSuccess = true;
+		}
 	}
 
+	bForceUpdateSnapshot = false;
 	bUpdateGPUEdgeHashes = false;
 	bDoNotPatchUntilGPUEdgeHashesUpdated = false;
 
diff --git a/Engine/Source/Runtime/Landscape/Private/LandscapeGroup.cpp b/Engine/Source/Runtime/Landscape/Private/LandscapeGroup.cpp
--- a/Engine/Source/Runtime/Landscape/Private/LandscapeGroup.cpp
+++ b/Engine/Source/Runtime/Landscape/Private/LandscapeGroup.cpp
@@ -597,22 +597,26 @@ void FLandscapeGroup::TickEdgeFixup(ULandscapeSubsystem* LandscapeSubsystem, boo
 			}
 #endif // WITH_EDITORONLY_DATA
 
-			EEdgeFlags ChangedEdges = EdgeFixup->UpdateEdgeSnapshotFromHeightmapSource(LandscapeGridScale);
-			SnapshotsCaptured++;
-			SnapshotEdgesChanged += FGenericPlatformMath::CountBits((uint64)ChangedEdges);
-
-			GROUP_DEBUG_LOG_DETAIL(TEXT("-- Snap Component %p (%d,%d) -------- UPDATED, Changed edges: %s"), EdgeFixup->ActiveComponent.Get(), EdgeFixup->GroupCoord.X, EdgeFixup->GroupCoord.Y, *DirectionFlagsToString(ChangedEdges));
-			if (ChangedEdges != EEdgeFlags::None)
+			bool bSuccess = false;
+			EEdgeFlags ChangedEdges = EdgeFixup->UpdateEdgeSnapshotFromHeightmapSource(LandscapeGridScale, false, &bSuccess);
+			if (bSuccess)
 			{
-				// if any edges were changed -- queue it for texture edge patching
-				HeightmapsNeedingEdgeTexturePatching.Add(EdgeFixup);
+				SnapshotsCaptured++;
+				SnapshotEdgesChanged += FGenericPlatformMath::CountBits((uint64)ChangedEdges);
+
+				GROUP_DEBUG_LOG_DETAIL(TEXT("-- Snap Component %p (%d,%d) -------- UPDATED, Changed edges: %s"), EdgeFixup->ActiveComponent.Get(), EdgeFixup->GroupCoord.X, EdgeFixup->GroupCoord.Y, *DirectionFlagsToString(ChangedEdges));
+				if (ChangedEdges != EEdgeFlags::None)
+				{
+					// if any edges were changed -- queue it for texture edge patching
+					HeightmapsNeedingEdgeTexturePatching.Add(EdgeFixup);
 
-				ENeighborFlags AffectedNeighbors = EdgesToAffectedNeighbors(ChangedEdges);
+					ENeighborFlags AffectedNeighbors = EdgesToAffectedNeighbors(ChangedEdges);
 
-				// and also mark the affected neighbors for potential texture edge patching
-				EdgeFixup->RequestEdgeTexturePatchingForNeighbors(AffectedNeighbors);
+					// and also mark the affected neighbors for potential texture edge patching
+					EdgeFixup->RequestEdgeTexturePatchingForNeighbors(AffectedNeighbors);
+				}
+				It.RemoveCurrent();
 			}
-			It.RemoveCurrent();
 		}
 	}
 #endif // WITH_EDITOR
