Specialized FCopyQuadsMultiSourcePS for Vulkan SM6, as it's currently limited to a different number of textures

diff --git a/Engine/Shaders/Private/Landscape/LandscapeEditLayersUtils.usf b/Engine/Shaders/Private/Landscape/LandscapeEditLayersUtils.usf
--- a/Engine/Shaders/Private/Landscape/LandscapeEditLayersUtils.usf
+++ b/Engine/Shaders/Private/Landscape/LandscapeEditLayersUtils.usf
@@ -27,6 +27,38 @@ Texture2D<float4> InSourceTexture;
 // CopyQuadsMultiSourcePS inputs/outputs :
 
 // Helper macro for handling an array of N textures (currently using 63 because 64 is the maximum supported number of SRV on all platforms needing to run this) :
+#if VULKAN_PROFILE_SM6
+#define FOREACH_SOURCETEXTURE(op) \
+	op(0) \
+	op(1) \
+	op(2) \
+	op(3) \
+	op(4) \
+	op(5) \
+	op(6) \
+	op(7) \
+	op(8) \
+	op(9) \
+	op(10) \
+	op(11) \
+	op(12) \
+	op(13) \
+	op(14) \
+	op(15) \
+	op(16) \
+	op(17) \
+	op(18) \
+	op(19) \
+	op(20) \
+	op(21) \
+	op(22) \
+	op(23) \
+	op(24) \
+	op(25) \
+	op(26) \
+	op(27) \
+	op(28) 
+#else // VULKAN_PROFILE_SM6 
 #define FOREACH_SOURCETEXTURE(op) \
 	op(0) \
 	op(1) \
@@ -90,7 +122,8 @@ Texture2D<float4> InSourceTexture;
 	op(59) \
 	op(60) \
 	op(61) \
-	op(62) 
+	op(62)
+#endif // !VULKAN_PROFILE_SM6 
 
 Buffer<uint2> InQuadInfos; // .x = source texture index, .y = source texture channel (for weightmaps)
 
diff --git a/Engine/Shaders/Private/Landscape/LandscapeEditLayersWeightmaps.usf b/Engine/Shaders/Private/Landscape/LandscapeEditLayersWeightmaps.usf
--- a/Engine/Shaders/Private/Landscape/LandscapeEditLayersWeightmaps.usf
+++ b/Engine/Shaders/Private/Landscape/LandscapeEditLayersWeightmaps.usf
@@ -62,7 +62,7 @@ uint4 InSourcePixelOffsets[4];  // For each channel, offset to add to the pixel'
 uint2 InSubsectionPixelOffset; // Offset of the subsection currently being rendered
 uint InIsAdditive; // = 1 if some channels of this texture have been packed in a previous draw, 0 otherwise
 
-Texture2DArray<float2> InSourceWeightmaps; // Source, single-channel (we don't care about the alpha flags at this point), texture to pack
+Texture2DArray<float4> InSourceWeightmaps; // Source, single-channel (we don't care about the alpha flags at this point), texture to pack
 Texture2D<float4> InWeightmapBeingPacked; // The weightmap being packed, in case the packing operation is additive, i.e. it occurs across multiple draws (contains the channels that have previously been packed)
 #endif // PACK_WEIGHTMAP
 
diff --git a/Engine/Source/Runtime/Landscape/Private/LandscapeEditLayers.cpp b/Engine/Source/Runtime/Landscape/Private/LandscapeEditLayers.cpp
--- a/Engine/Source/Runtime/Landscape/Private/LandscapeEditLayers.cpp
+++ b/Engine/Source/Runtime/Landscape/Private/LandscapeEditLayers.cpp
@@ -1922,13 +1922,13 @@ IMPLEMENT_GLOBAL_SHADER(FLandscapeEditLayersWeightmapsGenerateMipsPS, "/Engine/P
 
 // ----------------------------------------------------------------------------------
 // /Engine/Private/Landscape/LandscapeEditLayersUtils.usf shaders : 
-
-class FCopyQuadsMultiSourcePS : public FGlobalShader
+class FCopyQuadsMultiSourcePSDefault : public FGlobalShader
 {
-	DECLARE_GLOBAL_SHADER(FCopyQuadsMultiSourcePS);
-	SHADER_USE_PARAMETER_STRUCT(FCopyQuadsMultiSourcePS, FGlobalShader);
+public:
+	static constexpr int32 NumMultiSources = 63;
 
-	static constexpr int32 NumMultiSources = 63; // This has to match the number of input textures supported by CopyQuadsMultiSourcePS
+	DECLARE_GLOBAL_SHADER(FCopyQuadsMultiSourcePSDefault);
+	SHADER_USE_PARAMETER_STRUCT(FCopyQuadsMultiSourcePSDefault, FGlobalShader);
 
 	BEGIN_SHADER_PARAMETER_STRUCT(FParameters, )
 		SHADER_PARAMETER_STRUCT_REF(FViewUniformShaderParameters, View)
@@ -1950,14 +1950,124 @@ class FCopyQuadsMultiSourcePS : public FGlobalShader
 	}
 };
 
-IMPLEMENT_GLOBAL_SHADER(FCopyQuadsMultiSourcePS, "/Engine/Private/Landscape/LandscapeEditLayersUtils.usf", "CopyQuadsMultiSourcePS", SF_Pixel);
+class FCopyQuadsMultiSourcePSVulkanSM6 : public FGlobalShader
+{
+public:
+	static constexpr int32 NumMultiSources = 29;
+
+	DECLARE_GLOBAL_SHADER(FCopyQuadsMultiSourcePSVulkanSM6);
+	SHADER_USE_PARAMETER_STRUCT(FCopyQuadsMultiSourcePSVulkanSM6, FGlobalShader);
+
+	BEGIN_SHADER_PARAMETER_STRUCT(FParameters, )
+		SHADER_PARAMETER_STRUCT_REF(FViewUniformShaderParameters, View)
+		SHADER_PARAMETER_RDG_BUFFER_SRV(Buffer<uint2>, InQuadInfos)
+		SHADER_PARAMETER_RDG_TEXTURE_ARRAY(Texture2D, InSourceTexture, [NumMultiSources])
+	END_SHADER_PARAMETER_STRUCT()
+
+	class FCopyWeightmap : SHADER_PERMUTATION_BOOL("COPY_WEIGHTMAP");
+	using FPermutationDomain = TShaderPermutationDomain<FCopyWeightmap>;
+
+	static bool ShouldCompilePermutation(const FGlobalShaderPermutationParameters& InParameters)
+	{
+		return UE::Landscape::DoesPlatformSupportEditLayers(InParameters.Platform) && (InParameters.Platform == SP_VULKAN_SM6);
+	}
+
+	static void ModifyCompilationEnvironment(const FGlobalShaderPermutationParameters& Parameters, FShaderCompilerEnvironment& OutEnvironment)
+	{
+		OutEnvironment.SetDefine(TEXT("COPY_QUADS_MULTISOURCE"), 1);
+	}
+};
+
+// Helper class to trigger the appropriate PS regardless of the platform : 
+class FCopyQuadsMultiSourceBase
+{
+public:
+	virtual ~FCopyQuadsMultiSourceBase() = default;
+	virtual int32 GetNumMultiSources() const PURE_VIRTUAL(FCopyQuadsMultiSourceBase::GetNumMultiSources, return 0; );
+	virtual void CopyQuads(FRDGBuilder& GraphBuilder,
+		FRDGTextureRef InOutputTexture,
+		TConstArrayView<FRDGTextureRef> InSourceTextures,
+		FRDGBufferSRVRef InRectBufferSRV,
+		FRDGBufferSRVRef InQuadInfosBufferSRV,
+		FRDGBufferSRVRef InRectUVBufferSRV,
+		const FSceneView* InView,
+		int32 InArrayIndex,
+		int32 InNumRects,
+		const FIntPoint& InTextureSize,
+		bool bIsWeightmapMerge) PURE_VIRTUAL(FCopyQuadsMultiSourceBase::CopyQuads, );
+};
 
-BEGIN_SHADER_PARAMETER_STRUCT(FCopyQuadsMultiSourcePSParameters, )
-SHADER_PARAMETER_STRUCT_INCLUDE(FPixelShaderUtils::FRasterizeToRectsVS::FParameters, VS)
-SHADER_PARAMETER_STRUCT_INCLUDE(FCopyQuadsMultiSourcePS::FParameters, PS)
-RENDER_TARGET_BINDING_SLOTS()
+template <typename FCopyQuadsMultiSourcePSType, typename FCopyQuadsMultiSourcePSParametersType>
+class TCopyQuadsMultiSource : public FCopyQuadsMultiSourceBase
+{
+public:
+	static constexpr int32 NumMultiSources = FCopyQuadsMultiSourcePSType::NumMultiSources;
+
+	virtual int32 GetNumMultiSources() const override
+	{
+		return NumMultiSources;
+	}
+
+	virtual void CopyQuads(FRDGBuilder& GraphBuilder, 
+		FRDGTextureRef InOutputTexture, 
+		TConstArrayView<FRDGTextureRef> InSourceTextures, 
+		FRDGBufferSRVRef InRectBufferSRV, 
+		FRDGBufferSRVRef InQuadInfosBufferSRV, 
+		FRDGBufferSRVRef InRectUVBufferSRV,
+		const FSceneView* InView, 
+		int32 InArrayIndex,
+		int32 InNumRects, 
+		const FIntPoint& InTextureSize, 
+		bool bIsWeightmapMerge) override
+	{
+		typename FCopyQuadsMultiSourcePSType::FPermutationDomain PermutationVector;
+		PermutationVector.template Set<typename FCopyQuadsMultiSourcePSType::FCopyWeightmap>(bIsWeightmapMerge);
+		FGlobalShaderMap* ShaderMap = GetGlobalShaderMap(GMaxRHIFeatureLevel);
+		TShaderRef<FCopyQuadsMultiSourcePSType> PixelShader = ShaderMap->GetShader<FCopyQuadsMultiSourcePSType>(PermutationVector);
+
+		FCopyQuadsMultiSourcePSParametersType* PassParameters = GraphBuilder.AllocParameters<FCopyQuadsMultiSourcePSParametersType>();
+		PassParameters->RenderTargets[0] = FRenderTargetBinding(InOutputTexture, ERenderTargetLoadAction::ELoad, /*InMipIndex = */0, InArrayIndex);
+		PassParameters->PS.View = InView->ViewUniformBuffer;
+		PassParameters->PS.InQuadInfos = InQuadInfosBufferSRV;
+		check(InSourceTextures.Num() == NumMultiSources);
+		for (int32 TextureIndex = 0; TextureIndex < NumMultiSources; ++TextureIndex)
+		{
+			PassParameters->PS.InSourceTexture[TextureIndex] = InSourceTextures[TextureIndex];
+		}
+
+		FPixelShaderUtils::AddRasterizeToRectsPass<FCopyQuadsMultiSourcePSType>(GraphBuilder,
+			ShaderMap,
+			RDG_EVENT_NAME("CopyQuadsMultiSourcePS"),
+			PixelShader,
+			PassParameters,
+			/*ViewportSize = */InOutputTexture->Desc.Extent,
+			InRectBufferSRV,
+			InNumRects,
+			/*BlendState = */nullptr,
+			/*RasterizerState = */nullptr,
+			/*DepthStencilState = */nullptr,
+			/*StencilRef = */0,
+			InTextureSize,
+			InRectUVBufferSRV);
+	}
+};
+
+BEGIN_SHADER_PARAMETER_STRUCT(FCopyQuadsMultiSourcePSParametersDefault, )
+	SHADER_PARAMETER_STRUCT_INCLUDE(FPixelShaderUtils::FRasterizeToRectsVS::FParameters, VS)
+	SHADER_PARAMETER_STRUCT_INCLUDE(FCopyQuadsMultiSourcePSDefault::FParameters, PS)
+	RENDER_TARGET_BINDING_SLOTS() 
+END_SHADER_PARAMETER_STRUCT() 
+IMPLEMENT_GLOBAL_SHADER(FCopyQuadsMultiSourcePSDefault, "/Engine/Private/Landscape/LandscapeEditLayersUtils.usf", "CopyQuadsMultiSourcePS", SF_Pixel);
+
+BEGIN_SHADER_PARAMETER_STRUCT(FCopyQuadsMultiSourcePSParametersVulkanSM6, )
+	SHADER_PARAMETER_STRUCT_INCLUDE(FPixelShaderUtils::FRasterizeToRectsVS::FParameters, VS)
+	SHADER_PARAMETER_STRUCT_INCLUDE(FCopyQuadsMultiSourcePSVulkanSM6::FParameters, PS)
+	RENDER_TARGET_BINDING_SLOTS()
 END_SHADER_PARAMETER_STRUCT()
+IMPLEMENT_GLOBAL_SHADER(FCopyQuadsMultiSourcePSVulkanSM6, "/Engine/Private/Landscape/LandscapeEditLayersUtils.usf", "CopyQuadsMultiSourcePS", SF_Pixel);
 
+using FCopyQuadsMultiSourceDefault = TCopyQuadsMultiSource<FCopyQuadsMultiSourcePSDefault, FCopyQuadsMultiSourcePSParametersDefault>;
+using FCopyQuadsMultiSourceVulkanSM6 = TCopyQuadsMultiSource<FCopyQuadsMultiSourcePSVulkanSM6, FCopyQuadsMultiSourcePSParametersVulkanSM6>;
 
 // ----------------------------------------------------------------------------------
 // Copy texture render command
@@ -5636,7 +5746,7 @@ bool ULandscapeEditLayerPersistent::RenderLayer(UE::Landscape::EditLayers::FRend
 
 	// Now process this list of quads and prepare as many passes as necessary for performing all the copies :
 	auto RDGCommand =
-		[CopyQuadsParams = MoveTemp(CopyQuadsParams)
+		[ CopyQuadsParams = MoveTemp(CopyQuadsParams)
 		, OutputResource = WriteRT->GetRenderTarget()->GetResource()
 		, OutputResourceName = WriteRT->GetDebugName()
 		, SceneInterface
@@ -5647,10 +5757,15 @@ bool ULandscapeEditLayerPersistent::RenderLayer(UE::Landscape::EditLayers::FRend
 
 			FRDGTextureSRVRef BlackDummySRVRef = GraphBuilder.CreateSRV(FRDGTextureSRVDesc::Create(GSystemTextures.GetBlackDummy(GraphBuilder)));
 
-			FCopyQuadsMultiSourcePS::FPermutationDomain PermutationVector;
-			PermutationVector.Set< FCopyQuadsMultiSourcePS::FCopyWeightmap >(bIsWeightmapMerge);
-			FGlobalShaderMap* ShaderMap = GetGlobalShaderMap(GMaxRHIFeatureLevel);
-			TShaderRef<FCopyQuadsMultiSourcePS> PixelShader = ShaderMap->GetShader<FCopyQuadsMultiSourcePS>(PermutationVector);
+			TUniquePtr<FCopyQuadsMultiSourceBase> CopyQuadsMultiSource;
+			if (SceneInterface->GetShaderPlatform() == SP_VULKAN_SM6)
+			{
+				CopyQuadsMultiSource = MakeUnique<FCopyQuadsMultiSourceVulkanSM6>();
+			}
+			else
+			{
+				CopyQuadsMultiSource = MakeUnique<FCopyQuadsMultiSourceDefault>();
+			}
 
 			// TODO [jonathan.bard] this is just to avoid a RHI validation error for unoptimized shaders... once validation is made to not issue those errors, we can remove this
 			// Create a SceneView to please the shader bindings, but it's unused in practice 
@@ -5670,8 +5785,19 @@ bool ULandscapeEditLayerPersistent::RenderLayer(UE::Landscape::EditLayers::FRend
 			FRDGTextureRef OutputTexture = GraphBuilder.RegisterExternalTexture(CreateRenderTarget(OutputResource->GetTextureRHI(), TEXT("OutputTexture")));
 
 			// When we've reached the limit of what we can render in one pass, flush all and add a render pass : 
-			auto FlushCopyQuads = [&GraphBuilder, View, OutputTexture, &CurrentDestinationArrayIndex, &CurrentSourceTextureResources, &CurrentSourceRects, &CurrentDestinationRects,
-				&CurrentQuadInfos, &CommonSourceTextureSize, BlackDummySRVRef, ShaderMap, PixelShader]()
+			auto FlushCopyQuads = [
+				&GraphBuilder, 
+				View, 
+				OutputTexture, 
+				&CurrentDestinationArrayIndex, 
+				&CurrentSourceTextureResources, 
+				&CurrentSourceRects, 
+				&CurrentDestinationRects,
+				&CurrentQuadInfos, 
+				&CommonSourceTextureSize, 
+				BlackDummySRVRef, 
+				CopyQuadsMultiSource = CopyQuadsMultiSource.Get(), 
+				bIsWeightmapMerge]()
 				{
 					if (CurrentSourceTextureResources.IsEmpty())
 					{
@@ -5699,35 +5825,30 @@ bool ULandscapeEditLayerPersistent::RenderLayer(UE::Landscape::EditLayers::FRend
 						ArrayIndex = CurrentDestinationArrayIndex;
 					}
 
-					FCopyQuadsMultiSourcePSParameters* PassParameters = GraphBuilder.AllocParameters<FCopyQuadsMultiSourcePSParameters>();
-					PassParameters->RenderTargets[0] = FRenderTargetBinding(OutputTexture, ERenderTargetLoadAction::ELoad, /*InMipIndex = */0, ArrayIndex);
-					PassParameters->PS.View = View->ViewUniformBuffer;
-					PassParameters->PS.InQuadInfos = QuadInfosBufferSRV;
-					check(CurrentSourceTextureResources.Num() <= FCopyQuadsMultiSourcePS::NumMultiSources);
-					for (int32 TextureIndex = 0; TextureIndex < FCopyQuadsMultiSourcePS::NumMultiSources; ++TextureIndex)
+					TArray<FRDGTextureRef> SourceTextures;
+					check(CurrentSourceTextureResources.Num() <= CopyQuadsMultiSource->GetNumMultiSources());
+					for (int32 TextureIndex = 0; TextureIndex < CopyQuadsMultiSource->GetNumMultiSources(); ++TextureIndex)
 					{
 						FRDGTextureRef SourceTexture = BlackDummySRVRef->GetParent();
 						if (CurrentSourceTextureResources.IsValidIndex(TextureIndex))
 						{
 							SourceTexture = GraphBuilder.RegisterExternalTexture(CreateRenderTarget(CurrentSourceTextureResources[TextureIndex]->GetTexture2DRHI(), TEXT("SourceTexture")));
 						}
-						PassParameters->PS.InSourceTexture[TextureIndex] = SourceTexture;
+						SourceTextures.Add(SourceTexture);
 					}
 
-					FPixelShaderUtils::AddRasterizeToRectsPass<FCopyQuadsMultiSourcePS>(GraphBuilder,
-						ShaderMap,
-						RDG_EVENT_NAME("CopyQuadsMultiSourcePS"),
-						PixelShader,
-						PassParameters,
-						/*ViewportSize = */OutputTexture->Desc.Extent,
+					CopyQuadsMultiSource->CopyQuads(
+						GraphBuilder,
+						OutputTexture,
+						MakeConstArrayView(SourceTextures),
 						RectBufferSRV,
+						QuadInfosBufferSRV,
+						RectUVBufferSRV,
+						View,
+						ArrayIndex,
 						CurrentDestinationRects.Num(),
-						/*BlendState = */nullptr,
-						/*RasterizerState = */nullptr,
-						/*DepthStencilState = */nullptr,
-						/*StencilRef = */0,
-						/*TextureSize = */FIntPoint(FMath::FloorToInt32(CommonSourceTextureSize.X), FMath::FloorToInt32(CommonSourceTextureSize.Y)),
-						RectUVBufferSRV);
+						FIntPoint(FMath::FloorToInt32(CommonSourceTextureSize.X), FMath::FloorToInt32(CommonSourceTextureSize.Y)),
+						bIsWeightmapMerge);
 
 					// Reset all for the next pass that comes along : 
 					CurrentDestinationArrayIndex = INDEX_NONE;
@@ -5757,7 +5878,7 @@ bool ULandscapeEditLayerPersistent::RenderLayer(UE::Landscape::EditLayers::FRend
 				else if (InSingleQuadParams.SourceTextureResource != CurrentSourceTextureResources.Last())
 				{
 					// If we've reached the amount of textures we can render in a single pass, we flush the pass and initiate a new one : 
-					if (CurrentSourceTextureResources.Num() == FCopyQuadsMultiSourcePS::NumMultiSources)
+					if (CurrentSourceTextureResources.Num() == CopyQuadsMultiSource->GetNumMultiSources())
 					{
 						FlushCopyQuads();
 						check(CurrentSourceTextureResources.IsEmpty());
@@ -10737,7 +10858,7 @@ int32 ALandscape::PerformLayersWeightmapsBatchedMerge(FUpdateLayersContentContex
 							const FIntRect OutputRect = FIntRect(SubSection * ComponentSubsectionVerts, SubSection * ComponentSubsectionVerts + FIntPoint(ComponentSubsectionVerts, ComponentSubsectionVerts));
 
 							FLandscapeEditLayersWeightmapsPackWeightmapPS::FParameters* PSParams = GraphBuilder.AllocParameters<FLandscapeEditLayersWeightmapsPackWeightmapPS::FParameters>();
-							PSParams->RenderTargets[0] = FRenderTargetBinding(PackedTextureRef, ERenderTargetLoadAction::ENoAction);
+							PSParams->RenderTargets[0] = FRenderTargetBinding(PackedTextureRef, ERenderTargetLoadAction::ELoad);
 							PSParams->InSourceSliceIndices = WeightmapResolveInfo.SourceSliceIndexPerChannel;
 							for (int32 ChannelIndex = 0; ChannelIndex < 4; ++ChannelIndex)
 							{
