Fixed determinism issue when updating landscapes with conflicts on the border of landscape components.
* Landscape scratch texture pool logic has been fixed to prevent textures of a different format from being picked up from the pool
* Also use bitwise comparison to know whether a texture's clear color is identical (allows to pack 0xffffffff into the clear color, which is NaN in float and therefore failed comparison)
* In batched merge, per-target layer visibility (stencil-like) render targets have been replaced by a single ComponentId render target, because they were not needed per target layer anyway (saves a bunch of render passes at the beginning of each batch) which is rendered by rendering quads in a deterministic order. The component Id is simply the component's linear index into the batch. This allows edit layers, when rendering component quads, to discard pixels that are not matching the authoritative component for this pixel (because the rasterization order cannot be determined at this point).

diff --git a/Engine/Shaders/Private/Landscape/LandscapeCommon.ush b/Engine/Shaders/Private/Landscape/LandscapeCommon.ush
--- a/Engine/Shaders/Private/Landscape/LandscapeCommon.ush
+++ b/Engine/Shaders/Private/Landscape/LandscapeCommon.ush
@@ -21,6 +21,8 @@ static const float LANDSCAPE_MAX_VALUE = 65535.0f;
 // Must match LandscapeDataAccess::MidValue : 
 static const float LANDSCAPE_MID_VALUE = 32768.0f;
 
+static const uint INDEX_NONE = uint(-1);
+
 // Height is packed from a [0, 2^16-1] value to two bytes, each of which we need to provide as a 0-1 float for writing 
 float2 PackHeight(float InHeight)
 {
diff --git a/Engine/Shaders/Private/Landscape/LandscapeEditLayersHeightmaps.usf b/Engine/Shaders/Private/Landscape/LandscapeEditLayersHeightmaps.usf
--- a/Engine/Shaders/Private/Landscape/LandscapeEditLayersHeightmaps.usf
+++ b/Engine/Shaders/Private/Landscape/LandscapeEditLayersHeightmaps.usf
@@ -41,7 +41,7 @@ uint2 InNumComponents; // Number of components in the effective region of the te
 // TODO [jonathan.bard] : this can be replaced by GlobalPointClampedSampler :
 SamplerState InSourceHeightmapSampler;
 Texture2D<float4> InSourceHeightmap;
-Texture2D<float> InValidityTexture;
+Texture2D<float> InComponentIdTexture;
 #endif // GENERATE_NORMALS
 
 #if GENERATE_MIPS
@@ -85,7 +85,8 @@ float4 ExtractSampleInfo(int2 InTextureCoordinates, float2 InPackedHeight)
 {
 	// Note: don't use a bool here, as some compilers don't support it : 
 	int bIsValidSample = (int)(all(InTextureCoordinates >= 0) && all(InTextureCoordinates < (int2)InTextureSize.xy));
-	bIsValidSample &= (InValidityTexture.Load(int3(InTextureCoordinates, 0)).x > 0.0f);
+	uint ComponentId = asuint(InComponentIdTexture.Load(int3(InTextureCoordinates, 0)).x);
+	bIsValidSample &= (ComponentId != INDEX_NONE);
 	return float4((InTextureCoordinates + 0.5f) * InLandscapeGridScale.xy, ExtractAndScaleHeight(InPackedHeight), bIsValidSample);
 }
 #endif // GENERATE_NORMALS
diff --git a/Engine/Shaders/Private/Landscape/LandscapeEditLayersUtils.usf b/Engine/Shaders/Private/Landscape/LandscapeEditLayersUtils.usf
--- a/Engine/Shaders/Private/Landscape/LandscapeEditLayersUtils.usf
+++ b/Engine/Shaders/Private/Landscape/LandscapeEditLayersUtils.usf
@@ -7,7 +7,7 @@
 
 #if defined (__INTELLISENSE__)
 // Uncomment the appropriate define for enabling syntax highlighting with HLSL Tools for Visual Studio : 
-//#define MARK_VALIDITY 1
+//#define RENDER_COMPONENTID 1
 //#define COPY_QUADS 1
 //#define COPY_QUADS_MULTISOURCE 1
 //#define COPY_WEIGHTMAP 1
@@ -17,6 +17,11 @@
 
 // ----------------------------------------------------------------------------------
 
+#if RENDER_COMPONENTID
+// RenderComponentIdPS inputs/outputs :
+Buffer<uint> InQuadComponentIds; // Id to write into the output for each quad 
+#endif // RENDER_COMPONENTID
+
 #if COPY_QUADS
 // CopyQuadsPS inputs/outputs :
 
@@ -121,11 +126,11 @@ Texture2D<float4> InSourceTexture;
 	op(58) \
 	op(59) \
 	op(60) \
-	op(61) \
-	op(62)
+	op(61) 
 #endif // !VULKAN_PROFILE_SM6 
 
-Buffer<uint2> InQuadInfos; // .x = source texture index, .y = source texture channel (for weightmaps)
+Buffer<uint4> InQuadInfos; // .x = source texture index, .y = source texture channel (for weightmaps), .z = component id of the component this quad corresponds to, .w = unused
+Texture2D<float> InComponentIdTexture; // texture containing the linear index of the components in this batch. This allows to resolve conflicts in case multiple quads render to the same pixel
 
 // Define N input textures :
 #define DECLARE_SOURCE_TEXTURE(n) Texture2D<float4> InSourceTexture_##n;
@@ -140,14 +145,15 @@ Texture2D<float4> InSourceTexture;
 // ----------------------------------------------------------------------------------
 // Pixel shaders : 
 
-#if MARK_VALIDITY
+#if RENDER_COMPONENTID
 
-void MarkValidityPS(out float OutValidity : SV_Target0)
+void RenderComponentIdPS(float4 SvPosition : SV_POSITION, float2 SourceUV : TEXCOORD0, float2 UnusedQuadUV : TEXCOORD1, float InRectIndex : RECT_INDEX, out float OutColor : SV_Target0)
 {
-	OutValidity = 1.0f;
+	uint QuadIndex = (uint)InRectIndex;
+	OutColor = asfloat(InQuadComponentIds[QuadIndex]);
 }
 
-#endif // MARK_VALIDITY
+#endif // RENDER_COMPONENTID
 
 
 // ----------------------------------------------------------------------------------
@@ -169,9 +175,18 @@ void CopyQuadsPS(float4 SvPosition : SV_POSITION, float2 QuadUV : TEXCOORD0, out
 void CopyQuadsMultiSourcePS(float4 SvPosition : SV_POSITION, float2 SourceUV : TEXCOORD0, float2 UnusedQuadUV : TEXCOORD1, float InRectIndex : RECT_INDEX, out float4 OutColor : SV_Target0)
 {
 	uint QuadIndex = (uint)InRectIndex;
-	int SourceTextureIndex = InQuadInfos[QuadIndex].x;
+	uint SourceTextureIndex = InQuadInfos[QuadIndex].x;
+	uint QuadComponentId = InQuadInfos[QuadIndex].z;
 	float4 ColorSample = 0.0f;
 
+	float ComponentIdSample = InComponentIdTexture.Load(int3(SvPosition.xy, 0)); // xy = coordinates, z = mip level
+	uint ComponentIdUint = asuint(ComponentIdSample);
+	// if we're not the authoritative component for this pixel, bail out and don't render anything :
+	if (ComponentIdUint != QuadComponentId)
+	{
+		discard;
+	}
+
 	BRANCH
 	switch (SourceTextureIndex)
 	{
@@ -187,7 +202,7 @@ void CopyQuadsMultiSourcePS(float4 SvPosition : SV_POSITION, float2 SourceUV : T
 
 #if COPY_WEIGHTMAP
 	// For weightmaps, copy from the appropriate channel and pack : 
-	uint SourceChannelIndex = (uint)InQuadInfos[QuadIndex].y;
+	uint SourceChannelIndex = InQuadInfos[QuadIndex].y;
 	float Weight = ColorSample[SourceChannelIndex];
 	OutColor = float4(PackWeight(Weight), PackWeightAlpha(1.0f, EWEIGHTMAPALPHAFLAGS_NONE));
 #else // COPY_WEIGHTMAP
diff --git a/Engine/Source/Runtime/Landscape/Private/LandscapeEditLayerMergeRenderContext.cpp b/Engine/Source/Runtime/Landscape/Private/LandscapeEditLayerMergeRenderContext.cpp
--- a/Engine/Source/Runtime/Landscape/Private/LandscapeEditLayerMergeRenderContext.cpp
+++ b/Engine/Source/Runtime/Landscape/Private/LandscapeEditLayerMergeRenderContext.cpp
@@ -4,6 +4,8 @@
 #include "Algo/AllOf.h"
 #include "Algo/Transform.h"
 #include "Algo/Count.h"
+#include "Algo/ForEach.h"
+#include "Algo/Sort.h"
 #include "Engine/Engine.h"
 #include "Engine/TextureRenderTarget2D.h"
 #include "Engine/TextureRenderTarget2DArray.h"
@@ -37,12 +39,14 @@ namespace UE::Landscape::EditLayers
 // ----------------------------------------------------------------------------------
 // /Engine/Private/Landscape/LandscapeEditLayersUtils.usf shaders : 
 
-class FMarkValidityPS : public FGlobalShader
+class FRenderComponentIdPS : public FGlobalShader
 {
-	DECLARE_GLOBAL_SHADER(FMarkValidityPS);
-	SHADER_USE_PARAMETER_STRUCT(FMarkValidityPS, FGlobalShader);
+	DECLARE_GLOBAL_SHADER(FRenderComponentIdPS);
+	SHADER_USE_PARAMETER_STRUCT(FRenderComponentIdPS, FGlobalShader);
 
 	BEGIN_SHADER_PARAMETER_STRUCT(FParameters, )
+		SHADER_PARAMETER_STRUCT_REF(FViewUniformShaderParameters, View)
+		SHADER_PARAMETER_RDG_BUFFER_SRV(Buffer<uint>, InQuadComponentIds)
 	END_SHADER_PARAMETER_STRUCT()
 
 	using FPermutationDomain = TShaderPermutationDomain<>;
@@ -54,15 +58,15 @@ class FMarkValidityPS : public FGlobalShader
 
 	static void ModifyCompilationEnvironment(const FGlobalShaderPermutationParameters& Parameters, FShaderCompilerEnvironment& OutEnvironment)
 	{
-		OutEnvironment.SetDefine(TEXT("MARK_VALIDITY"), 1);
+		OutEnvironment.SetDefine(TEXT("RENDER_COMPONENTID"), 1);
 	}
 };
 
-IMPLEMENT_GLOBAL_SHADER(FMarkValidityPS, "/Engine/Private/Landscape/LandscapeEditLayersUtils.usf", "MarkValidityPS", SF_Pixel);
+IMPLEMENT_GLOBAL_SHADER(FRenderComponentIdPS, "/Engine/Private/Landscape/LandscapeEditLayersUtils.usf", "RenderComponentIdPS", SF_Pixel);
 
-BEGIN_SHADER_PARAMETER_STRUCT(FMarkValidityPSParameters, )
+BEGIN_SHADER_PARAMETER_STRUCT(FRenderComponentIdPSParameters, )
 	SHADER_PARAMETER_STRUCT_INCLUDE(FPixelShaderUtils::FRasterizeToRectsVS::FParameters, VS)
-	SHADER_PARAMETER_STRUCT_INCLUDE(FMarkValidityPS::FParameters, PS)
+	SHADER_PARAMETER_STRUCT_INCLUDE(FRenderComponentIdPS::FParameters, PS)
 	RENDER_TARGET_BINDING_SLOTS()
 END_SHADER_PARAMETER_STRUCT()
 
@@ -224,14 +228,19 @@ FRenderParams::FRenderParams(
 
 bool FMergeRenderBatch::operator<(const FMergeRenderBatch& InOther) const
 {
-	// Sort by coordinates for making debugging more "logical" : 
-	if (MinComponentKey.Y < InOther.MinComponentKey.Y)
+	// Sort batches in a deterministic order (upper-left -> upper-right -> lower-left- -> lower-right) to ensure that components with a conflict on their
+	//  borders get rendered in a deterministic result : 
+	uint32 LinearIndex = ComputeComponentLinearIndex(MinComponentKey);
+	uint32 OtherLinearIndex = InOther.ComputeComponentLinearIndex(MinComponentKey);
+	if (LinearIndex < OtherLinearIndex)
 	{
 		return true;
 	}
-	else if (MinComponentKey.Y == InOther.MinComponentKey.Y)
+	else if (LinearIndex == OtherLinearIndex)
 	{
-		return (MinComponentKey.X < InOther.MinComponentKey.X);
+		LinearIndex = ComputeComponentLinearIndex(MaxComponentKey);
+		OtherLinearIndex = InOther.ComputeComponentLinearIndex(MaxComponentKey);
+		return LinearIndex < OtherLinearIndex;
 	}
 	return false;
 }
@@ -330,6 +339,13 @@ void FMergeRenderBatch::ComputeAllSubsectionRects(TArray<FIntRect>& OutSubsectio
 	}
 }
 
+uint32 FMergeRenderBatch::ComputeComponentLinearIndex(const FIntPoint& InComponentKey) const
+{
+	check((InComponentKey.X >= MinComponentKey.X) && (InComponentKey.Y >= MinComponentKey.Y));
+	check((InComponentKey.X <= MaxComponentKey.X) && (InComponentKey.Y <= MaxComponentKey.Y));
+	return (InComponentKey.X - MinComponentKey.X) + (InComponentKey.Y - MinComponentKey.Y) * (MaxComponentKey.X - MinComponentKey.X + 1);
+}
+
 FIntPoint FMergeRenderBatch::GetRenderTargetResolution(bool bInWithDuplicateBorders) const
 {
 	return bInWithDuplicateBorders ? Resolution : SectionRect.Size();
@@ -351,7 +367,8 @@ FMergeRenderContext::~FMergeRenderContext()
 {
 	FreeResources();
 
-	checkf(Algo::AllOf(BlendRenderTargets, [](ULandscapeScratchRenderTarget* InRT) { return (InRT == nullptr); }), TEXT("Every scratch render target should have been freed at this point."));
+	checkf(Algo::AllOf(BlendRenderTargets, [](ULandscapeScratchRenderTarget* InRT) { return (InRT == nullptr); }) 
+		&& (ComponentIdRenderTarget == nullptr), TEXT("Every scratch render target should have been freed at this point."));
 }
 
 void FMergeRenderContext::AllocateResources()
@@ -385,8 +402,8 @@ void FMergeRenderContext::AllocateResources()
 	checkf(Algo::AllOf(BlendRenderTargets, [](ULandscapeScratchRenderTarget* InRT) { return (InRT == nullptr); }), TEXT("We shouldn't allocate without having freed first."));
 	check(CurrentBlendRenderTargetWriteIndex == -1);
 
-	// We need N render targets large enough to fit all batches : 
 	{
+		// We need N render targets large enough to fit all batches : 
 		// Write : 
 		FScratchRenderTargetParams ScratchRenderTargetParams(TEXT("ScratchRT0"), /*bInExactDimensions = */false, /*bInUseUAV = */false, /*bInTargetArraySlicesIndependently = */(NumSlices > 0),
 			MaxNeededResolution, NumSlices, RenderTargetFormat, RenderTargetClearColor, ERHIAccess::RTV);
@@ -398,6 +415,25 @@ void FMergeRenderContext::AllocateResources()
 		ScratchRenderTargetParams.DebugName = TEXT("ScratchRT2");
 		BlendRenderTargets[2] = LandscapeEditResourcesSubsystem->RequestScratchRenderTarget(ScratchRenderTargetParams);
 	}
+
+	{
+		// We need a render target to store the component linear Ids within a given batch. i.e. for a batch dealing with N x M components : 
+		//  (0, 0) ---------- (N - 1, 0)  --> linear Ids --> 0 -------------------------- (N - 1)
+		//  |                          |                     |                                  |
+		//  |                          |                     |                                  |
+		//  (0, M - 1) -- (N - 1, M - 1)                     (M - 1) * N -- (N - 1 + (M - 1) * N)
+		check(ComponentIdRenderTarget == nullptr);
+		check((MaxNumComponents.X > 0) && (MaxNumComponents.Y > 0));
+		checkf(static_cast<uint64>(MaxNumComponents.X) * static_cast<uint64>(MaxNumComponents.Y) < static_cast<uint64>(MAX_uint32 - 1), TEXT("Cannot render %d x %d components in a given batch : the component index wouldn't fit on 32 bits"), MaxNumComponents.X, MaxNumComponents.Y);
+
+		static const uint32 NoComponentId = 0xffffffff;
+		static const FLinearColor ClearValue(*reinterpret_cast<const float*>(&NoComponentId), 0.0f, 0.0f, 0.0f);
+
+		// Using a uint32 for component ids is a tad overkill, but the cost is minimal and this avoids issues with very large batches : 
+		FScratchRenderTargetParams ScratchRenderTargetParams(TEXT("ComponentIdRT"), /*bInExactDimensions = */false, /*bInUseUAV = */false, /*bInTargetArraySlicesIndependently = */false,
+			MaxNeededResolution, 0, ETextureRenderTargetFormat::RTF_R32f, ClearValue, ERHIAccess::RTV);
+		ComponentIdRenderTarget = LandscapeEditResourcesSubsystem->RequestScratchRenderTarget(ScratchRenderTargetParams);
+	}
 }
 
 void FMergeRenderContext::FreeResources()
@@ -416,48 +452,24 @@ void FMergeRenderContext::FreeResources()
 			BlendRenderTargets[Index] = nullptr;
 		}
 	}
-	
 	CurrentBlendRenderTargetWriteIndex = -1;
+
+	if (ComponentIdRenderTarget != nullptr)
+	{
+		// We can now return the scratch render targets to the pool : 
+		LandscapeEditResourcesSubsystem->ReleaseScratchRenderTarget(ComponentIdRenderTarget);
+		ComponentIdRenderTarget = nullptr;
+	}
 }
 
 void FMergeRenderContext::AllocateBatchResources(const FMergeRenderBatch& InRenderBatch)
 {
 	TRACE_CPUPROFILER_EVENT_SCOPE(FMergeRenderContext::AllocateBatchResources);
-
-	using namespace UE::Landscape;
-
-	// Prepare the transient render resources we'll need for this batch:
-	ULandscapeEditResourcesSubsystem* LandscapeEditResourcesSubsystem = GEngine->GetEngineSubsystem<ULandscapeEditResourcesSubsystem>();
-	check(LandscapeEditResourcesSubsystem != nullptr);
-	check(PerTargetLayerValidityRenderTargets.IsEmpty());
-
-	// We need a RT version of the stencil buffer, one per target layer, to let users sample it as a standard texture :
-	int32 VisibilityScratchRTIndex = 0;
-	ForEachTargetLayer(InRenderBatch.TargetLayerBitIndices,
-		[this, LandscapeEditResourcesSubsystem, &VisibilityScratchRTIndex](int32 InTargetLayerIndex, const FName& InTargetLayerName, ULandscapeLayerInfoObject* InWeightmapLayerInfo)
-		{
-			FScratchRenderTargetParams ScratchRenderTargetParams(FString::Printf(TEXT("VisibilityScratchRT(%i)"), VisibilityScratchRTIndex), /*bInExactDimensions = */false, /*bInUseUAV = */false, /*bInTargetArraySlicesIndependently = */false,
-				MaxNeededResolution, 0, ETextureRenderTargetFormat::RTF_R8, FLinearColor::Black, ERHIAccess::RTV);
-			ULandscapeScratchRenderTarget* RenderTarget = LandscapeEditResourcesSubsystem->RequestScratchRenderTarget(ScratchRenderTargetParams);
-			PerTargetLayerValidityRenderTargets.FindOrAdd(InTargetLayerName) = RenderTarget;
-			++VisibilityScratchRTIndex;
-			return true;
-		});
 }
 
 void FMergeRenderContext::FreeBatchResources(const FMergeRenderBatch& InRenderBatch)
 {
 	TRACE_CPUPROFILER_EVENT_SCOPE(FMergeRenderContext::FreeBatchResources);
-
-	ULandscapeEditResourcesSubsystem* LandscapeEditResourcesSubsystem = GEngine->GetEngineSubsystem<ULandscapeEditResourcesSubsystem>();
-	check(LandscapeEditResourcesSubsystem != nullptr);
-
-	// We can now return those scratch render targets to the pool : 
-	for (auto ItLayerNameRenderTargetPair : PerTargetLayerValidityRenderTargets)
-	{
-		LandscapeEditResourcesSubsystem->ReleaseScratchRenderTarget(ItLayerNameRenderTargetPair.Value);
-	}
-	PerTargetLayerValidityRenderTargets.Empty();
 }
 
 void FMergeRenderContext::CycleBlendRenderTargets(UE::Landscape::FRDGBuilderRecorder& RDGBuilderRecorder)
@@ -493,10 +505,10 @@ ULandscapeScratchRenderTarget* FMergeRenderContext::GetBlendRenderTargetReadPrev
 	return BlendRenderTargets[(CurrentBlendRenderTargetWriteIndex + BlendRenderTargets.Num() - 2) % BlendRenderTargets.Num()];
 }
 
-ULandscapeScratchRenderTarget* FMergeRenderContext::GetValidityRenderTarget(const FName& InTargetLayerName) const
+ULandscapeScratchRenderTarget* FMergeRenderContext::GetComponentIdRenderTarget() const
 {
-	check(PerTargetLayerValidityRenderTargets.Contains(InTargetLayerName));
-	return PerTargetLayerValidityRenderTargets[InTargetLayerName];
+	check(ComponentIdRenderTarget != nullptr);
+	return ComponentIdRenderTarget;
 }
 
 FTransform FMergeRenderContext::ComputeVisualLogTransform(const FTransform& InTransform) const
@@ -516,90 +528,116 @@ void FMergeRenderContext::ResetVisualLogOffset()
 	CurrentVisualLogOffset = FVector(ForceInitToZero);
 }
 
-void FMergeRenderContext::RenderValidityRenderTargets(UE::Landscape::FRDGBuilderRecorder& RDGBuilderRecorder)
+void FMergeRenderContext::RenderComponentIds(UE::Landscape::FRDGBuilderRecorder& RDGBuilderRecorder)
 {
 	using namespace UE::Landscape;
 
 	const FMergeRenderBatch* RenderBatch = GetCurrentRenderBatch();
 	check(RenderBatch != nullptr);
+	check(ComponentIdRenderTarget != nullptr);
+	FSceneInterface* SceneInterface = Landscape->GetWorld()->Scene;
 
-	struct FTextureAndRects
+	struct FQuadParams
 	{
-		FTextureAndRects(const FName& InTargetLayerName, const FString& InTextureDebugName, FTextureResource* InTextureResource)
-			: TargetLayerName(InTargetLayerName)
-			, TextureDebugName(InTextureDebugName)
-			, TextureResource(InTextureResource)
-		{}
-
-		FName TargetLayerName;
-		FString TextureDebugName;
-		FTextureResource* TextureResource = nullptr;
-		TArray<FUintVector4> Rects;
+		// Sort key for rendering overlapping quads in a deterministic way :
+		uint32 ComponentLinearIndex = INDEX_NONE;
+		// Texture region to write to :
+		FIntRect DestinationRect;
 	};
-
-	TArray<FTextureAndRects> TexturesAndRects;
-	TexturesAndRects.Reserve(RenderBatch->TargetLayerBitIndices.CountSetBits());
-	ForEachTargetLayer(RenderBatch->TargetLayerBitIndices,
-		[this, &RDGBuilderRecorder, &TexturesAndRects, RenderBatch](int32 InTargetLayerIndex, const FName& InTargetLayerName, ULandscapeLayerInfoObject* InWeightmapLayerInfo)
+	TArray<FQuadParams> QuadParams;
+	QuadParams.Reserve(RenderBatch->ComponentsToRender.Num() * Landscape->NumSubsections * Landscape->NumSubsections);
+	for (ULandscapeComponent* Component : RenderBatch->ComponentsToRender)
+	{
+		FQuadParams& Quad = QuadParams.Emplace_GetRef();
+		Quad.ComponentLinearIndex = RenderBatch->ComputeComponentLinearIndex(Component->GetComponentKey());
+		check(Quad.ComponentLinearIndex < static_cast<uint64>(MaxNumComponents.X * MaxNumComponents.Y));
+		Quad.DestinationRect = RenderBatch->ComputeSectionRect(Component, /*bInWithDuplicateBorders = */false);
+	}
+	// We want to render quads for each component and we want to render them in a deterministic order : 
+	Algo::SortBy(QuadParams, &FQuadParams::ComponentLinearIndex);
+
+	TArray<FUintVector4> QuadRects;
+	TArray<uint32> QuadComponentLinearIds;
+	QuadRects.Reserve(QuadParams.Num());
+	QuadComponentLinearIds.Reserve(QuadParams.Num());
+	Algo::ForEach(QuadParams, [&QuadRects, &QuadComponentLinearIds](const FQuadParams& InQuadParams)
 		{
-			ULandscapeScratchRenderTarget* ScratchRenderTarget = PerTargetLayerValidityRenderTargets.FindChecked(InTargetLayerName);
-			check(ScratchRenderTarget != nullptr);
-
-			// Make sure the validity mask is entirely cleared first:
-			ScratchRenderTarget->Clear(RDGBuilderRecorder);
-
-			FTextureAndRects& TextureAndRects = TexturesAndRects.Emplace_GetRef(InTargetLayerName, ScratchRenderTarget->GetDebugName(), ScratchRenderTarget->GetRenderTarget2D()->GetResource());
-
-			// Then build a list of quads for marking where the components are valid for this target layer on this batch:
-			const TSet<ULandscapeComponent*>& Components = RenderBatch->TargetLayersToComponents[InTargetLayerIndex];
-			TextureAndRects.Rects.Reserve(Components.Num());
-			for (ULandscapeComponent* Component : Components)
-			{
-				FIntRect ComponentRect = RenderBatch->ComputeSectionRect(Component, /*bInWithDuplicateBorders = */false);
-				TextureAndRects.Rects.Add(FUintVector4(ComponentRect.Min.X, ComponentRect.Min.Y, ComponentRect.Max.X + 1, ComponentRect.Max.Y + 1));
-			}
-
-			ScratchRenderTarget->TransitionTo(ERHIAccess::RTV, RDGBuilderRecorder);
-			return true;
+			QuadRects.Emplace(InQuadParams.DestinationRect.Min.X, InQuadParams.DestinationRect.Min.Y, InQuadParams.DestinationRect.Max.X + 1, InQuadParams.DestinationRect.Max.Y + 1);
+			QuadComponentLinearIds.Emplace(InQuadParams.ComponentLinearIndex);
 		});
 
+	ComponentIdRenderTarget->Clear(RDGBuilderRecorder);
+	ComponentIdRenderTarget->TransitionTo(ERHIAccess::RTV, RDGBuilderRecorder);
+
 	auto RDGCommand =
-		[TexturesAndRects = MoveTemp(TexturesAndRects)](FRDGBuilder& GraphBuilder)
+		[ComponentIdRenderTargetResource = ComponentIdRenderTarget->GetRenderTarget2D()->GetResource(), 
+		SceneInterface,
+		QuadRects = MoveTemp(QuadRects), 
+		QuadComponentLinearIds = MoveTemp(QuadComponentLinearIds)]
+		(FRDGBuilder& GraphBuilder)
 	{
-		for (const FTextureAndRects& TextureAndRects : TexturesAndRects)
-		{
-			FRDGBufferRef RectBuffer = CreateUploadBuffer(GraphBuilder, TEXT("MarkValidityRects"), MakeArrayView(TextureAndRects.Rects));
+			FRDGBufferRef RectBuffer = CreateUploadBuffer(GraphBuilder, TEXT("ComponentIdRects"), MakeArrayView(QuadRects));
 			FRDGBufferSRVRef RectBufferSRV = GraphBuilder.CreateSRV(FRDGBufferSRVDesc(RectBuffer, PF_R32G32B32A32_UINT));
-			FRDGTextureRef OutputTexture = GraphBuilder.RegisterExternalTexture(CreateRenderTarget(TextureAndRects.TextureResource->GetTexture2DRHI(), TEXT("ValidityMask")));
+			FRDGBufferRef ComponentIdBuffer = CreateUploadBuffer(GraphBuilder, TEXT("ComponentIds"), MakeArrayView(QuadComponentLinearIds));
+			FRDGBufferSRVRef ComponentIdBufferSRV = GraphBuilder.CreateSRV(FRDGBufferSRVDesc(ComponentIdBuffer, PF_R32_UINT));
+			FRDGTextureRef OutputTexture = GraphBuilder.RegisterExternalTexture(CreateRenderTarget(ComponentIdRenderTargetResource->GetTexture2DRHI(), TEXT("ComponentIds")));
+
+			// We need a depth texture to ensure overlapping quads render in a deterministic order :
+			FRDGTextureDesc DepthTextureDesc = FRDGTextureDesc::Create2D(OutputTexture->Desc.Extent, PF_DepthStencil, FClearValueBinding::DepthOne, TexCreate_DepthStencilTargetable);
+			FRDGTextureRef DepthTexture = GraphBuilder.CreateTexture(DepthTextureDesc, TEXT("ComponentIdRectsDepth"));
+
+			// TODO [jonathan.bard] this is just to avoid a RHI validation error for unoptimized shaders... once validation is made to not issue those errors, we can remove this
+			// Create a SceneView to please the shader bindings, but it's unused in practice 
+			FSceneViewFamilyContext ViewFamily(FSceneViewFamily::ConstructionValues(nullptr, SceneInterface, FEngineShowFlags(ESFIM_Game)).SetTime(FGameTime::GetTimeSinceAppStart()));
+			FSceneViewInitOptions ViewInitOptions;
+			ViewInitOptions.ViewFamily = &ViewFamily;
+			ViewInitOptions.SetViewRectangle(FIntRect(0, 0, 1, 1)); // Use a dummy rect to avoid a check(slow)
+			GetRendererModule().CreateAndInitSingleView(GraphBuilder.RHICmdList, &ViewFamily, &ViewInitOptions);
+			const FSceneView* View = ViewFamily.Views[0];
 
-			FMarkValidityPSParameters* PassParameters = GraphBuilder.AllocParameters<FMarkValidityPSParameters>();
+			FRenderComponentIdPSParameters* PassParameters = GraphBuilder.AllocParameters<FRenderComponentIdPSParameters>();
+			PassParameters->RenderTargets.DepthStencil = FDepthStencilBinding(DepthTexture,
+				/*InDepthLoadAction = */ ERenderTargetLoadAction::EClear,
+				/*InStencilLoadAction = */ERenderTargetLoadAction::ENoAction,
+				/*InDepthStencilAccess = */FExclusiveDepthStencil::DepthWrite_StencilNop);
 			PassParameters->RenderTargets[0] = FRenderTargetBinding(OutputTexture, ERenderTargetLoadAction::ELoad);
+			PassParameters->PS.View = View->ViewUniformBuffer;
+			PassParameters->PS.InQuadComponentIds = ComponentIdBufferSRV;
 			FGlobalShaderMap* ShaderMap = GetGlobalShaderMap(GMaxRHIFeatureLevel);
-			TShaderRef<FMarkValidityPS> PixelShader = ShaderMap->GetShader<FMarkValidityPS>();
+			TShaderRef<FRenderComponentIdPS> PixelShader = ShaderMap->GetShader<FRenderComponentIdPS>();
+
+			// There is some overlap between quads of 2 neighboring components at this step (e.g. the last column of component (X, Y) is supposed to correspond to the first of component (X+1, Y))
+			//  Normally, this is fine but if both components don't agree on the values of those overlapped lines/columns (e.g. one component was submitted with updated values on the edge but
+			//  not the other), then we still need to ensure a deterministic value to be output. Since quads are sorted by coordinates and we can get AddRasterizeToRectsPass to render quads in 
+			//  a consistent near to far order (if DepthIncrement is passed), we can ensure a consistent ordering by enabling the depth test :
+			FRHIDepthStencilState* DepthStencilState = TStaticDepthStencilState</*bEnableDepthWrite = */true, CF_LessEqual>::GetRHI();
+
+			// Spread out the rects across the [0, 1[ range in case this is used with the depth test enabled (e.g. to ensure a deterministic result on overlapping quads) : 
+			//  Divide the Z space in NumQuads increments, so that the first rect stands at Z == 0 and the last one at Z == 1.0 - DepthIncrement (avoids standing exactly on the far clipping plane,
+			//  which could be problematic for precision reasons)
+			float DepthIncrement = 1.0f / QuadRects.Num();
 
-			FPixelShaderUtils::AddRasterizeToRectsPass<FMarkValidityPS>(GraphBuilder,
+			FPixelShaderUtils::AddRasterizeToRectsPass<FRenderComponentIdPS>(GraphBuilder,
 				ShaderMap,
-				RDG_EVENT_NAME("MarkValidity(%s) -> %s", *TextureAndRects.TargetLayerName.ToString(), *TextureAndRects.TextureDebugName),
+				RDG_EVENT_NAME("RenderComponentId"),
 				PixelShader,
 				PassParameters,
 				/*ViewportSize = */OutputTexture->Desc.Extent,
 				RectBufferSRV,
-				TextureAndRects.Rects.Num(),
+				QuadRects.Num(),
 				/*BlendState = */ nullptr,
 				/*RasterizerState = */ nullptr,
-				/*DepthStencilState = */ nullptr,
+				DepthStencilState,
 				/*StencilRef = */ 0,
-				/*TextureSize = */ OutputTexture->Desc.Extent);
-		}
+				/*TextureSize = */ OutputTexture->Desc.Extent,
+				/*RectUVBufferSRV = */nullptr,
+				/*DownsampleFactor = */1,
+				/*AdditionalPassFlags = */ERDGPassFlags::None,
+				DepthIncrement);
 		};
 
 	// We need to specify the final state of the external textures to prevent the graph builder from transitioning them to SRVMask :
-	TArray<FRDGBuilderRecorder::FRDGExternalTextureAccessFinal> RDGExternalTextureAccessFinalList;
-	Algo::Transform(TexturesAndRects, RDGExternalTextureAccessFinalList, [](const FTextureAndRects& InTextureAndRects) -> FRDGBuilderRecorder::FRDGExternalTextureAccessFinal
-		{ 
-			return { InTextureAndRects.TextureResource, ERHIAccess::RTV };
-	});
-	RDGBuilderRecorder.EnqueueRDGCommand(RDGCommand, RDGExternalTextureAccessFinalList);
+	RDGBuilderRecorder.EnqueueRDGCommand(RDGCommand, { { ComponentIdRenderTarget->GetRenderTarget2D()->GetResource(), ERHIAccess::RTV } });
 }
 
 void FMergeRenderContext::RenderExpandedRenderTarget(UE::Landscape::FRDGBuilderRecorder& RDGBuilderRecorder)
diff --git a/Engine/Source/Runtime/Landscape/Private/LandscapeEditLayerRendererPrivate.cpp b/Engine/Source/Runtime/Landscape/Private/LandscapeEditLayerRendererPrivate.cpp
--- a/Engine/Source/Runtime/Landscape/Private/LandscapeEditLayerRendererPrivate.cpp
+++ b/Engine/Source/Runtime/Landscape/Private/LandscapeEditLayerRendererPrivate.cpp
@@ -60,7 +60,7 @@ bool ULandscapeDefaultEditLayerRenderer::RenderLayer(UE::Landscape::EditLayers::
 	check(WriteRT->GetCurrentState() == ERHIAccess::RTV);
 
 	// Render the components of the batch for each target layer into the "pseudo-stencil" buffer, so that it can be sampled by users as a UTextures in UMaterials and such : 
-	RenderContext->RenderValidityRenderTargets(RDGBuilderRecorder);
+	RenderContext->RenderComponentIds(RDGBuilderRecorder);
 
 	return true; 
 }
diff --git a/Engine/Source/Runtime/Landscape/Private/LandscapeEditLayers.cpp b/Engine/Source/Runtime/Landscape/Private/LandscapeEditLayers.cpp
--- a/Engine/Source/Runtime/Landscape/Private/LandscapeEditLayers.cpp
+++ b/Engine/Source/Runtime/Landscape/Private/LandscapeEditLayers.cpp
@@ -1720,7 +1720,7 @@ public:
 		SHADER_PARAMETER(FUintVector2, InNumComponents)
 		SHADER_PARAMETER_SAMPLER(SamplerState, InSourceHeightmapSampler)
 		SHADER_PARAMETER_RDG_TEXTURE_SRV(Texture2D<float4>, InSourceHeightmap)
-		SHADER_PARAMETER_RDG_TEXTURE_SRV(Texture2D<uint8>, InValidityTexture)
+		SHADER_PARAMETER_RDG_TEXTURE_SRV(Texture2D<float>, InComponentIdTexture)
 		RENDER_TARGET_BINDING_SLOTS()
 	END_SHADER_PARAMETER_STRUCT()
 
@@ -1925,14 +1925,15 @@ IMPLEMENT_GLOBAL_SHADER(FLandscapeEditLayersWeightmapsGenerateMipsPS, "/Engine/P
 class FCopyQuadsMultiSourcePSDefault : public FGlobalShader
 {
 public:
-	static constexpr int32 NumMultiSources = 63;
+	static constexpr int32 NumMultiSources = 62;
 
 	DECLARE_GLOBAL_SHADER(FCopyQuadsMultiSourcePSDefault);
 	SHADER_USE_PARAMETER_STRUCT(FCopyQuadsMultiSourcePSDefault, FGlobalShader);
 
 	BEGIN_SHADER_PARAMETER_STRUCT(FParameters, )
 		SHADER_PARAMETER_STRUCT_REF(FViewUniformShaderParameters, View)
-		SHADER_PARAMETER_RDG_BUFFER_SRV(Buffer<uint2>, InQuadInfos)
+		SHADER_PARAMETER_RDG_BUFFER_SRV(Buffer<uint4>, InQuadInfos)
+		SHADER_PARAMETER_RDG_TEXTURE_SRV(Texture2D, InComponentIdTexture)
 		SHADER_PARAMETER_RDG_TEXTURE_ARRAY(Texture2D, InSourceTexture, [NumMultiSources])
 	END_SHADER_PARAMETER_STRUCT()
 
@@ -1953,14 +1954,15 @@ public:
 class FCopyQuadsMultiSourcePSVulkanSM6 : public FGlobalShader
 {
 public:
-	static constexpr int32 NumMultiSources = 29;
+	static constexpr int32 NumMultiSources = 28;
 
 	DECLARE_GLOBAL_SHADER(FCopyQuadsMultiSourcePSVulkanSM6);
 	SHADER_USE_PARAMETER_STRUCT(FCopyQuadsMultiSourcePSVulkanSM6, FGlobalShader);
 
 	BEGIN_SHADER_PARAMETER_STRUCT(FParameters, )
 		SHADER_PARAMETER_STRUCT_REF(FViewUniformShaderParameters, View)
-		SHADER_PARAMETER_RDG_BUFFER_SRV(Buffer<uint2>, InQuadInfos)
+		SHADER_PARAMETER_RDG_BUFFER_SRV(Buffer<uint4>, InQuadInfos)
+		SHADER_PARAMETER_RDG_TEXTURE_SRV(Texture2D, InComponentIdTexture)
 		SHADER_PARAMETER_RDG_TEXTURE_ARRAY(Texture2D, InSourceTexture, [NumMultiSources])
 	END_SHADER_PARAMETER_STRUCT()
 
@@ -1990,6 +1992,7 @@ public:
 		FRDGBufferSRVRef InRectBufferSRV,
 		FRDGBufferSRVRef InQuadInfosBufferSRV,
 		FRDGBufferSRVRef InRectUVBufferSRV,
+		FRDGTextureSRVRef InComponentIdTextureSRV,
 		const FSceneView* InView,
 		int32 InArrayIndex,
 		int32 InNumRects,
@@ -2014,6 +2017,7 @@ public:
 		FRDGBufferSRVRef InRectBufferSRV, 
 		FRDGBufferSRVRef InQuadInfosBufferSRV, 
 		FRDGBufferSRVRef InRectUVBufferSRV,
+		FRDGTextureSRVRef InComponentIdTextureSRV,
 		const FSceneView* InView, 
 		int32 InArrayIndex,
 		int32 InNumRects, 
@@ -2029,6 +2033,7 @@ public:
 		PassParameters->RenderTargets[0] = FRenderTargetBinding(InOutputTexture, ERenderTargetLoadAction::ELoad, /*InMipIndex = */0, InArrayIndex);
 		PassParameters->PS.View = InView->ViewUniformBuffer;
 		PassParameters->PS.InQuadInfos = InQuadInfosBufferSRV;
+		PassParameters->PS.InComponentIdTexture = InComponentIdTextureSRV;
 		check(InSourceTextures.Num() == NumMultiSources);
 		for (int32 TextureIndex = 0; TextureIndex < NumMultiSources; ++TextureIndex)
 		{
@@ -2060,9 +2065,9 @@ END_SHADER_PARAMETER_STRUCT()
 IMPLEMENT_GLOBAL_SHADER(FCopyQuadsMultiSourcePSDefault, "/Engine/Private/Landscape/LandscapeEditLayersUtils.usf", "CopyQuadsMultiSourcePS", SF_Pixel);
 
 BEGIN_SHADER_PARAMETER_STRUCT(FCopyQuadsMultiSourcePSParametersVulkanSM6, )
-	SHADER_PARAMETER_STRUCT_INCLUDE(FPixelShaderUtils::FRasterizeToRectsVS::FParameters, VS)
+SHADER_PARAMETER_STRUCT_INCLUDE(FPixelShaderUtils::FRasterizeToRectsVS::FParameters, VS)
 	SHADER_PARAMETER_STRUCT_INCLUDE(FCopyQuadsMultiSourcePSVulkanSM6::FParameters, PS)
-	RENDER_TARGET_BINDING_SLOTS()
+RENDER_TARGET_BINDING_SLOTS()
 END_SHADER_PARAMETER_STRUCT()
 IMPLEMENT_GLOBAL_SHADER(FCopyQuadsMultiSourcePSVulkanSM6, "/Engine/Private/Landscape/LandscapeEditLayersUtils.usf", "CopyQuadsMultiSourcePS", SF_Pixel);
 
@@ -5635,6 +5640,8 @@ bool ULandscapeEditLayerPersistent::RenderLayer(UE::Landscape::EditLayers::FRend
 		int32 DestinationArrayIndex = 0;
 		// Weightmaps are packed in the source texture. This specifies which channel to read from in the source texture for this quad :
 		uint8 SourceChannelIndex = 0;
+		// Linear index of this component within this batch. Allows to resolve conflicts in case 2 components render to the same pixel on a border, for example :
+		uint32 ComponentId = INDEX_NONE;
 	};
 	TArray<FCopyQuadParams> CopyQuadsParams;
 	CopyQuadsParams.Reserve(RenderParams.SortedComponentMergeRenderInfos.Num() * NumTargetLayersInGroup * Landscape->NumSubsections * Landscape->NumSubsections);
@@ -5647,6 +5654,7 @@ bool ULandscapeEditLayerPersistent::RenderLayer(UE::Landscape::EditLayers::FRend
 	// Build the full list of quads to copy :
 	for (const FComponentMergeRenderInfo& ComponentMergeRenderInfo : RenderParams.SortedComponentMergeRenderInfos)
 	{
+		const uint32 ComponentId = RenderParams.MergeRenderContext->GetCurrentRenderBatch()->ComputeComponentLinearIndex(ComponentMergeRenderInfo.Component->GetComponentKey());
 		for (int32 TargetLayerIndexInGroup = 0; TargetLayerIndexInGroup < NumTargetLayersInGroup; ++TargetLayerIndexInGroup)
 		{
 			const FName TargetLayerName = RenderParams.TargetLayerGroupLayerNames[TargetLayerIndexInGroup];
@@ -5705,6 +5713,7 @@ bool ULandscapeEditLayerPersistent::RenderLayer(UE::Landscape::EditLayers::FRend
 						QuadParams.DestinationRect = FIntRect(DestinationPosition, DestinationPosition + FIntPoint(SubsectionSizeVerts, SubsectionSizeVerts));
 						QuadParams.DestinationArrayIndex = TargetLayerIndexInGroup;
 						QuadParams.SourceChannelIndex = SourceChannelIndex;
+						QuadParams.ComponentId = ComponentId;
 					}
 				}
 			}
@@ -5725,6 +5734,10 @@ bool ULandscapeEditLayerPersistent::RenderLayer(UE::Landscape::EditLayers::FRend
 	// We will write to the RT using a (bunch of) PS : 
 	check(WriteRT->GetCurrentState() == ERHIAccess::RTV);
 
+	ULandscapeScratchRenderTarget* ComponentIdRT = RenderParams.MergeRenderContext->GetComponentIdRenderTarget();
+	check(ComponentIdRT != nullptr);
+	ComponentIdRT->TransitionTo(ERHIAccess::SRVMask, RDGBuilderRecorder);
+
 	// In this render step, we'll simply render the edit layer's component quads
 
 	// Sort the list of quads by output texture first, then by source texture, because we'll coalesce several texture copies into the same pass (one output, multiple inputs)
@@ -5746,8 +5759,9 @@ bool ULandscapeEditLayerPersistent::RenderLayer(UE::Landscape::EditLayers::FRend
 
 	// Now process this list of quads and prepare as many passes as necessary for performing all the copies :
 	auto RDGCommand =
-		[ CopyQuadsParams = MoveTemp(CopyQuadsParams)
+		[CopyQuadsParams = MoveTemp(CopyQuadsParams)
 		, OutputResource = WriteRT->GetRenderTarget()->GetResource()
+		, ComponentIdResource = ComponentIdRT->GetRenderTarget()->GetResource()
 		, OutputResourceName = WriteRT->GetDebugName()
 		, SceneInterface
 		, CommonSourceTextureSize
@@ -5779,16 +5793,19 @@ bool ULandscapeEditLayerPersistent::RenderLayer(UE::Landscape::EditLayers::FRend
 			// The following variables allow to accumulate quads for rendering in a single pass. Then, we flush them all out when the pass is full and we start again : 
 			int32 CurrentDestinationArrayIndex = INDEX_NONE;
 			TArray<FTextureResource*> CurrentSourceTextureResources;
-			TArray<FUintVector2> CurrentQuadInfos;
+			TArray<FUintVector4> CurrentQuadInfos;
 			TArray<FUintVector4> CurrentSourceRects, CurrentDestinationRects;
 
 			FRDGTextureRef OutputTexture = GraphBuilder.RegisterExternalTexture(CreateRenderTarget(OutputResource->GetTextureRHI(), TEXT("OutputTexture")));
+			FRDGTextureRef ComponentIdTexture = GraphBuilder.RegisterExternalTexture(CreateRenderTarget(ComponentIdResource->GetTextureRHI(), TEXT("ComponentIdTexture")));
+			FRDGTextureSRVRef ComponentIdTextureSRV = GraphBuilder.CreateSRV(FRDGTextureSRVDesc::Create(ComponentIdTexture));
 
 			// When we've reached the limit of what we can render in one pass, flush all and add a render pass : 
 			auto FlushCopyQuads = [
 				&GraphBuilder, 
 				View, 
 				OutputTexture, 
+				ComponentIdTextureSRV,
 				&CurrentDestinationArrayIndex, 
 				&CurrentSourceTextureResources, 
 				&CurrentSourceRects, 
@@ -5808,14 +5825,14 @@ bool ULandscapeEditLayerPersistent::RenderLayer(UE::Landscape::EditLayers::FRend
 
 					check(!CurrentSourceRects.IsEmpty() && (CurrentSourceRects.Num() == CurrentDestinationRects.Num()) && (CurrentSourceRects.Num() == CurrentQuadInfos.Num()));
 
-					FRDGBufferRef RectBuffer = CreateUploadBuffer(GraphBuilder, TEXT("DestinationRects"), TConstArrayView<FUintVector4>(MakeArrayView(CurrentDestinationRects)));
+					FRDGBufferRef RectBuffer = CreateUploadBuffer(GraphBuilder, TEXT("DestinationRects"), MakeConstArrayView(CurrentDestinationRects));
 					FRDGBufferSRVRef RectBufferSRV = GraphBuilder.CreateSRV(FRDGBufferSRVDesc(RectBuffer, PF_R32G32B32A32_UINT));
 
-					FRDGBufferRef RectUVBuffer = CreateUploadBuffer(GraphBuilder, TEXT("RectUVs"), TConstArrayView<FUintVector4>(MakeArrayView(CurrentSourceRects)));
+					FRDGBufferRef RectUVBuffer = CreateUploadBuffer(GraphBuilder, TEXT("RectUVs"), MakeConstArrayView(CurrentSourceRects));
 					FRDGBufferSRVRef RectUVBufferSRV = GraphBuilder.CreateSRV(FRDGBufferSRVDesc(RectUVBuffer, PF_R32G32B32A32_UINT));
 
-					FRDGBufferRef QuadInfosBuffer = CreateUploadBuffer(GraphBuilder, TEXT("QuadInfos"), TConstArrayView<FUintVector2>(MakeArrayView(CurrentQuadInfos)));
-					FRDGBufferSRVRef QuadInfosBufferSRV = GraphBuilder.CreateSRV(FRDGBufferSRVDesc(QuadInfosBuffer, PF_R32G32_UINT));
+					FRDGBufferRef QuadInfosBuffer = CreateUploadBuffer(GraphBuilder, TEXT("QuadInfos"), MakeConstArrayView(CurrentQuadInfos));
+					FRDGBufferSRVRef QuadInfosBufferSRV = GraphBuilder.CreateSRV(FRDGBufferSRVDesc(QuadInfosBuffer, PF_R32G32B32A32_UINT)); 
 
 					// We may target an array slice directly : 
 					int32 ArrayIndex = -1;
@@ -5844,6 +5861,7 @@ bool ULandscapeEditLayerPersistent::RenderLayer(UE::Landscape::EditLayers::FRend
 						RectBufferSRV,
 						QuadInfosBufferSRV,
 						RectUVBufferSRV,
+						ComponentIdTextureSRV,
 						View,
 						ArrayIndex,
 						CurrentDestinationRects.Num(),
@@ -5892,7 +5910,7 @@ bool ULandscapeEditLayerPersistent::RenderLayer(UE::Landscape::EditLayers::FRend
 				if (InSingleQuadParams.SourceTextureResource == CurrentSourceTextureResources.Last())
 				{
 					const int32 SourceTextureIndex = CurrentSourceTextureResources.Num() - 1;
-					CurrentQuadInfos.Add(FUintVector2(static_cast<uint32>(SourceTextureIndex), InSingleQuadParams.SourceChannelIndex));
+					CurrentQuadInfos.Add(FUintVector4(static_cast<uint32>(SourceTextureIndex), InSingleQuadParams.SourceChannelIndex, InSingleQuadParams.ComponentId, 0));
 					CurrentSourceRects.Add(FUintVector4(InSingleQuadParams.SourceRect.Min.X, InSingleQuadParams.SourceRect.Min.Y, InSingleQuadParams.SourceRect.Max.X, InSingleQuadParams.SourceRect.Max.Y));
 					CurrentDestinationRects.Add(FUintVector4(InSingleQuadParams.DestinationRect.Min.X, InSingleQuadParams.DestinationRect.Min.Y, InSingleQuadParams.DestinationRect.Max.X, InSingleQuadParams.DestinationRect.Max.Y));
 				}
@@ -5902,8 +5920,14 @@ bool ULandscapeEditLayerPersistent::RenderLayer(UE::Landscape::EditLayers::FRend
 			FlushCopyQuads();
 		};
 
-	// We need to specify the final state of the external texture to prevent the graph builder from transitioning it to SRVMask :
-	RDGBuilderRecorder.EnqueueRDGCommand(RDGCommand, { { WriteRT->GetRenderTarget()->GetResource(), ERHIAccess::RTV } });
+	// We need to specify the final state of the external textures to prevent the graph builder from transitioning them to SRVMask (even those that end up as SRVMask at the end of this command,
+	//  because they will likely be part of another RDGCommand down the line so we need to maintain an accurate picture of every external texture ever involved in the recorded command so that 
+	//  we can set a proper access when the recorder is flushed (and the FRDGBuilder, executed) :
+	RDGBuilderRecorder.EnqueueRDGCommand(RDGCommand, 
+	{ 
+		{ WriteRT->GetRenderTarget()->GetResource(), ERHIAccess::RTV }, 
+		{ ComponentIdRT->GetRenderTarget()->GetResource(), ERHIAccess::SRVMask } 
+	});
 
 	// We've rendered at least a quad : 
 	return true;
@@ -6476,12 +6500,12 @@ bool ULandscapeHeightmapNormalsEditLayerRenderer::RenderLayer(UE::Landscape::Edi
 	RenderContext->CycleBlendRenderTargets(RDGBuilderRecorder);
 	ULandscapeScratchRenderTarget* WriteRT = RenderContext->GetBlendRenderTargetWrite();
 	ULandscapeScratchRenderTarget* ReadRT = RenderContext->GetBlendRenderTargetRead();
-	ULandscapeScratchRenderTarget* ValidityRT = RenderContext->GetValidityRenderTarget(RenderParams.TargetLayerGroupLayerNames[0]);
-	check((WriteRT != nullptr) && (ReadRT != nullptr) && (ValidityRT != nullptr));
+	ULandscapeScratchRenderTarget* ComponentIdRT = RenderContext->GetComponentIdRenderTarget();
+	check((WriteRT != nullptr) && (ReadRT != nullptr) && (ComponentIdRT != nullptr));
 
 	WriteRT->TransitionTo(ERHIAccess::RTV, RDGBuilderRecorder);
 	ReadRT->TransitionTo(ERHIAccess::SRVMask, RDGBuilderRecorder);
-	ValidityRT->TransitionTo(ERHIAccess::SRVMask, RDGBuilderRecorder);
+	ComponentIdRT->TransitionTo(ERHIAccess::SRVMask, RDGBuilderRecorder);
 
 	FIntPoint NumComponentsInRenderArea = RenderBatch->MaxComponentKey - RenderBatch->MinComponentKey + 1;
 	check((NumComponentsInRenderArea.X > 0) && (NumComponentsInRenderArea.Y > 0));
@@ -6490,7 +6514,7 @@ bool ULandscapeHeightmapNormalsEditLayerRenderer::RenderLayer(UE::Landscape::Edi
 		[OutputResource = WriteRT->GetRenderTarget2D()->GetResource()
 		, OutputResourceName = WriteRT->GetDebugName()
 		, SourceResource = ReadRT->GetRenderTarget2D()->GetResource()
-		, ValidityResource = ValidityRT->GetRenderTarget2D()->GetResource()
+		, ComponentIdResource = ComponentIdRT->GetRenderTarget2D()->GetResource()
 		, EffectiveTextureSize = RenderBatch->GetRenderTargetResolution(/*bInWithDuplicateBorders = */false)
 		, LandscapeGridScale = Landscape->GetRootComponent()->GetRelativeScale3D()
 		, ComponentSizeQuads = Landscape->ComponentSizeQuads
@@ -6498,7 +6522,7 @@ bool ULandscapeHeightmapNormalsEditLayerRenderer::RenderLayer(UE::Landscape::Edi
 	{
 		FRDGTextureRef OutputTextureRef = GraphBuilder.RegisterExternalTexture(CreateRenderTarget(OutputResource->TextureRHI, TEXT("OutputTexture")));
 		FRDGTextureRef SourceTextureRef = GraphBuilder.RegisterExternalTexture(CreateRenderTarget(SourceResource->TextureRHI, TEXT("SourceTexture")));
-		FRDGTextureRef ValidityTextureRef = GraphBuilder.RegisterExternalTexture(CreateRenderTarget(ValidityResource->TextureRHI, TEXT("ValidityTexture")));
+		FRDGTextureRef ComponentIdTextureRef = GraphBuilder.RegisterExternalTexture(CreateRenderTarget(ComponentIdResource->TextureRHI, TEXT("ComponentIdTexture")));
 
 		FLandscapeEditLayersHeightmapsGenerateNormalsPS::FParameters* PSParams = GraphBuilder.AllocParameters<FLandscapeEditLayersHeightmapsGenerateNormalsPS::FParameters>();
 		PSParams->RenderTargets[0] = FRenderTargetBinding(OutputTextureRef, ERenderTargetLoadAction::ENoAction);
@@ -6508,7 +6532,7 @@ bool ULandscapeHeightmapNormalsEditLayerRenderer::RenderLayer(UE::Landscape::Edi
 		PSParams->InNumComponents = FUintVector2(NumComponentsInRenderArea.X, NumComponentsInRenderArea.Y);
 		PSParams->InSourceHeightmapSampler = TStaticSamplerState<SF_Point, AM_Clamp, AM_Clamp>::GetRHI();
 		PSParams->InSourceHeightmap = GraphBuilder.CreateSRV(FRDGTextureSRVDesc::Create(SourceTextureRef));
-		PSParams->InValidityTexture = GraphBuilder.CreateSRV(FRDGTextureSRVDesc::Create(ValidityTextureRef));
+		PSParams->InComponentIdTexture = GraphBuilder.CreateSRV(FRDGTextureSRVDesc::Create(ComponentIdTextureRef));
 			FLandscapeEditLayersHeightmapsGenerateNormalsPS::GenerateNormalsPS(
 				RDG_EVENT_NAME("GenerateNormals -> %s", *OutputResourceName), GraphBuilder, PSParams, EffectiveTextureSize);
 		};
@@ -6520,7 +6544,7 @@ bool ULandscapeHeightmapNormalsEditLayerRenderer::RenderLayer(UE::Landscape::Edi
 	{
 		{ WriteRT->GetRenderTarget()->GetResource(), ERHIAccess::RTV },
 		{ ReadRT->GetRenderTarget()->GetResource(), ERHIAccess::SRVMask },
-		{ ValidityRT->GetRenderTarget()->GetResource(), ERHIAccess::SRVMask },
+		{ ComponentIdRT->GetRenderTarget()->GetResource(), ERHIAccess::SRVMask },
 	};
 	RDGBuilderRecorder.EnqueueRDGCommand(RDGCommand, RDGExternalTextureAccessFinalList);
 
@@ -7180,7 +7204,7 @@ UE::Landscape::EditLayers::FMergeRenderContext ALandscape::PrepareEditLayersMerg
 			MergeRenderBatch.TargetLayerBitIndices.SetNum(MergeRenderContext.AllTargetLayerNames.Num(), false);
 
 			MergeRenderContext.MaxNeededResolution = MergeRenderContext.MaxNeededResolution.ComponentMax(MergeRenderBatch.Resolution);
-		
+			MergeRenderContext.MaxNumComponents = MergeRenderContext.MaxNumComponents.ComponentMax(MergeRenderBatch.MaxComponentKey - MergeRenderBatch.MinComponentKey + 1);
 
 			// For each renderer, find the list of components actually needed and store that in a separate array, to avoid recomputing it for each target group : 
 			int32 LocalNumRenderers = OrderedEditLayerRendererRenderInfos.Num();
@@ -7409,7 +7433,7 @@ UE::Landscape::EditLayers::FMergeRenderContext ALandscape::PrepareEditLayersMerg
 			}
 		}
 
-		// Sort the batches for easier debugging : 
+		// Sort the batches in a deterministic order :
 		MergeRenderContext.RenderBatches.Sort();
 	}
 
diff --git a/Engine/Source/Runtime/Landscape/Private/LandscapeEditResourcesSubsystem.cpp b/Engine/Source/Runtime/Landscape/Private/LandscapeEditResourcesSubsystem.cpp
--- a/Engine/Source/Runtime/Landscape/Private/LandscapeEditResourcesSubsystem.cpp
+++ b/Engine/Source/Runtime/Landscape/Private/LandscapeEditResourcesSubsystem.cpp
@@ -190,7 +190,7 @@ void ULandscapeScratchRenderTarget::Clear(UE::Landscape::FRDGBuilderRecorder& RD
 		, EffectiveNumSlices = GetEffectiveNumSlices()]
 		(FRDGBuilder& GraphBuilder)
 		{
-			FRDGTextureRef TextureRef = GraphBuilder.RegisterExternalTexture(CreateRenderTarget(Resource->GetTextureRHI(), TEXT("ClearTexture")));
+			FRDGTextureRef TextureRef = GraphBuilder.RegisterExternalTexture(CreateRenderTarget(Resource->GetTextureRHI(), TEXT("ScratchTexture")));
 			FRDGTextureClearInfo ClearInfo;
 			if (TextureRef->Desc.IsTextureArray())
 			{
@@ -294,8 +294,14 @@ bool ULandscapeScratchRenderTarget::IsCompatibleWith(const UE::Landscape::FScrat
 		return false;
 	}
 
+	if (RenderTargetFormat != InParams.Format)
+	{
+		return false;
+	}
+
 	FLinearColor RenderTargetClearColor = GetClearColor();
-	if (RenderTargetClearColor != InParams.ClearColor)
+	// Don't use FLinearColor::operator != on purpose : the clear color needs to match exactly to be considered compatible
+	if (GetTypeHash(RenderTargetClearColor) != GetTypeHash(InParams.ClearColor))
 	{
 		return false;
 	}
diff --git a/Engine/Source/Runtime/Landscape/Public/LandscapeEditLayerMergeRenderContext.h b/Engine/Source/Runtime/Landscape/Public/LandscapeEditLayerMergeRenderContext.h
--- a/Engine/Source/Runtime/Landscape/Public/LandscapeEditLayerMergeRenderContext.h
+++ b/Engine/Source/Runtime/Landscape/Public/LandscapeEditLayerMergeRenderContext.h
@@ -205,6 +205,9 @@ public:
 	 * @param OutSubsectionRectsWithDuplicateBorders List of all subsection rects when taking into account duplicate borders (inclusive bounds)
 	 */
 	void ComputeAllSubsectionRects(TArray<FIntRect>& OutSubsectionRects, TArray<FIntRect>& OutSubsectionRectsWithDuplicateBorders) const;
+
+	/** Compute the component's linear index within that batch starting from the upper-left corner. This ensures a deterministic order within batch */
+	uint32 ComputeComponentLinearIndex(const FIntPoint& InComponentKey) const;
 };
 
 
@@ -315,7 +318,7 @@ public:
 	LANDSCAPE_EDIT_LAYERS_BATCHED_MERGE_EXPERIMENTAL
 	LANDSCAPE_API ULandscapeScratchRenderTarget* GetBlendRenderTargetReadPrevious() const;
 	LANDSCAPE_EDIT_LAYERS_BATCHED_MERGE_EXPERIMENTAL
-	LANDSCAPE_API ULandscapeScratchRenderTarget* GetValidityRenderTarget(const FName& InTargetLayerName) const;
+	LANDSCAPE_API ULandscapeScratchRenderTarget* GetComponentIdRenderTarget() const;
 	
 	struct FOnRenderBatchTargetGroupDoneParams
 	{
@@ -372,11 +375,11 @@ public:
 #endif // ENABLE_VISUAL_LOG
 
 	/** 
-	 * Render the stencil render targets for each target layer in this merge for this batch 
+	 * Render the component ids in this merge for this batch 
 	 * 
 	 * @param RDGBuilderRecorder recorder to append operations to a single FRDGBuilder, if the the recorder in "recording" mode. In "immediate" mode, a render command will be enqueued immediately
 	 */
-	void RenderValidityRenderTargets(UE::Landscape::FRDGBuilderRecorder& RDGBuilderRecorder);
+	void RenderComponentIds(UE::Landscape::FRDGBuilderRecorder& RDGBuilderRecorder);
 
 	/** 
 	 * Duplicates the vertex data from the (sub-)sections of the batch, assuming GetBlendRenderTargetRead() is the RT that is read from and GetBlendRenderTargetWrite() the one that is written to 
@@ -473,16 +476,20 @@ private:
 	  because it has been requested or one of the target layers that have been requested needs it to be present (e.g. weight-blending). Each bit in that bit array corresponds to an entry in AllTargetLayerNames */
 	TBitArray<> FinalTargetLayerBitIndices;
 
-	/** Render targets storing the validity of each pixel wrt the target layer (i.e. stencil-like buffer, but stored as a RT to let users access it as a standard texture) : 
-	 Useful when sampling neighbors to know whether the data there corresponds to a valid neighbor.*/
-	TMap<FName, ULandscapeScratchRenderTarget*> PerTargetLayerValidityRenderTargets;
+	/** Render targets storing the id (uint32, stored as a RTF_R32f) of the component within the batch that each pixel belongs to. It can be used as a stencil-like buffer, which can be useful when sampling neighbors,
+	  to know whether the data there corresponds to a valid neighbor, but also to resolve ambiguities when rendering overlapping quads (because the borders then belong to 2 components but this render target
+	  has the authority as to which one of those a given pixel "belongs to") */
+	ULandscapeScratchRenderTarget* ComponentIdRenderTarget = nullptr;
 
-	/** Maximum resolution needed by a given batch in this context (means we won't ever need more than this size of a render target during the whole merge */
+	/** Maximum resolution needed by a given batch in this context (means we won't ever need more than this size of a render target during the whole merge) */
 	FIntPoint MaxNeededResolution = FIntPoint(ForceInit);
 
-	/** Maximum number of slices needed by a given batch / target layer group in this context (means we won't ever need more than this number of slices in a given render target texture array during the whole merge */
+	/** Maximum number of slices needed by a given batch / target layer group in this context (means we won't ever need more than this number of slices in a given render target texture array during the whole merge) */
 	int32 MaxNeededNumSlices = 0;
 
+	/** Maximum number of components needed by a given batch in this context (means we won't ever render more than this amount of components during the whole merge) */
+	FIntPoint MaxNumComponents = FIntPoint(ForceInit);
+
 	/** Successive batches of components being processed by this context. Each batch should be self-contained so that we won't ever need to keep more than one in memory (VRAM)
 	  in order to compute the info we need (heightmap/weightmaps), i.e. nothing should ever depend on 2 different batches. */
 	TArray<FMergeRenderBatch> RenderBatches;
