Landscape Editor - Move landscape feature level warning notification to console log

diff --git a/Engine/Source/Editor/LandscapeEditor/Private/SLandscapeEditor.cpp b/Engine/Source/Editor/LandscapeEditor/Private/SLandscapeEditor.cpp
--- a/Engine/Source/Editor/LandscapeEditor/Private/SLandscapeEditor.cpp
+++ b/Engine/Source/Editor/LandscapeEditor/Private/SLandscapeEditor.cpp
@@ -655,11 +655,11 @@ FText SLandscapeEditor::GetErrorText() const
 		{
 			if (LandscapeEdMode->NewLandscapePreviewMode != ENewLandscapePreviewMode::None)
 			{
-				return LOCTEXT("IsFLError_create", "Can't create landscape with a feature level less than SM4!");
+				return LOCTEXT("IsFLError_create", "Can't create landscape with a feature level less than SM5!");
 			}
 			else
 			{
-				return LOCTEXT("IsFLError_edit", "Can't edit landscape with a feature level less than SM4!");
+				return LOCTEXT("IsFLError_edit", "Can't edit landscape with a feature level less than SM5!");
 			}
 		}
 		case ELandscapeEditingState::NoLandscape:
diff --git a/Engine/Source/Runtime/Landscape/Classes/Landscape.h b/Engine/Source/Runtime/Landscape/Classes/Landscape.h
--- a/Engine/Source/Runtime/Landscape/Classes/Landscape.h
+++ b/Engine/Source/Runtime/Landscape/Classes/Landscape.h
@@ -772,8 +772,6 @@ private:
 	FWaitingForResourcesNotificationHelper WaitingForTexturesNotificationHelper;
 	FWaitingForResourcesNotificationHelper WaitingForEditLayerResourcesNotificationHelper;
 
-	TSharedPtr<FLandscapeNotification> InvalidShadingModelNotification;
-
 	// Represent all the resolved paint layer, from all layers blended together (size of the landscape x material layer count)
 	class FLandscapeTexture2DArrayResource* CombinedLayersWeightmapAllMaterialLayersResource;
 	
@@ -788,6 +786,7 @@ private:
 
 	uint32 LastFlushedLayerUpdateFrame = 0;
 	uint32 LastPrepareTextureResourcesCalled = 0;
+	bool bHasWarnedAboutInvalidShadingModel = false;
 
 	// Counter to detect re-entrance
 	uint32 InLayerUpdateCount = 0;
diff --git a/Engine/Source/Runtime/Landscape/Private/LandscapeEditLayers.cpp b/Engine/Source/Runtime/Landscape/Private/LandscapeEditLayers.cpp
--- a/Engine/Source/Runtime/Landscape/Private/LandscapeEditLayers.cpp
+++ b/Engine/Source/Runtime/Landscape/Private/LandscapeEditLayers.cpp
@@ -11676,7 +11676,6 @@ void ALandscape::UpdateLayersContent(bool bInWaitForStreaming, bool bInSkipMonit
 		{
 			WaitingForTexturesNotificationHelper.Reset();
 			WaitingForEditLayerResourcesNotificationHelper.Reset();
-			InvalidShadingModelNotification.Reset();
 		}
 
 		// If nothing to do, let's do some garbage collecting on async readback tasks so that we slowly get rid of staging textures 
@@ -11786,27 +11785,6 @@ void ALandscape::UpdateLayersContent(bool bInWaitForStreaming, bool bInSkipMonit
 		}
 	}
 
-	// The Edit layers shaders only work on SM5 : cancel any update that might happen when SM5+ shading model is not active :
-	if (World->GetFeatureLevel() < ERHIFeatureLevel::SM5)
-	{
-		if (LandscapeNotificationManager)
-		{
-			if (!InvalidShadingModelNotification.IsValid())
-			{
-				InvalidShadingModelNotification = MakeShared<FLandscapeNotification>(this, ELandscapeNotificationType::ShadingModelInvalid);
-				static const FText NotificationText(LOCTEXT("InvalidShadingModel", "Cannot update landscape with a feature level less than SM5"));
-				InvalidShadingModelNotification->NotificationText = NotificationText;
-			}
-			LandscapeNotificationManager->RegisterNotification(InvalidShadingModelNotification);
-			bHideNotifications = false;
-		}
-		return;
-	}
-	else
-	{
-		InvalidShadingModelNotification.Reset();
-	}
-
 	bResourcesReady &= PrepareLayersTextureResources(bInWaitForStreaming);
 	if (!bResourcesReady && LandscapeNotificationManager)
 	{
@@ -11861,6 +11839,25 @@ void ALandscape::UpdateLayersContent(bool bInWaitForStreaming, bool bInSkipMonit
 		LayerContentUpdateModes |= UpdateAfterReadbackResolves(ComponentReadbackResults);
 	}
 
+	// The Edit layers shaders only work on SM5+ : log warning that any landscape updates will not happen when SM5+ shading model is not active :
+	if (World->GetFeatureLevel() >= ERHIFeatureLevel::SM5)
+	{
+		bHasWarnedAboutInvalidShadingModel = false;
+	}
+	else
+	{
+		// Log once when an update is requested, then early out so no updates are processed until a valid feature level is re-enabled
+		if (LayerContentUpdateModes != ELandscapeLayerUpdateMode::Update_None)
+		{
+			if (!bHasWarnedAboutInvalidShadingModel)
+			{
+				UE_LOG(LogLandscape, Warning, TEXT("Cannot update landscape with a feature level less than SM5"));
+				bHasWarnedAboutInvalidShadingModel = true;
+			}
+		}
+		return;
+	}
+
 	if (LayerContentUpdateModes == 0 && !bForceRender)
 	{
 		return;
