Adding control for landscape grass to align to the triangle normal when using jittered grid mode (instead of aligning to a grid-estimated normal)
Also fixed an issue with the triangle normal calculation, that wasn't taking DrawScale into account.
Original review under 43667518

diff --git a/Engine/Source/Runtime/Landscape/Classes/LandscapeGrassType.h b/Engine/Source/Runtime/Landscape/Classes/LandscapeGrassType.h
--- a/Engine/Source/Runtime/Landscape/Classes/LandscapeGrassType.h
+++ b/Engine/Source/Runtime/Landscape/Classes/LandscapeGrassType.h
@@ -114,6 +114,10 @@ struct FGrassVariety
 	UPROPERTY(BlueprintReadOnly, EditAnywhere, Category = Grass)
 	bool AlignToSurface;
 
+	/** When aligning to the surface when using grid mode, whether we should align to triangle normal (true) or a smooth grid-based normal (false).  Halton always uses triangle normal. */
+	UPROPERTY(BlueprintReadOnly, EditAnywhere, Category = Grass, meta = (EditCondition = "AlignToSurface && bUseGrid"))
+	bool bAlignToTriangleNormals;
+
 	/* Whether to use the landscape's lightmap when rendering the grass. */
 	UPROPERTY(BlueprintReadOnly, EditAnywhere, Category = Grass)
 	bool bUseLandscapeLightmap;
diff --git a/Engine/Source/Runtime/Landscape/Private/LandscapeGrass.cpp b/Engine/Source/Runtime/Landscape/Private/LandscapeGrass.cpp
--- a/Engine/Source/Runtime/Landscape/Private/LandscapeGrass.cpp
+++ b/Engine/Source/Runtime/Landscape/Private/LandscapeGrass.cpp
@@ -1253,6 +1253,7 @@ FGrassVariety::FGrassVariety()
 	, ScaleZ(1.0f, 1.0f)
 	, RandomRotation(true)
 	, AlignToSurface(true)
+	, bAlignToTriangleNormals(true)
 	, bUseLandscapeLightmap(false)
 	, bReceivesDecals(true)
 	, bAffectDistanceFieldLighting(false)
@@ -1823,6 +1824,7 @@ struct FAsyncGrassBuilder : public FGrassBuilderBase
 	bool bRandomRotation;
 	bool bRandomScale;
 	bool bAlignToSurface;
+	bool bAlignToTriangleNormals;
 	float PlacementJitter;
 	FRandomStream RandomStream;
 	FMatrix XForm;
@@ -1862,6 +1864,7 @@ struct FAsyncGrassBuilder : public FGrassBuilderBase
 		, bRandomRotation(GrassVariety.RandomRotation)
 		, bRandomScale(false)
 		, bAlignToSurface(GrassVariety.AlignToSurface)
+		, bAlignToTriangleNormals(GrassVariety.bAlignToTriangleNormals)
 		, PlacementJitter(GrassVariety.PlacementJitter)
 		, RandomStream(GrassInstancedStaticMeshComponent->InstancingRandomSeed)
 		, XForm(LandscapeToWorld * GrassInstancedStaticMeshComponent->GetComponentTransform().ToMatrixWithScale().Inverse())
@@ -2066,6 +2069,16 @@ struct FAsyncGrassBuilder : public FGrassBuilderBase
 		const FVector DefaultScale = GetDefaultScale();
 		float Div = 1.0f / float(SqrtMaxInstances);
 		TArray<FMatrix> InstanceTransforms;
+
+		auto ApplyNormalAlignmentTransform = [](const FVector& Normal, FMatrix& InOutXForm)
+			{
+				const FVector NewZ = Normal * FMath::Sign(Normal.Z);
+				const FVector NewX = (FVector(0, -1, 0) ^ NewZ).GetSafeNormal();
+				const FVector NewY = NewZ ^ NewX;
+				const FMatrix Align = FMatrix(NewX, NewY, NewZ, FVector::ZeroVector);
+				InOutXForm = InOutXForm * Align;
+			};
+
 		if (HaltonBaseIndex)
 		{
 			if (Extent.X < 0)
@@ -2097,20 +2110,12 @@ struct FAsyncGrassBuilder : public FGrassBuilderBase
 				{
 					const FVector Scale = bRandomScale ? GetRandomScale(Weight) : DefaultScale;
 					const float Rot = bRandomRotation ? RandomStream.GetFraction() * 360.0f : 0.0f;
-					const FMatrix BaseXForm = FScaleRotationTranslationMatrix(Scale, FRotator(0.0f, Rot, 0.0f), FVector::ZeroVector);
-					FMatrix OutXForm;
+					FMatrix OutXForm = FScaleRotationTranslationMatrix(Scale, FRotator(0.0f, Rot, 0.0f), FVector::ZeroVector);
 					if (bAlignToSurface && !ComputedNormal.IsNearlyZero())
 					{
-						const FVector NewZ = ComputedNormal * FMath::Sign(ComputedNormal.Z);
-						const FVector NewX = (FVector(0, -1, 0) ^ NewZ).GetSafeNormal();
-						const FVector NewY = NewZ ^ NewX;
-						const FMatrix Align = FMatrix(NewX, NewY, NewZ, FVector::ZeroVector);
-						OutXForm = (BaseXForm * Align).ConcatTranslation(LocationWithHeight) * XForm;
-					}
-					else
-					{
-						OutXForm = BaseXForm.ConcatTranslation(LocationWithHeight) * XForm;
+						ApplyNormalAlignmentTransform(ComputedNormal, OutXForm);
 					}
+					OutXForm = OutXForm.ConcatTranslation(LocationWithHeight) * XForm;
 					InstanceTransforms.Add(OutXForm);
 				}
 			}
@@ -2121,6 +2126,8 @@ struct FAsyncGrassBuilder : public FGrassBuilderBase
 				for (int32 InstanceIndex = 0; InstanceIndex < InstanceTransforms.Num(); InstanceIndex++)
 				{
 					const FMatrix& OutXForm = InstanceTransforms[InstanceIndex];
+
+					// NOTE: the random fraction parameter is not used, but we keep it to avoid changing the deterministic random behavior
 					SetInstance(InstanceIndex, OutXForm, RandomStream.GetFraction());
 				}
 			}
@@ -2135,6 +2142,7 @@ struct FAsyncGrassBuilder : public FGrassBuilderBase
 			struct FInstanceLocal
 			{
 				FVector Pos;
+				FVector Normal;
 				bool bKeep;
 				float Weight;
 			};
@@ -2157,7 +2165,7 @@ struct FAsyncGrassBuilder : public FGrassBuilderBase
 
 						FInstanceLocal& Instance = Instances[InstanceIndex];
 						float Weight = 0.f;
-						SampleLandscapeAtLocationLocal(Location, Instance.Pos, Weight);
+						SampleLandscapeAtLocationLocal(Location, Instance.Pos, Weight, (bAlignToSurface && bAlignToTriangleNormals) ? &Instance.Normal : nullptr);
 						Instance.bKeep = (Weight > AllowedDensityRange.Min) && (Weight <= AllowedDensityRange.Max) 
 							&& (Weight >= RandomStream.GetFraction()) 
 							&& !IsExcluded(Instance.Pos);
@@ -2170,6 +2178,8 @@ struct FAsyncGrassBuilder : public FGrassBuilderBase
 					}
 				}
 			}
+			
+			// A second pass is necessary because we need derivatives to calculate the smooth grid-based normal
 			if (NumKept)
 			{
 				InstanceTransforms.AddUninitialized(NumKept);
@@ -2187,35 +2197,34 @@ struct FAsyncGrassBuilder : public FGrassBuilderBase
 							{
 								const FVector Scale = bRandomScale ? GetRandomScale(Instance.Weight) : DefaultScale;
 								const float Rot = bRandomRotation ? RandomStream.GetFraction() * 360.0f : 0.0f;
-								const FMatrix BaseXForm = FScaleRotationTranslationMatrix(Scale, FRotator(0.0f, Rot, 0.0f), FVector::ZeroVector);
-								FMatrix OutXForm;
+								FMatrix OutXForm = FScaleRotationTranslationMatrix(Scale, FRotator(0.0f, Rot, 0.0f), FVector::ZeroVector);
 								if (bAlignToSurface)
 								{
-									FVector PosX1 = xStart ? Instances[InstanceIndex - SqrtMaxInstances].Pos : Instance.Pos;
-									FVector PosX2 = (xStart + 1 < SqrtMaxInstances) ? Instances[InstanceIndex + SqrtMaxInstances].Pos : Instance.Pos;
-									FVector PosY1 = yStart ? Instances[InstanceIndex - 1].Pos : Instance.Pos;
-									FVector PosY2 = (yStart + 1 < SqrtMaxInstances) ? Instances[InstanceIndex + 1].Pos : Instance.Pos;
-
-									if (PosX1 != PosX2 && PosY1 != PosY2)
+									if (bAlignToTriangleNormals)
 									{
-										FVector NewZ = ((PosX1 - PosX2) ^ (PosY1 - PosY2)).GetSafeNormal();
-										NewZ *= FMath::Sign(NewZ.Z);
-
-										const FVector NewX = (FVector(0, -1, 0) ^ NewZ).GetSafeNormal();
-										const FVector NewY = NewZ ^ NewX;
-
-										FMatrix Align = FMatrix(NewX, NewY, NewZ, FVector::ZeroVector);
-										OutXForm = (BaseXForm * Align).ConcatTranslation(Instance.Pos) * XForm;
+										ApplyNormalAlignmentTransform(Instance.Normal, OutXForm);
 									}
 									else
 									{
-										OutXForm = BaseXForm.ConcatTranslation(Instance.Pos) * XForm;
+										// NOTE: This calculation of smooth normals uses jittered grid points to compute the derivatives.
+										// The jitter is generally ok, though may ocasionally result in a small delta that is sensitive to error.
+										// However, the grid may be at any resolution, so the derivative will be estimated from taking steps at that resolution.
+										// which may result in normals that don't match the local terrain well, especially for low resolution grids that take large steps.
+										FVector PosX1 = xStart ? Instances[InstanceIndex - SqrtMaxInstances].Pos : Instance.Pos;
+										FVector PosX2 = (xStart + 1 < SqrtMaxInstances) ? Instances[InstanceIndex + SqrtMaxInstances].Pos : Instance.Pos;
+										FVector PosY1 = yStart ? Instances[InstanceIndex - 1].Pos : Instance.Pos;
+										FVector PosY2 = (yStart + 1 < SqrtMaxInstances) ? Instances[InstanceIndex + 1].Pos : Instance.Pos;
+
+										if (PosX1 != PosX2 && PosY1 != PosY2)
+										{
+											FVector Normal = ((PosX1 - PosX2) ^ (PosY1 - PosY2)).GetSafeNormal();
+											ApplyNormalAlignmentTransform(Normal, OutXForm);
+										}
 									}
 								}
-								else
-								{
-									OutXForm = BaseXForm.ConcatTranslation(Instance.Pos) * XForm;
-								}
+
+								// apply translation and outer transform
+								OutXForm = OutXForm.ConcatTranslation(Instance.Pos) * XForm;
 								InstanceTransforms[OutInstanceIndex] = OutXForm;
 								SetInstance(OutInstanceIndex++, OutXForm, RandomStream.GetFraction());
 							}
@@ -2318,19 +2327,19 @@ struct FAsyncGrassBuilder : public FGrassBuilderBase
 		// Compute normal
 		if (OutNormal)
 		{
-			const FVector P11((float)X1, (float)Y1, SampleHeight11);
-			const FVector P22((float)X2, (float)Y2, SampleHeight22);
+			const FVector P11 = FVector((float)X1, (float)Y1, SampleHeight11) * DrawScale;
+			const FVector P22 = FVector((float)X2, (float)Y2, SampleHeight22) * DrawScale;
 			
 			// Choose triangle and compute normal
 			if (LerpX > LerpY)
 			{
-				const FVector P21((float)X2, (float)Y1, SampleHeight21);
-				*OutNormal = (P11 != P21 && P22 != P21) ? FVector(((P22 - P21) ^ (P11 - P21)).GetSafeNormal()) : FVector::ZeroVector;
+				const FVector P21 = FVector((float)X2, (float)Y1, SampleHeight21) * DrawScale;
+				*OutNormal = (P11 != P21 && P22 != P21) ? FVector(((P22 - P21) ^ (P11 - P21)).GetSafeNormal()) : FVector::UpVector;
 			}
 			else 
 			{
-				const FVector P12((float)X1, (float)Y2, SampleHeight12);
-				*OutNormal = (P11 != P12 && P22 != P12) ? FVector(((P11 - P12) ^ (P22 - P12)).GetSafeNormal()) : FVector::ZeroVector;
+				const FVector P12 = FVector((float)X1, (float)Y2, SampleHeight12) * DrawScale;
+				*OutNormal = (P11 != P12 && P22 != P12) ? FVector(((P11 - P12) ^ (P22 - P12)).GetSafeNormal()) : FVector::UpVector;
 			}
 		}
 	}
