ALandscapeProxy : add console variables to cook dependencies

#rnx
#incrementalcook

diff --git a/Engine/Source/Runtime/Landscape/Classes/LandscapeProxy.h b/Engine/Source/Runtime/Landscape/Classes/LandscapeProxy.h
--- a/Engine/Source/Runtime/Landscape/Classes/LandscapeProxy.h
+++ b/Engine/Source/Runtime/Landscape/Classes/LandscapeProxy.h
@@ -1214,7 +1214,9 @@ PRAGMA_ENABLE_DEPRECATION_WARNINGS
 	//~ Begin UObject Interface.
 #if WITH_EDITOR
 	virtual void BeginCacheForCookedPlatformData(const ITargetPlatform* TargetPlatform) override;
+	virtual void OnCookEvent(UE::Cook::ECookEvent CookEvent, UE::Cook::FCookEventContext& CookContext) override;
 #endif // WITH_EDITOR
+
 	virtual void PreSave(FObjectPreSaveContext ObjectSaveContext) override;
 	virtual void Serialize(FArchive& Ar) override;
 	static void AddReferencedObjects(UObject* InThis, FReferenceCollector& Collector);
diff --git a/Engine/Source/Runtime/Landscape/Private/Landscape.cpp b/Engine/Source/Runtime/Landscape/Private/Landscape.cpp
--- a/Engine/Source/Runtime/Landscape/Private/Landscape.cpp
+++ b/Engine/Source/Runtime/Landscape/Private/Landscape.cpp
@@ -106,6 +106,7 @@ Landscape.cpp: Terrain rendering
 #include "Engine/Texture2DArray.h"
 
 #if WITH_EDITOR
+#include "Cooker/CookEvents.h"
 #include "Rendering/StaticLightingSystemInterface.h"
 #include "Misc/ScopedSlowTask.h"
 #include "UnrealEdGlobals.h"
@@ -230,18 +231,20 @@ static FAutoConsoleVariable CVarAllowGrassStripping(
 	true,
 	TEXT("Enables the conditional stripping of grass data during cook.  Disabling this means the bStripGrassWhenCooked* will be ignored."));
 
+static const TCHAR* CVarLandscapeHeightmapCompressionModeName = TEXT("landscape.HeightmapCompressionMode");
 int32 GLandscapeHeightmapCompressionMode = 1;
 static FAutoConsoleVariableRef CVarLandscapeHeightmapCompressionMode(
-	TEXT("landscape.HeightmapCompressionMode"),
+	CVarLandscapeHeightmapCompressionModeName,
 	GLandscapeHeightmapCompressionMode,
 	TEXT("Defines whether compression is applied to landscapes.  Can be defined per platform.\n")
 	TEXT(" 0: force disable heightmap compression on all landscapes\n")
 	TEXT(" 1: force enable heightmap compression on all landscapes (default)\n"),
 	ECVF_Preview | ECVF_ReadOnly);
 
+static const TCHAR* CVarLandscapeHeightmapCompressionMipThresholdName = TEXT("landscape.HeightmapCompressionMipThreshold");
 int32 GLandscapeHeightmapCompressionMipThreshold = 32;
 static FAutoConsoleVariableRef CVarLandscapeHeightmapCompressionMipThreshold(
-	TEXT("landscape.HeightmapCompressionMipThreshold"),
+	CVarLandscapeHeightmapCompressionMipThresholdName,
 	GLandscapeHeightmapCompressionMipThreshold,
 	TEXT("Sets the minimum size for which heightmap mips are stored in a compressed layout.  Can be defined per platform.\n")
 	TEXT("Below this size, mips are stored in an uncompressed layout.\n")
@@ -3828,6 +3831,19 @@ void ALandscapeProxy::BeginCacheForCookedPlatformData(const ITargetPlatform* Tar
 	InstallOrUpdateTextureUserDatas(TargetPlatform);
 	Super::BeginCacheForCookedPlatformData(TargetPlatform);
 }
+
+void ALandscapeProxy::OnCookEvent(UE::Cook::ECookEvent CookEvent, UE::Cook::FCookEventContext& CookContext)
+{
+	using namespace UE::Cook;
+	Super::OnCookEvent(CookEvent, CookContext);
+
+	if (CookEvent == ECookEvent::PlatformCookDependencies && CookContext.IsCooking())
+	{
+		CookContext.AddLoadBuildDependency(FCookDependency::ConsoleVariable(FStringView(CVarLandscapeHeightmapCompressionModeName), CookContext.GetTargetPlatform(), /*bFallbackToNonPlatform*/true));
+		CookContext.AddLoadBuildDependency(FCookDependency::ConsoleVariable(FStringView(CVarLandscapeHeightmapCompressionMipThresholdName), CookContext.GetTargetPlatform(), /*bFallbackToNonPlatform*/true));
+	}
+}
+
 #endif
 
 
