Clean up `FLandscapeSharedBuffers` to avoid querying the result of reference release.

diff --git a/Engine/Plugins/Experimental/GPULightmass/Source/GPULightmass/Private/Scene/Scene.cpp b/Engine/Plugins/Experimental/GPULightmass/Source/GPULightmass/Private/Scene/Scene.cpp
--- a/Engine/Plugins/Experimental/GPULightmass/Source/GPULightmass/Private/Scene/Scene.cpp
+++ b/Engine/Plugins/Experimental/GPULightmass/Source/GPULightmass/Private/Scene/Scene.cpp
@@ -1449,11 +1449,8 @@ void FScene::RemoveGeometryInstanceFromComponent(ULandscapeComponent* InComponen
 		[ElementId, &RenderState = RenderState](FRHICommandListImmediate&) mutable
 	{
 		FLandscapeRenderState& LandscapeRenderState = RenderState.LandscapeRenderStates.Elements[ElementId];
-
-		if (LandscapeRenderState.SharedBuffers->Release() == 0)
-		{
-			FLandscapeComponentSceneProxy::SharedBuffersMap.Remove(LandscapeRenderState.SharedBuffersKey);
-		}
+		LandscapeRenderState.SharedBuffers->Release();
+		LandscapeRenderState.SharedBuffers = nullptr;
 
 		LandscapeRenderState.LandscapeUniformShaderParameters->ReleaseResource();
 
diff --git a/Engine/Source/Runtime/Landscape/Private/LandscapeRender.cpp b/Engine/Source/Runtime/Landscape/Private/LandscapeRender.cpp
--- a/Engine/Source/Runtime/Landscape/Private/LandscapeRender.cpp
+++ b/Engine/Source/Runtime/Landscape/Private/LandscapeRender.cpp
@@ -532,7 +532,7 @@ static uint32 GetRayTracingSectionVertCountForLOD(uint32 SubsectionSizeVerts, in
 //
 // FLandscapeComponentSceneProxy
 //
-TMap<uint32, FLandscapeSharedBuffers*>FLandscapeComponentSceneProxy::SharedBuffersMap;
+TMap<uint32, FLandscapeSharedBuffers*> FLandscapeComponentSceneProxy::SharedBuffersMap;
 
 const static FName NAME_LandscapeResourceNameForDebugging(TEXT("Landscape"));
 
@@ -2008,10 +2008,7 @@ FLandscapeComponentSceneProxy::~FLandscapeComponentSceneProxy()
 	if (SharedBuffers)
 	{
 		check(SharedBuffers == FLandscapeComponentSceneProxy::SharedBuffersMap.FindRef(SharedBuffersKey));
-		if (SharedBuffers->Release() == 0)
-		{
-			FLandscapeComponentSceneProxy::SharedBuffersMap.Remove(SharedBuffersKey);
-		}
+		SharedBuffers->Release();
 		SharedBuffers = nullptr;
 	}
 }
@@ -3563,6 +3560,10 @@ FLandscapeSharedBuffers::FLandscapeSharedBuffers(FRHICommandListBase& RHICmdList
 
 FLandscapeSharedBuffers::~FLandscapeSharedBuffers()
 {
+	// If we run the destructor we are the last reference so we need to remove
+	// ourselves from the shared pool of buffers.
+    FLandscapeComponentSceneProxy::SharedBuffersMap.Remove(SharedBuffersKey);
+
 	delete VertexBuffer;
 
 	for (int32 i = 0; i < NumIndexBuffers; i++)
