Fixed various landscape+water new cases (introduced in UE5.6) where it was assumed that GetWorld() returns a valid world while it's not the case when a world is deleted from the content browser

diff --git a/Engine/Plugins/Experimental/Water/Source/Runtime/Private/WaterBodyExclusionVolume.cpp b/Engine/Plugins/Experimental/Water/Source/Runtime/Private/WaterBodyExclusionVolume.cpp
--- a/Engine/Plugins/Experimental/Water/Source/Runtime/Private/WaterBodyExclusionVolume.cpp
+++ b/Engine/Plugins/Experimental/Water/Source/Runtime/Private/WaterBodyExclusionVolume.cpp
@@ -34,8 +34,14 @@ void AWaterBodyExclusionVolume::UpdateOverlappingWaterBodies(const FWaterExclusi
 {
 	TArray<FOverlapResult> Overlaps;
 
+	UWorld* World = GetWorld();
+	if (World == nullptr)
+	{
+		return;
+	}
+
 	FBoxSphereBounds Bounds = GetBounds();
-	GetWorld()->OverlapMultiByObjectType(Overlaps, Bounds.Origin, FQuat::Identity, FCollisionObjectQueryParams::AllObjects, FCollisionShape::MakeBox(Bounds.BoxExtent));
+	World->OverlapMultiByObjectType(Overlaps, Bounds.Origin, FQuat::Identity, FCollisionObjectQueryParams::AllObjects, FCollisionShape::MakeBox(Bounds.BoxExtent));
 
 	// Find any new overlapping bodies and notify them that this exclusion volume influences them
 	TSet<UWaterBodyComponent*> ExistingOverlappingBodies;
@@ -44,7 +50,7 @@ void AWaterBodyExclusionVolume::UpdateOverlappingWaterBodies(const FWaterExclusi
 	TSoftObjectPtr<AWaterBodyExclusionVolume> SoftThis(this);
 
 	// Fixup overlapping bodies (iterating on actors on post-load will fail, but this is fine as this exclusion volume should not yet be referenced by an existing water body upon loading) : 
-	FWaterBodyManager::ForEachWaterBodyComponent(GetWorld(), [SoftThis, &ExistingOverlappingBodies](UWaterBodyComponent* WaterBodyComponent)
+	FWaterBodyManager::ForEachWaterBodyComponent(World, [SoftThis, &ExistingOverlappingBodies](UWaterBodyComponent* WaterBodyComponent)
 	{
 		if (WaterBodyComponent->ContainsExclusionVolume(SoftThis))
 		{
diff --git a/Engine/Plugins/Experimental/Water/Source/Runtime/Private/WaterBodyIslandActor.cpp b/Engine/Plugins/Experimental/Water/Source/Runtime/Private/WaterBodyIslandActor.cpp
--- a/Engine/Plugins/Experimental/Water/Source/Runtime/Private/WaterBodyIslandActor.cpp
+++ b/Engine/Plugins/Experimental/Water/Source/Runtime/Private/WaterBodyIslandActor.cpp
@@ -184,12 +184,18 @@ void AWaterBodyIsland::UpdateOverlappingWaterBodyComponents()
 {
 	TArray<FOverlapResult> Overlaps;
 
+	UWorld* World = GetWorld();
+	if (World == nullptr)
+	{
+		return;
+	}
+
 	if (SplineComp)
 	{
 		FCollisionShape OverlapShape;
 		// Expand shape in Z to ensure we get overlaps for islands slighty above or below water level
 		OverlapShape.SetBox((FVector3f)SplineComp->Bounds.BoxExtent+FVector3f(0,0,10000));
-		GetWorld()->OverlapMultiByObjectType(Overlaps, SplineComp->Bounds.Origin, FQuat::Identity, FCollisionObjectQueryParams::AllObjects, OverlapShape);
+		World->OverlapMultiByObjectType(Overlaps, SplineComp->Bounds.Origin, FQuat::Identity, FCollisionObjectQueryParams::AllObjects, OverlapShape);
 	}
 
 	// Find any new overlapping bodies and notify them that this island influences them
@@ -199,7 +205,7 @@ void AWaterBodyIsland::UpdateOverlappingWaterBodyComponents()
 	TSoftObjectPtr<AWaterBodyIsland> SoftThis(this);
 
 	// Fixup overlapping bodies
-	FWaterBodyManager::ForEachWaterBodyComponent(GetWorld(), [SoftThis, &ExistingOverlappingBodies](UWaterBodyComponent* WaterBodyComponent)
+	FWaterBodyManager::ForEachWaterBodyComponent(World, [SoftThis, &ExistingOverlappingBodies](UWaterBodyComponent* WaterBodyComponent)
 	{
 		if (WaterBodyComponent->ContainsIsland(SoftThis))
 		{
diff --git a/Engine/Plugins/Experimental/Water/Source/Runtime/Private/WaterZoneActor.cpp b/Engine/Plugins/Experimental/Water/Source/Runtime/Private/WaterZoneActor.cpp
--- a/Engine/Plugins/Experimental/Water/Source/Runtime/Private/WaterZoneActor.cpp
+++ b/Engine/Plugins/Experimental/Water/Source/Runtime/Private/WaterZoneActor.cpp
@@ -255,13 +255,16 @@ void AWaterZone::PostRegisterAllComponents()
 
 #if WITH_EDITOR
 	// Register for landscape notifications
-	if (ULandscapeSubsystem* LandscapeSubsystem = GetWorld()->GetSubsystem<ULandscapeSubsystem>())
+	if (UWorld* World = GetWorld())
 	{
-		check(!OnHeightmapStreamedHandle.IsValid());
-		OnHeightmapStreamedHandle = LandscapeSubsystem->OnHeightmapStreamed().AddUObject(this, &AWaterZone::OnLandscapeHeightmapStreamed);
+		if (ULandscapeSubsystem* LandscapeSubsystem = World->GetSubsystem<ULandscapeSubsystem>())
+		{
+			check(!OnHeightmapStreamedHandle.IsValid());
+			OnHeightmapStreamedHandle = LandscapeSubsystem->OnHeightmapStreamed().AddUObject(this, &AWaterZone::OnLandscapeHeightmapStreamed);
 
-		check(!OnLandscapeComponentDataChangedHandle.IsValid());
-		OnLandscapeComponentDataChangedHandle = LandscapeSubsystem->OnLandscapeProxyComponentDataChanged().AddUObject(this, &AWaterZone::OnLandscapeComponentDataChanged);
+			check(!OnLandscapeComponentDataChangedHandle.IsValid());
+			OnLandscapeComponentDataChangedHandle = LandscapeSubsystem->OnLandscapeProxyComponentDataChanged().AddUObject(this, &AWaterZone::OnLandscapeComponentDataChanged);
+		}
 	}
 #endif // WITH_EDITOR
 }
@@ -271,15 +274,18 @@ void AWaterZone::PostUnregisterAllComponents()
 	Super::PostUnregisterAllComponents();
 
 #if WITH_EDITOR
-	if (ULandscapeSubsystem* LandscapeSubsystem = GetWorld()->GetSubsystem<ULandscapeSubsystem>())
+	if (UWorld* World = GetWorld())
 	{
-		if (OnHeightmapStreamedHandle.IsValid())
-		{
-			LandscapeSubsystem->OnHeightmapStreamed().Remove(OnHeightmapStreamedHandle);
-		}
-		if (OnLandscapeComponentDataChangedHandle.IsValid())
+		if (ULandscapeSubsystem* LandscapeSubsystem = World->GetSubsystem<ULandscapeSubsystem>())
 		{
-			LandscapeSubsystem->OnLandscapeProxyComponentDataChanged().Remove(OnLandscapeComponentDataChangedHandle);
+			if (OnHeightmapStreamedHandle.IsValid())
+			{
+				LandscapeSubsystem->OnHeightmapStreamed().Remove(OnHeightmapStreamedHandle);
+			}
+			if (OnLandscapeComponentDataChangedHandle.IsValid())
+			{
+				LandscapeSubsystem->OnLandscapeProxyComponentDataChanged().Remove(OnLandscapeComponentDataChangedHandle);
+			}
 		}
 	}
 	OnHeightmapStreamedHandle.Reset();
diff --git a/Engine/Source/Runtime/Landscape/Private/Landscape.cpp b/Engine/Source/Runtime/Landscape/Private/Landscape.cpp
--- a/Engine/Source/Runtime/Landscape/Private/Landscape.cpp
+++ b/Engine/Source/Runtime/Landscape/Private/Landscape.cpp
@@ -494,9 +494,13 @@ FGraphEventRef ALandscapeProxy::UpdateNaniteRepresentationAsync(const ITargetPla
 
 		BatchBuildEvent = FFunctionGraphTask::CreateAndDispatchWhenReady([] {}, TStatId(), &UpdateDependencies, ENamedThreads::GameThread);
 
-		// Register the finalize build event so that it can be tracked globally by : 
-		ULandscapeSubsystem* LandscapeSubsystem = GetWorld()->GetSubsystem<ULandscapeSubsystem>();
-		LandscapeSubsystem->AddNaniteFinalizeBuildEvent(BatchBuildEvent);
+		// Register the finalize build event so that it can be tracked globally : 
+		if (UWorld* World = GetWorld())
+		{
+			ULandscapeSubsystem* LandscapeSubsystem = World->GetSubsystem<ULandscapeSubsystem>();
+			check(LandscapeSubsystem != nullptr);
+			LandscapeSubsystem->AddNaniteFinalizeBuildEvent(BatchBuildEvent);
+		}
 	}
 	else
 	{
@@ -512,7 +516,15 @@ void ALandscapeProxy::UpdateNaniteRepresentation(const ITargetPlatform* InTarget
 	check(IsInGameThread());
 
 	FGraphEventRef GraphEvent = UpdateNaniteRepresentationAsync(InTargetPlatform);
-	ULandscapeSubsystem* LandscapeSubsystem = GetWorld()->GetSubsystem<ULandscapeSubsystem>();
+
+	UWorld* World = GetWorld();
+	if (World == nullptr)
+	{
+		return;
+	}
+	
+	ULandscapeSubsystem* LandscapeSubsystem = World->GetSubsystem<ULandscapeSubsystem>();
+	check(LandscapeSubsystem != nullptr);
 
 	if (!GraphEvent.IsValid())
 	{
@@ -540,7 +552,15 @@ void ALandscapeProxy::InvalidateNaniteRepresentation(bool bInCheckContentId)
 
 void ALandscapeProxy::InvalidateOrUpdateNaniteRepresentation(bool bInCheckContentId, const ITargetPlatform* InTargetPlatform)
 {
-	ULandscapeSubsystem* Subsystem = GetWorld()->GetSubsystem<ULandscapeSubsystem>();
+	UWorld* World = GetWorld();
+	if (World == nullptr)
+	{
+		return;
+	}
+	
+	ULandscapeSubsystem* Subsystem = World->GetSubsystem<ULandscapeSubsystem>();
+	check(Subsystem != nullptr);
+
 	if (Subsystem->IsLiveNaniteRebuildEnabled())
 	{
 		UpdateNaniteRepresentation(InTargetPlatform);
@@ -1032,7 +1052,6 @@ UMaterialInterface* ULandscapeComponent::GetLandscapeMaterial(int8 InLODIndex) c
 	if (InLODIndex != INDEX_NONE)
 	{
 		UWorld* World = GetWorld();
-
 		if (World != nullptr)
 		{
 			if (const FLandscapePerLODMaterialOverride* LocalMaterialOverride = PerLODOverrideMaterials.FindByPredicate(
@@ -2325,7 +2344,6 @@ void ULandscapeComponent::OnUnregister()
 
 		// AActor::GetWorld checks for Unreachable and BeginDestroyed
 		UWorld* World = Proxy->GetWorld();
-
 		if (World)
 		{
 			if (ULandscapeInfo* Info = GetLandscapeInfo())
@@ -2504,14 +2522,15 @@ TArray<FWeightmapLayerAllocationInfo>& ULandscapeComponent::GetWeightmapLayerAll
 
 const TArray<FWeightmapLayerAllocationInfo>& ULandscapeComponent::GetCurrentRuntimeWeightmapLayerAllocations() const
 {
-
-	bool bIsMobile = GetWorld()->GetFeatureLevel() == ERHIFeatureLevel::ES3_1;
+	UWorld* World = GetWorld();
+	bool bIsMobile = (World != nullptr) && (World->GetFeatureLevel() == ERHIFeatureLevel::ES3_1);
 	return bIsMobile ? MobileWeightmapLayerAllocations : WeightmapLayerAllocations;
 }
 
 TArray<FWeightmapLayerAllocationInfo>& ULandscapeComponent::GetCurrentRuntimeWeightmapLayerAllocations()
 {
-	bool bIsMobile = GetWorld()->GetFeatureLevel() == ERHIFeatureLevel::ES3_1;
+	UWorld* World = GetWorld();
+	bool bIsMobile = (World != nullptr) && (World->GetFeatureLevel() == ERHIFeatureLevel::ES3_1);
 	return bIsMobile ? MobileWeightmapLayerAllocations : WeightmapLayerAllocations;
 }
 
@@ -5248,10 +5267,13 @@ int32 ULandscapeInfo::GetModifiedPackageCount() const
 bool ULandscapeInfo::TryAddToModifiedPackages(UPackage* InPackage, const ALandscape* InLandscapeOverride)
 {
 	const ALandscape* LocalLandscapeActor = (InLandscapeOverride != nullptr) ? InLandscapeOverride : LandscapeActor.Get();
-	check(LocalLandscapeActor);
+	check(LocalLandscapeActor != nullptr);
+
+	UWorld* World = LocalLandscapeActor->GetWorld();
+	check(World != nullptr);
 
 	// We don't want to bother with packages being marked dirty for anything else than the Editor world 
-	if (LocalLandscapeActor->GetWorld()->WorldType != EWorldType::Editor)
+	if (World->WorldType != EWorldType::Editor)
 	{
 		return false;
 	}
@@ -5348,12 +5370,13 @@ bool ULandscapeInfo::ModifyObject(UObject* InObject, bool bAlwaysMarkDirty)
 
 void ULandscapeInfo::DeleteActorWhenApplyingModifiedStatus(AActor* InActor, bool bInAllowUI)
 {
-	check(InActor->GetWorld() && InActor->GetWorld()->WorldType == EWorldType::Editor);
+	UWorld* World = InActor->GetWorld();
+	check((World != nullptr) && (World->WorldType == EWorldType::Editor));
 
 	// If we can mark the package dirty, then we can also delete it right away.  If we can't, then enqueue it for deletion at the same time as other deferred package dirtying.
 	if (InActor->MarkPackageDirty())
 	{
-		UE::Landscape::DeleteActors({ InActor }, InActor->GetWorld(), bInAllowUI);
+		UE::Landscape::DeleteActors({ InActor }, World, bInAllowUI);
 	}
 	else
 	{
@@ -7130,18 +7153,22 @@ void ALandscapeProxy::InvalidateGeneratedComponentData(const TArray<ULandscapeCo
 		Proxy->FlushGrassComponents(&Iter.Value());
 
 #if WITH_EDITOR
-		ULandscapeSubsystem* Subsystem = Proxy->GetWorld()->GetSubsystem<ULandscapeSubsystem>();
-		if (Subsystem->IsLiveNaniteRebuildEnabled())
-		{
-			Proxy->GetAsyncWorkMonitor().SetDelayedUpdateTimer(FAsyncWorkMonitor::EAsyncWorkType::BuildNaniteMeshes, LandscapeNaniteBuildLag);
-		}
-		else
+		FLandscapeProxyComponentDataChangedParams ChangeParams(Iter.Value());
+		if (UWorld* ProxyWorld = Proxy->GetWorld())
 		{
-			Proxy->InvalidateOrUpdateNaniteRepresentation(/* bInCheckContentId = */true, /*InTargetPlatform = */nullptr);
-		}
+			ULandscapeSubsystem* Subsystem = ProxyWorld->GetSubsystem<ULandscapeSubsystem>();
+			check(Subsystem != nullptr);
+			if (Subsystem->IsLiveNaniteRebuildEnabled())
+			{
+				Proxy->GetAsyncWorkMonitor().SetDelayedUpdateTimer(FAsyncWorkMonitor::EAsyncWorkType::BuildNaniteMeshes, LandscapeNaniteBuildLag);
+			}
+			else
+			{
+				Proxy->InvalidateOrUpdateNaniteRepresentation(/* bInCheckContentId = */true, /*InTargetPlatform = */nullptr);
+			}
 
-		FLandscapeProxyComponentDataChangedParams ChangeParams(Iter.Value());
-		Subsystem->GetDelegateAccess().OnLandscapeProxyComponentDataChangedDelegate.Broadcast(Iter.Key(), ChangeParams);
+			Subsystem->GetDelegateAccess().OnLandscapeProxyComponentDataChangedDelegate.Broadcast(Iter.Key(), ChangeParams);
+		}
 
 PRAGMA_DISABLE_DEPRECATION_WARNINGS
 		Proxy->OnComponentDataChanged.Broadcast(Iter.Key(), ChangeParams);
diff --git a/Engine/Source/Runtime/Landscape/Private/LandscapeEdit.cpp b/Engine/Source/Runtime/Landscape/Private/LandscapeEdit.cpp
--- a/Engine/Source/Runtime/Landscape/Private/LandscapeEdit.cpp
+++ b/Engine/Source/Runtime/Landscape/Private/LandscapeEdit.cpp
@@ -4199,57 +4199,60 @@ TSharedRef<UE::Landscape::Nanite::FAsyncBuildData> ALandscapeProxy::MakeAsyncNan
 {
 	TSharedRef<UE::Landscape::Nanite::FAsyncBuildData> AsyncBuildData = MakeShared<UE::Landscape::Nanite::FAsyncBuildData>();
 
-	// Make sure the requested LOD is valid
-	int32 FinalLODToExport = FMath::Clamp<int32>(InLODToExport, 0, FMath::CeilLogTwo(SubsectionSizeQuads + 1) - 1);
+	if (UWorld* World = GetWorld())
+	{
+		// Make sure the requested LOD is valid
+		int32 FinalLODToExport = FMath::Clamp<int32>(InLODToExport, 0, FMath::CeilLogTwo(SubsectionSizeQuads + 1) - 1);
 
-	AsyncBuildData->LOD = FinalLODToExport;
-	AsyncBuildData->LandscapeWeakRef = MakeWeakObjectPtr(const_cast<ALandscapeProxy*>(this));
-	AsyncBuildData->LandscapeSubSystemWeakRef = MakeWeakObjectPtr(GetWorld()->GetSubsystem<ULandscapeSubsystem>());
+		AsyncBuildData->LOD = FinalLODToExport;
+		AsyncBuildData->LandscapeWeakRef = MakeWeakObjectPtr(const_cast<ALandscapeProxy*>(this));
+		AsyncBuildData->LandscapeSubSystemWeakRef = MakeWeakObjectPtr(World->GetSubsystem<ULandscapeSubsystem>());
 
-	for (ULandscapeComponent* Component : InComponentsToExport)
-	{
-		check(LandscapeComponents.Contains(Component)); // component we're requesting to export has to be in the proxy.
-		UMaterialInterface* Material = nullptr;
-		if (Component)
+		for (ULandscapeComponent* Component : InComponentsToExport)
 		{
-			Material = Component->GetMaterialInstance(0u);
-			AsyncBuildData->InputMaterialSlotNames.Add(FName(*FString::Format(TEXT("LandscapeMat_{0}"), { AsyncBuildData->InputComponents.Num() })));
-			AsyncBuildData->InputMaterials.Add(Material ? Material : UMaterial::GetDefaultMaterial(MD_Surface));
-			AsyncBuildData->InputComponents.Add(Component);
+			check(LandscapeComponents.Contains(Component)); // component we're requesting to export has to be in the proxy.
+			UMaterialInterface* Material = nullptr;
+			if (Component)
+			{
+				Material = Component->GetMaterialInstance(0u);
+				AsyncBuildData->InputMaterialSlotNames.Add(FName(*FString::Format(TEXT("LandscapeMat_{0}"), { AsyncBuildData->InputComponents.Num() })));
+				AsyncBuildData->InputMaterials.Add(Material ? Material : UMaterial::GetDefaultMaterial(MD_Surface));
+				AsyncBuildData->InputComponents.Add(Component);
+			}
 		}
-	}
 
-	if (AsyncBuildData->InputComponents.Num() == 0)
-	{
-		UE_LOG(LogLandscape, Verbose, TEXT("%s : no Nanite mesh to export"), *GetActorNameOrLabel());
-	}
+		if (AsyncBuildData->InputComponents.Num() == 0)
+		{
+			UE_LOG(LogLandscape, Verbose, TEXT("%s : no Nanite mesh to export"), *GetActorNameOrLabel());
+		}
 
-	if (AsyncBuildData->InputMaterials.Num() > NANITE_MAX_CLUSTER_MATERIALS)
-	{
-		UE_LOG(LogLandscape, Warning, TEXT("%s : Nanite landscape mesh would have more than %i materials, which is currently not supported. Please reduce the number of components in this landscape actor to enable Nanite."), *GetActorNameOrLabel(), NANITE_MAX_CLUSTER_MATERIALS)
-	}
+		if (AsyncBuildData->InputMaterials.Num() > NANITE_MAX_CLUSTER_MATERIALS)
+		{
+			UE_LOG(LogLandscape, Warning, TEXT("%s : Nanite landscape mesh would have more than %i materials, which is currently not supported. Please reduce the number of components in this landscape actor to enable Nanite."), *GetActorNameOrLabel(), NANITE_MAX_CLUSTER_MATERIALS)
+		}
 
-	for (ULandscapeComponent* LandscapeComponent : InComponentsToExport)
-	{
-		TRACE_CPUPROFILER_EVENT_SCOPE(ALandscapeProxy::MakeAsyncBuildData-CopyHeightAndVisibility);
-		FLandscapeComponentDataInterface DataInterface(LandscapeComponent, FinalLODToExport, false);
+		for (ULandscapeComponent* LandscapeComponent : InComponentsToExport)
+		{
+			TRACE_CPUPROFILER_EVENT_SCOPE(ALandscapeProxy::MakeAsyncBuildData-CopyHeightAndVisibility);
+			FLandscapeComponentDataInterface DataInterface(LandscapeComponent, FinalLODToExport, false);
 
-		UE::Landscape::Nanite::FAsyncComponentData AsyncComponentData;
+			UE::Landscape::Nanite::FAsyncComponentData AsyncComponentData;
 
-		DataInterface.GetHeightmapTextureData(AsyncComponentData.HeightAndNormalData, false);
-		DataInterface.GetWeightmapTextureData(LandscapeComponent->GetVisibilityLayer(), AsyncComponentData.Visibility);
+			DataInterface.GetHeightmapTextureData(AsyncComponentData.HeightAndNormalData, false);
+			DataInterface.GetWeightmapTextureData(LandscapeComponent->GetVisibilityLayer(), AsyncComponentData.Visibility);
 
-		AsyncComponentData.ComponentDataInterface = MakeShared<FLandscapeComponentDataInterfaceBase>(LandscapeComponent, FinalLODToExport, false);
-		int32 HeightmapSize = ((LandscapeComponent->SubsectionSizeQuads + 1) * LandscapeComponent->NumSubsections) >> FinalLODToExport;
-		AsyncComponentData.ComponentDataInterface->HeightmapStride = HeightmapSize;
-		AsyncComponentData.ComponentDataInterface->HeightmapComponentOffsetX = 0;
-		AsyncComponentData.ComponentDataInterface->HeightmapComponentOffsetY = 0;
-		AsyncBuildData->ComponentData.Add(LandscapeComponent, AsyncComponentData);
-	}
+			AsyncComponentData.ComponentDataInterface = MakeShared<FLandscapeComponentDataInterfaceBase>(LandscapeComponent, FinalLODToExport, false);
+			int32 HeightmapSize = ((LandscapeComponent->SubsectionSizeQuads + 1) * LandscapeComponent->NumSubsections) >> FinalLODToExport;
+			AsyncComponentData.ComponentDataInterface->HeightmapStride = HeightmapSize;
+			AsyncComponentData.ComponentDataInterface->HeightmapComponentOffsetX = 0;
+			AsyncComponentData.ComponentDataInterface->HeightmapComponentOffsetY = 0;
+			AsyncBuildData->ComponentData.Add(LandscapeComponent, AsyncComponentData);
+		}
 
-	// Create the UStaticMesh in advance, from the game thread.
-	UPackage* Package = GetPackage();
-	AsyncBuildData->NaniteStaticMesh.Reset(NewObject<UStaticMesh>(/*Outer = */Package, MakeUniqueObjectName(/*Parent = */Package, UStaticMesh::StaticClass(), TEXT("LandscapeNaniteMesh"))));
+		// Create the UStaticMesh in advance, from the game thread.
+		UPackage* Package = GetPackage();
+		AsyncBuildData->NaniteStaticMesh.Reset(NewObject<UStaticMesh>(/*Outer = */Package, MakeUniqueObjectName(/*Parent = */Package, UStaticMesh::StaticClass(), TEXT("LandscapeNaniteMesh"))));
+	}
 
 	return AsyncBuildData;
 }
@@ -5388,9 +5391,9 @@ void ALandscape::PostRegisterAllComponents()
 	};
 
 	// Until it is properly supported, ALandscape with layer brushes will force all of its proxies to be loaded in editor
-	if (GEditor && !GetWorld()->IsGameWorld() && LandscapeGuid.IsValid() && HasValidBrush() && !IsRunningCommandlet())
+	if (UWorld* World = GetWorld(); GEditor && (World != nullptr) && !World->IsGameWorld() && LandscapeGuid.IsValid() && HasValidBrush() && !IsRunningCommandlet())
 	{
-		if (UWorldPartition* WorldPartition = GetWorld()->GetWorldPartition())
+		if (UWorldPartition* WorldPartition = World->GetWorldPartition())
 		{
 			FWorldPartitionHelpers::ForEachActorDescInstance<ALandscapeProxy>(WorldPartition, [this, WorldPartition](const FWorldPartitionActorDescInstance* ActorDescInstance)
 			{
@@ -6279,6 +6282,10 @@ void ALandscapeProxy::PostEditChangeProperty(FPropertyChangedEvent& PropertyChan
 	bool bChangedPhysMaterial = false;
 	bool bChangedMaterial = false;
 
+	// In some cases, GetWorld() can return nullptr here (e.g. when replacing references when deleting a world asset from the content browser)
+	//  , so we should not assume the World is valid here
+	UWorld* World = GetWorld();
+
 	if (MemberPropertyName == FName(TEXT("RelativeScale3D")))
 	{
 		// RelativeScale3D isn't even a property of ALandscapeProxy, it's a property of the root component
@@ -6364,7 +6371,10 @@ void ALandscapeProxy::PostEditChangeProperty(FPropertyChangedEvent& PropertyChan
 	else if (MemberPropertyName == GET_MEMBER_NAME_CHECKED(ALandscapeProxy, StreamingDistanceMultiplier))
 	{
 		// Recalculate in a few seconds.
-		GetWorld()->TriggerStreamingDataRebuild();
+		if (World != nullptr)
+		{
+			World->TriggerStreamingDataRebuild();
+		}
 	}
 	else if (MemberPropertyName == GET_MEMBER_NAME_CHECKED(ALandscapeProxy, DefaultPhysMaterial))
 	{
@@ -6498,22 +6508,19 @@ void ALandscapeProxy::PostEditChangeProperty(FPropertyChangedEvent& PropertyChan
 
 			UpdateAllComponentMaterialInstances();
 
-			if (UWorld* World = GetWorld())
+			if ((World != nullptr) && (World->GetFeatureLevel() == ERHIFeatureLevel::ES3_1))
 			{
-				if (World->GetFeatureLevel() == ERHIFeatureLevel::ES3_1)
+				for (ULandscapeComponent* Component : LandscapeComponents)
 				{
-					for (ULandscapeComponent* Component : LandscapeComponents)
+					if (Component != nullptr)
 					{
-						if (Component != nullptr)
-						{
-							Component->CheckGenerateMobilePlatformData(/*bIsCooking = */ false, /*TargetPlatform = */ nullptr);
-						}
+						Component->CheckGenerateMobilePlatformData(/*bIsCooking = */ false, /*TargetPlatform = */ nullptr);
 					}
 				}
-
-				InvalidateOrUpdateNaniteRepresentation(/* bInCheckContentId = */true, /*InTargetPlatform = */nullptr);
-				UpdateRenderingMethod();
 			}
+
+			InvalidateOrUpdateNaniteRepresentation(/* bInCheckContentId = */true, /*InTargetPlatform = */nullptr);
+			UpdateRenderingMethod();
 		}
 	}
 
@@ -6530,9 +6537,14 @@ void ALandscapeProxy::PostEditChangeProperty(FPropertyChangedEvent& PropertyChan
 		ChangedPhysMaterial();
 	}
 
-	if (ULandscapeSubsystem* Subsystem = GetWorld()->GetSubsystem<ULandscapeSubsystem>(); Subsystem != nullptr && bChangedMaterial)
+	if (bChangedMaterial)
 	{
-		Subsystem->GetDelegateAccess().OnLandscapeProxyMaterialChangedDelegate.Broadcast(this, FOnLandscapeProxyMaterialChangedParams());
+		if (World != nullptr)
+		{
+			ULandscapeSubsystem* Subsystem = World->GetSubsystem<ULandscapeSubsystem>();
+			check(Subsystem != nullptr);
+			Subsystem->GetDelegateAccess().OnLandscapeProxyMaterialChangedDelegate.Broadcast(this, FOnLandscapeProxyMaterialChangedParams());
+		}
 
 PRAGMA_DISABLE_DEPRECATION_WARNINGS
 		LandscapeMaterialChangedDelegate.Broadcast();
@@ -6549,7 +6561,10 @@ bool ALandscapeStreamingProxy::CanEditChange(const FProperty* InProperty) const
 
 	if (InProperty && InProperty->GetFName() == GET_MEMBER_NAME_CHECKED(ALandscapeStreamingProxy, LandscapeActorRef))
 	{
-		return !GetWorld()->GetSubsystem<ULandscapeSubsystem>()->IsGridBased();
+		if (UWorld* World = GetWorld())
+		{
+			return !World->GetSubsystem<ULandscapeSubsystem>()->IsGridBased();
+		}
 	}
 
 	return true;
@@ -6621,9 +6636,9 @@ void ALandscapeStreamingProxy::PostRegisterAllComponents()
 	{
 		ULandscapeInfo* LandscapeInfo = GetLandscapeInfo();
 		check(LandscapeInfo);
-		if (GEditor && !GetWorld()->IsGameWorld() && !IsRunningCommandlet())
+		if (UWorld* World = GetWorld(); GEditor && (World != nullptr) && !World->IsGameWorld() && !IsRunningCommandlet())
 		{
-			if (UWorldPartition* WorldPartition = GetWorld()->GetWorldPartition(); WorldPartition && WorldPartition->IsInitialized())
+			if (UWorldPartition* WorldPartition = World->GetWorldPartition(); WorldPartition && WorldPartition->IsInitialized())
 			{
 				const FVector ActorLocation = GetActorLocation();
 				FBox Bounds(ActorLocation, ActorLocation + (GetGridSize() * LandscapeInfo->DrawScale));
@@ -6750,7 +6765,7 @@ bool ULandscapeInfo::CanDeleteLandscape(FText& OutReason) const
 	if (AActor* Actor = LandscapeActor.Get())
 	{
 		UWorld* World = Actor->GetWorld();
-		if (UWorldPartition* WorldPartition = World->GetWorldPartition())
+		if (UWorldPartition* WorldPartition = (World != nullptr) ? World->GetWorldPartition() : nullptr)
 		{
 			FWorldPartitionHelpers::ForEachActorDescInstance<ALandscapeProxy>(WorldPartition, [this, &UndeletedProxyCount](const FWorldPartitionActorDescInstance* ActorDescInstance)
 			{
@@ -6812,7 +6827,7 @@ bool ULandscapeInfo::CanDeleteLandscape(FText& OutReason) const
 	if (AActor* Actor = LandscapeActor.Get())
 	{
 		UWorld* World = Actor->GetWorld();
-		if (UWorldPartition* WorldPartition = World->GetWorldPartition())
+		if (UWorldPartition* WorldPartition = (World != nullptr) ? World->GetWorldPartition() : nullptr)
 		{
 			FWorldPartitionHelpers::ForEachActorDescInstance<ALandscapeSplineActor>(WorldPartition, [this, &UndeletedSplineCount](const FWorldPartitionActorDescInstance* ActorDescInstance)
 			{
