Removed spurious assert when the landscape material references a texture named "Heightmap" (implicit name for heightmaps in landscape). It was previously assumed the heightmap could only come from within landscape but now it's obvious it's actually possible by other means in the material graph.
The outcome is simply that this texture will then be registered twice to the streaming texture manager for a given primitive component, but the manager is already able to handle this situation so the assert is just useless

diff --git a/Engine/Source/Runtime/Landscape/Private/LandscapeRender.cpp b/Engine/Source/Runtime/Landscape/Private/LandscapeRender.cpp
--- a/Engine/Source/Runtime/Landscape/Private/LandscapeRender.cpp
+++ b/Engine/Source/Runtime/Landscape/Private/LandscapeRender.cpp
@@ -4420,12 +4420,9 @@ void ULandscapeComponent::GetStreamingRenderAssetInfo(FStreamingTextureLevelCont
 		}
 	}
 
-	// Heightmap has not been accounted for by GetUsedTextures on the material :
+	// Heightmap has usually not been accounted for by GetUsedTextures on the material (unless referenced directly by name, but that's borderline a hack) :
 	if (HeightmapTexture)
 	{
-		// Heightmap should not have been accounted for already:
-		check(OutStreamingRenderAssets.FindByPredicate([this](const FStreamingRenderAssetPrimitiveInfo& StreamingWeightmap) { return StreamingWeightmap.RenderAsset == HeightmapTexture; }) == nullptr);
-
 		FStreamingRenderAssetPrimitiveInfo& StreamingHeightmap = *new(OutStreamingRenderAssets)FStreamingRenderAssetPrimitiveInfo;
 		StreamingHeightmap.Bounds = BoundingSphere;
 
