Changing comments on landscape grass, to provide info on subtle assumptions of the grass alignment code.

diff --git a/Engine/Source/Runtime/Landscape/Private/LandscapeGrass.cpp b/Engine/Source/Runtime/Landscape/Private/LandscapeGrass.cpp
--- a/Engine/Source/Runtime/Landscape/Private/LandscapeGrass.cpp
+++ b/Engine/Source/Runtime/Landscape/Private/LandscapeGrass.cpp
@@ -2072,6 +2072,8 @@ struct FAsyncGrassBuilder : public FGrassBuilderBase
 
 		auto ApplyNormalAlignmentTransform = [](const FVector& Normal, FMatrix& InOutXForm)
 			{
+				// Note: NewZ should never be perpendicular to Z axis, as it is constructed from a heightmap.
+				// Therefore it should never have a zero cross product with (0,-1,0)
 				const FVector NewZ = Normal * FMath::Sign(Normal.Z);
 				const FVector NewX = (FVector(0, -1, 0) ^ NewZ).GetSafeNormal();
 				const FVector NewY = NewZ ^ NewX;
@@ -2215,6 +2217,8 @@ struct FAsyncGrassBuilder : public FGrassBuilderBase
 										FVector PosY1 = yStart ? Instances[InstanceIndex - 1].Pos : Instance.Pos;
 										FVector PosY2 = (yStart + 1 < SqrtMaxInstances) ? Instances[InstanceIndex + 1].Pos : Instance.Pos;
 
+										// this only fails if your grid is 1x1, or if the jittered offsets of neighboring grids randomly intersected
+										// in either case, we can't align, so leave the grass transform alone
 										if (PosX1 != PosX2 && PosY1 != PosY2)
 										{
 											FVector Normal = ((PosX1 - PosX2) ^ (PosY1 - PosY2)).GetSafeNormal();
