Add support to selectively render landscape edit layer heightmaps using batch merge code.

diff --git a/Engine/Source/Runtime/Landscape/Classes/Landscape.h b/Engine/Source/Runtime/Landscape/Classes/Landscape.h
--- a/Engine/Source/Runtime/Landscape/Classes/Landscape.h
+++ b/Engine/Source/Runtime/Landscape/Classes/Landscape.h
@@ -246,6 +246,37 @@ struct FOnLandscapeEditLayersMergedParams
 	bool bIsHeightmapMerge = false;
 };
 
+
+// Parameters to describe a subset of a landscape heightmap to render.  Can select a subset of edit layers and a subregion of the landscape coordinates.
+struct FLandscapeEditLayerRenderHeightParams
+{
+	// Region to render, in Landscape coordinates
+	FIntRect Bounds;  
+
+	// Which edit layers to include.  Order matches ALandscape::GetEditLayers().
+	TBitArray<> ActiveEditLayers;
+
+	// Provide either a cpu buffer or render target to receive the data.  Buffer should be allocated large enough to receive the samples in Bounds.
+	TArrayView<uint16> CpuResult;
+	UTextureRenderTarget2D* RTResult = nullptr;
+};
+
+// Parameters to describe a subset of a landscape weightmaps to render.  Can select a subset of edit layers and a subregion of the landscape coordinates, for a list of weightmaps.
+struct FLandscapeEditLayerRenderWeightParams
+{
+	// Region to render, in Landscape coordinates
+	FIntRect Bounds;
+
+	// Which edit layers to include.  Order matches ALandscape::GetEditLayers().
+	TBitArray<> ActiveEditLayers;
+
+	TArrayView<FName> WeightmapNames;
+
+	// Provide either a cpu buffer or render target to receive the data.  Buffer should be allocated large enough to receive the samples in Bounds.
+	TArrayView<uint8> CpuResult;
+	UTextureRenderTarget2D* RTResult = nullptr;
+};
+
 UCLASS(MinimalAPI, showcategories=(Display, Movement, Collision, Lighting, LOD, Input), hidecategories=(Mobility))
 class ALandscape : public ALandscapeProxy
 #if CPP && WITH_EDITOR // UHT doesn't support inheriting from namespaced class
@@ -501,6 +532,10 @@ PRAGMA_ENABLE_DEPRECATION_WARNINGS
 	
 	LANDSCAPE_API void ToggleCanHaveLayersContent();
 	LANDSCAPE_API void ForceUpdateLayersContent(bool bIntermediateRender = false);
+
+	LANDSCAPE_API bool SelectiveRenderEditLayersHeightmaps(const FLandscapeEditLayerRenderHeightParams& InSelectiveRenderParams);
+	LANDSCAPE_API bool SelectiveRenderEditLayersWeightmaps(const FLandscapeEditLayerRenderWeightParams& InSelectiveRenderParams);
+
 	void FlushLayerContentThisFrame();
 
 	UFUNCTION(BlueprintCallable, Category = "Landscape")
@@ -530,6 +565,8 @@ PRAGMA_ENABLE_DEPRECATION_WARNINGS
 
 	LANDSCAPE_API FName GenerateUniqueLayerName(FName InName = NAME_None) const;
 
+	UE_INTERNAL LANDSCAPE_API bool AllowSelectiveRenderEditLayers() const;
+
 private:
 	FLandscapeLayer* GetLayerInternal(int32 InLayerIndex);
 	void OnLayerCreatedInternal(ULandscapeEditLayerBase* EditLayer);
@@ -537,7 +574,7 @@ private:
 
 	bool SupportsEditLayersLocalMerge();
 
-	ELandscapeEditLayersMergeMode GetEditLayersMergeMode();
+	ELandscapeEditLayersMergeMode GetEditLayersMergeMode() const;
 	bool CreateLayersRenderingResource(bool bInForceLayerResourceReset);
 	void PrepareEditLayersLocalMergeResources();
 	void UpdateLayersContent(bool bInWaitForStreaming = false, bool bInSkipMonitorLandscapeEdModeChanges = false, bool bIntermediateRender = false, bool bFlushRender = false);
@@ -548,6 +585,7 @@ private:
 	//~ Begin UE::Landscape::EditLayers::IEditLayerRendererProvider implementation
 	virtual TArray<UE::Landscape::EditLayers::FEditLayerRendererState> GetEditLayerRendererStates(const UE::Landscape::EditLayers::FMergeContext* InMergeContext) override;
 	//~ End UE::Landscape::EditLayers::IEditLayerRendererProvider implementation
+	TArray<UE::Landscape::EditLayers::FEditLayerRendererState> GetEditLayerRendererStatesEnableOverride(const UE::Landscape::EditLayers::FMergeContext* InMergeContext, const TBitArray<>& InLayerEnableOverride);
 
 	UE::Landscape::EditLayers::FMergeRenderContext PrepareEditLayersMergeRenderContext(const UE::Landscape::EditLayers::FMergeContext& InMergeContext, const UE::Landscape::EditLayers::FMergeRenderParams& InParams);
 
@@ -559,6 +597,9 @@ private:
 	int32 PerformLayersHeightmapsBatchedMerge(const FUpdateLayersContentContext& InUpdateLayersContentContext, const FEditLayersHeightmapMergeParams& InMergeParams);
 	void ResolveLayersHeightmapTexture(const FTextureToComponentHelper& MapHelper, const TSet<UTexture2D*>& HeightmapsToResolve, bool bIntermediateRender, TArray<FLandscapeEditLayerComponentReadbackResult>& InOutComponentReadbackResults);
 
+	bool PerformSelectiveRenderEditLayersHeightmapsCPU(const FUpdateLayersContentContext& InUpdateLayersContentContext, const FLandscapeEditLayerRenderHeightParams& InSelectiveRenderParams);
+	bool PerformSelectiveRenderEditLayersHeightmapsRT(const FUpdateLayersContentContext& InUpdateLayersContentContext, const FLandscapeEditLayerRenderHeightParams& InSelectiveRenderParams);
+
 	int32 RegenerateLayersWeightmaps(FUpdateLayersContentContext& InUpdateLayersContentContext);
 	int32 PerformLayersWeightmapsLocalMerge(FUpdateLayersContentContext& InUpdateLayersContentContext, const FEditLayersWeightmapMergeParams& InMergeParams);
 	int32 PerformLayersWeightmapsGlobalMerge(FUpdateLayersContentContext& InUpdateLayersContentContext, const FEditLayersWeightmapMergeParams& InMergeParams);
diff --git a/Engine/Source/Runtime/Landscape/Private/LandscapeEditLayerMergeRenderContext.cpp b/Engine/Source/Runtime/Landscape/Private/LandscapeEditLayerMergeRenderContext.cpp
--- a/Engine/Source/Runtime/Landscape/Private/LandscapeEditLayerMergeRenderContext.cpp
+++ b/Engine/Source/Runtime/Landscape/Private/LandscapeEditLayerMergeRenderContext.cpp
@@ -1127,6 +1127,10 @@ const FMergeRenderBatch* FMergeRenderContext::GetCurrentRenderBatch() const
 	return RenderBatches.IsValidIndex(CurrentRenderBatchIndex) ? &RenderBatches[CurrentRenderBatchIndex] : nullptr;
 }
 
+const int32 FMergeRenderContext::GetCurrentRenderBatchIdx() const
+{
+	return RenderBatches.IsValidIndex(CurrentRenderBatchIndex) ? CurrentRenderBatchIndex : INDEX_NONE;
+}
 
 bool FMergeRenderContext::IsValid() const
 {
diff --git a/Engine/Source/Runtime/Landscape/Private/LandscapeEditLayerRendererState.cpp b/Engine/Source/Runtime/Landscape/Private/LandscapeEditLayerRendererState.cpp
--- a/Engine/Source/Runtime/Landscape/Private/LandscapeEditLayerRendererState.cpp
+++ b/Engine/Source/Runtime/Landscape/Private/LandscapeEditLayerRendererState.cpp
@@ -58,7 +58,7 @@ FEditLayerRendererState::FEditLayerRendererState(const FMergeContext* InMergeCon
 
 void FEditLayerRendererState::EnableTargetType(ELandscapeToolTargetType InTargetType)
 {
-	checkf(!EnumHasAnyFlags(SupportedTargetTypeState.GetTargetTypeMask(), GetLandscapeToolTargetTypeAsFlags(InTargetType)),
+	checkf(EnumHasAllFlags(SupportedTargetTypeState.GetTargetTypeMask(), GetLandscapeToolTargetTypeAsFlags(InTargetType)),
 		TEXT("Target type %s cannot be enabled on this renderer state because it is not supported. Make sure that target types are supported before enabling them"), *UEnum::GetValueAsString(InTargetType));
 	EnabledTargetTypeState.AddTargetType(InTargetType);
 	UpdateActiveTargetTypeState();
diff --git a/Engine/Source/Runtime/Landscape/Private/LandscapeEditLayers.cpp b/Engine/Source/Runtime/Landscape/Private/LandscapeEditLayers.cpp
--- a/Engine/Source/Runtime/Landscape/Private/LandscapeEditLayers.cpp
+++ b/Engine/Source/Runtime/Landscape/Private/LandscapeEditLayers.cpp
@@ -90,6 +90,9 @@ LandscapeEditLayers.cpp: Landscape editing layers mode
 #include "VisualLogger/VisualLogger.h"
 #include "LandscapeTextureStorageProvider.h"
 #include "UObject/Script.h"
+#include "RHIGPUReadback.h"
+#include "EngineUtils.h"
+#include "ImageUtils.h"
 #endif
 
 #define LOCTEXT_NAMESPACE "Landscape"
@@ -266,6 +269,13 @@ TAutoConsoleVariable<bool> CVarSilenceMergeBatchResolutionWarning(
 	false,
 	TEXT("When true, don't warn about about exceeding batch merge resolution from landscape.EditLayersLocalMerge.MaxResolutionPerRenderBatch"));
 
+// TODO:  remove with IntermediateRender.
+TAutoConsoleVariable<bool> CVarLandscapeSelectiveRenderEditLayers(
+	TEXT("landscape.BatchedMerge.SelectiveRenderEditLayers"),
+	true,
+	TEXT("When true, use SelectiveRenderEditLayers codepath instead of intermediate renders"));
+
+
 void OnLandscapeEditLayersLocalMergeChanged(IConsoleVariable* CVar)
 {
 	for (TObjectIterator<UWorld> It(/*AdditionalExclusionFlags = */RF_ClassDefaultObject, /*bIncludeDerivedClasses = */true, /*InInternalExclusionFlags = */EInternalObjectFlags::Garbage); It; ++It)
@@ -320,6 +330,15 @@ TAutoConsoleVariable<int32> CVarLandscapeEditLayersClearBeforeEachWriteToScratch
 	TEXT("Debug to help with RenderDoc debugging : clear each time we're about to write on a scratch render target (since those are reused and can be used to write RTs of different resolutions"),
 	ECVF_RenderThreadSafe);
 
+
+static void LandscapeDumpSelectiveLayerRender(const TArray<FString>& Args);
+
+FAutoConsoleCommand CmdLandscapeDumpSelectiveLayerRender(
+	TEXT("landscape.DumpSelectiveLayerRender"),
+	TEXT("Perform a selective layer render of a landscape and write the output to an image file."),
+	FConsoleCommandWithArgsDelegate::CreateStatic(&LandscapeDumpSelectiveLayerRender)
+);
+
 struct FLandscapeDirty
 {
 	FLandscapeDirty()
@@ -382,6 +401,14 @@ struct FTextureToComponentHelper
 		Refresh(ERefreshFlags::RefreshAll);
 	}
 
+	FTextureToComponentHelper(const ULandscapeInfo& InLandscapeInfo, TArrayView<ULandscapeComponent*> Components, ERefreshFlags Flags)
+		: LandscapeInfo(&InLandscapeInfo),
+		LandscapeComponents(Components)
+	{
+		check(!EnumHasAnyFlags(Flags, ERefreshFlags::RefreshComponents));  // Expect Height and/or Weight
+		Refresh(Flags);
+	}
+
 	void Refresh(ERefreshFlags InRefreshFlags)
 	{
 		TRACE_CPUPROFILER_EVENT_SCOPE(TextureToComponentHelper_Refresh);
@@ -2479,20 +2506,20 @@ bool ALandscape::ComputeLandscapeLayerBrushInfo(FTransform& OutLandscapeTransfor
 	return true;
 }
 
-ELandscapeEditLayersMergeMode ALandscape::GetEditLayersMergeMode()
+ELandscapeEditLayersMergeMode ALandscape::GetEditLayersMergeMode() const
 {
 	ELandscapeEditLayersMergeMode DesiredMode = static_cast<ELandscapeEditLayersMergeMode>(FMath::Min<uint8>(static_cast<uint8>(LandscapeEditLayersLocalMerge), static_cast<uint8>(ELandscapeEditLayersMergeMode::Invalid) - 1));
 
 	if (DesiredMode == ELandscapeEditLayersMergeMode::LocalMerge)
-{
-	ALandscape* Landscape = GetLandscapeActor();
-	check(Landscape);
-
-	for (const FLandscapeLayer& Layer : Landscape->LandscapeEditLayers)
 	{
-		// No BP brush is supported for local merge of edit layers : 
-		if (Layer.EditLayer->IsVisible() && Algo::AnyOf(Layer.Brushes, [](const FLandscapeLayerBrush& Brush) { return (Brush.GetBrush() != nullptr) && Brush.GetBrush()->IsVisible(); }))
+		const ALandscape* Landscape = GetLandscapeActor();
+		check(Landscape);
+
+		for (const FLandscapeLayer& Layer : Landscape->LandscapeEditLayers)
 		{
+			// No BP brush is supported for local merge of edit layers : 
+			if (Layer.EditLayer->IsVisible() && Algo::AnyOf(Layer.Brushes, [](const FLandscapeLayerBrush& Brush) { return (Brush.GetBrush() != nullptr) && Brush.GetBrush()->IsVisible(); }))
+			{
 				return ELandscapeEditLayersMergeMode::GlobalMerge;
 			}
 		}
@@ -2501,6 +2528,14 @@ ELandscapeEditLayersMergeMode ALandscape::GetEditLayersMergeMode()
 	return DesiredMode;
 }
 
+bool ALandscape::AllowSelectiveRenderEditLayers() const
+{
+	// TODO:  remove when GlobalMerge goes
+
+	return GetEditLayersMergeMode() == ELandscapeEditLayersMergeMode::BatchedMerge &&
+		   CVarLandscapeSelectiveRenderEditLayers.GetValueOnGameThread();
+}
+
 bool ALandscape::CreateLayersRenderingResource(bool bInForceLayerResourceReset)
 {
 	ULandscapeInfo* Info = GetLandscapeInfo();
@@ -4469,14 +4504,30 @@ struct FUpdateLayersContentContext
 	};
 	FRIEND_ENUM_CLASS_FLAGS(ERefreshFlags);
 
+	// Set up the layer content context for a selective render.  Rendering the components provided by InMapHelper based on their LayerUpdateFlags.
 	FUpdateLayersContentContext(const FTextureToComponentHelper& InMapHelper, bool bInPartialUpdate)
 		: bPartialUpdate(bInPartialUpdate)
+		, bSelectiveRender(false)
 		, MapHelper(InMapHelper)
 	{
 		// No need to update the map helper, it's assumed to be already ready in the constructor
 		Refresh(ERefreshFlags::RefreshComponentInfos | ERefreshFlags::RefreshHeightmapInfos | ERefreshFlags::RefreshWeightmapInfos);
 	}
 
+	// Set up the layer content context for a selective render.  Rendering the components provided by InMapHelper, regardless of their LayerUpdateFlags.
+	// Can specify RefreshHeightmapInfos and/or RefreshWeightmapInfos.
+	struct FSelectiveRenderTag {};
+	FUpdateLayersContentContext(const FTextureToComponentHelper& InMapHelper, FSelectiveRenderTag, ERefreshFlags HeightOrWeightFlags)
+		: bPartialUpdate(false)
+		, bSelectiveRender(true)
+		, MapHelper(InMapHelper)
+	{
+		// expect only RefreshHeightmapInfos and/or RefreshWeightmapInfos
+		check((HeightOrWeightFlags & (ERefreshFlags::RefreshHeightmapInfos | ERefreshFlags::RefreshWeightmapInfos)) != ERefreshFlags::None);
+		check((HeightOrWeightFlags & ~(ERefreshFlags::RefreshHeightmapInfos | ERefreshFlags::RefreshWeightmapInfos)) == ERefreshFlags::None);  
+		Refresh(ERefreshFlags::RefreshComponentInfos | HeightOrWeightFlags);
+	}
+
 	FTextureToComponentHelper::ERefreshFlags RefreshFlagsToMapHelperRefreshFlags(ERefreshFlags InRefreshFlags)
 	{
 		FTextureToComponentHelper::ERefreshFlags MapHelperRefreshFlags = FTextureToComponentHelper::ERefreshFlags::None;
@@ -4507,8 +4558,8 @@ struct FUpdateLayersContentContext
 		// Then triage the dirty/non-dirty components  :
 		if (EnumHasAnyFlags(InRefreshFlags, ERefreshFlags::RefreshComponentInfos))
 		{
-			// When components are refreshed, all other info has to be :
-			check(EnumHasAllFlags(InRefreshFlags, ERefreshFlags::RefreshHeightmapInfos | ERefreshFlags::RefreshWeightmapInfos));
+			// When components are refreshed, all other info has to be, except in selective renders :
+			check(bSelectiveRender || EnumHasAllFlags(InRefreshFlags, ERefreshFlags::RefreshHeightmapInfos | ERefreshFlags::RefreshWeightmapInfos));
 
 			DirtyLandscapeComponents.Reset();
 			NonDirtyLandscapeComponents.Reset();
@@ -4660,12 +4711,15 @@ struct FUpdateLayersContentContext
 			// All selected components will have to be resolved : 
 			AllLandscapeComponentsToResolve.Append(AllLandscapeComponents);
 
-			// Add components with deferred flag to update list
-			for (ULandscapeComponent* Component : AllLandscapeComponents)
+			if (!bSelectiveRender)
 			{
-				if (Component->GetLayerUpdateFlagPerMode() & ELandscapeLayerUpdateMode::Update_Client_Deferred)
+				// Add components with deferred flag to update list
+				for (ULandscapeComponent* Component : AllLandscapeComponents)
 				{
-					UE::Landscape::Private::FindOrAddByComponent(AllLandscapeComponentReadbackResults, Component, ELandscapeLayerUpdateMode::Update_Client_Deferred);
+					if (Component->GetLayerUpdateFlagPerMode() & ELandscapeLayerUpdateMode::Update_Client_Deferred)
+					{
+						UE::Landscape::Private::FindOrAddByComponent(AllLandscapeComponentReadbackResults, Component, ELandscapeLayerUpdateMode::Update_Client_Deferred);
+					}
 				}
 			}
 		}
@@ -4673,6 +4727,8 @@ struct FUpdateLayersContentContext
 
 	// Indicates whether all components of the landscape are marked dirty :
 	const bool bPartialUpdate = false;
+	const bool bSelectiveRender = false;
+
 	// Helper to gather mappings between heightmaps/weightmaps and components :
 	FTextureToComponentHelper MapHelper;
 	// List of landscape components that have been made dirty and need to be updated : 
@@ -4701,9 +4757,9 @@ ENUM_CLASS_FLAGS(FUpdateLayersContentContext::ERefreshFlags);
 // Little struct that holds information common to PerformLayersHeightmapsLocalMerge, PerformLayersHeightmapsBatchedMerge and PerformLayersHeightmapsGlobalMerge
 struct FEditLayersHeightmapMergeParams
 {
-	int32 HeightmapUpdateModes;
-	bool bForceRender;
-	bool bSkipBrush;
+	int32 HeightmapUpdateModes = 0;
+	bool bForceRender = false;
+	bool bSkipBrush = false;
 };
 
 // Render-thread version of the data / functions we need for the local merge of edit layers : 
@@ -5466,14 +5522,42 @@ int32 ALandscape::PerformLayersHeightmapsLocalMerge(const FUpdateLayersContentCo
 }
 
 TArray<UE::Landscape::EditLayers::FEditLayerRendererState> ALandscape::GetEditLayerRendererStates(const UE::Landscape::EditLayers::FMergeContext* InMergeContext)
+{
+	return GetEditLayerRendererStatesEnableOverride(InMergeContext, {});
+}
+
+TArray<UE::Landscape::EditLayers::FEditLayerRendererState> ALandscape::GetEditLayerRendererStatesEnableOverride(const UE::Landscape::EditLayers::FMergeContext* InMergeContext, const TBitArray<>& InLayerEnableOverride)
 {
 	using namespace UE::Landscape::EditLayers;
 
 	TArray<FEditLayerRendererState> RendererStates;
+	int32 LayerIdx = 0;
+	bool bUseEnableOverride = InLayerEnableOverride.Num() > 0;
+	check(!bUseEnableOverride || InLayerEnableOverride.Num() == LandscapeEditLayers.Num());
 	for (FLandscapeLayer& Layer : LandscapeEditLayers)
 	{
 		TArray<FEditLayerRendererState> LayerRendererStates = Layer.GetEditLayerRendererStates(InMergeContext);
+
+		if (bUseEnableOverride)
+		{
+			ELandscapeToolTargetType TargetType = InMergeContext->IsHeightmapMerge() ? ELandscapeToolTargetType::Heightmap : ELandscapeToolTargetType::Weightmap;
+			bool bEnabled = InLayerEnableOverride[LayerIdx];
+			Algo::ForEach(LayerRendererStates, [TargetType, bEnabled](FEditLayerRendererState& State)
+				{
+					if (bEnabled)
+					{
+						State.EnableTargetType(TargetType);
+					}
+					else
+					{
+						State.DisableTargetType(TargetType);
+					}
+				});
+		}
+
 		RendererStates.Append(LayerRendererStates);
+
+		++LayerIdx;
 	}
 	return RendererStates;
 }
@@ -7964,6 +8048,229 @@ int32 ALandscape::PerformLayersHeightmapsBatchedMerge(const FUpdateLayersContent
 	return InMergeParams.HeightmapUpdateModes;
 }
 
+bool ALandscape::PerformSelectiveRenderEditLayersHeightmapsCPU(const FUpdateLayersContentContext& InUpdateLayersContentContext, const FLandscapeEditLayerRenderHeightParams& InSelectiveRenderParams)
+{
+	TRACE_CPUPROFILER_EVENT_SCOPE(PerformSelectiveRenderEditLayersHeightmapsCPU);
+	RHI_BREADCRUMB_EVENT_GAMETHREAD("PerformSelectiveRenderEditLayersHeightmapsCPU");
+
+	using namespace UE::Landscape;
+	using namespace UE::Landscape::Private;
+	using namespace UE::Landscape::EditLayers;
+	using namespace UE::Landscape::EditLayers::Private;
+
+	check(InSelectiveRenderParams.CpuResult.Num() >= InSelectiveRenderParams.Bounds.Area());
+
+	// Set up the merge batches similar to PerformLayersHeightmapsBatchedMerge.  Difference is that the edit layer visibility is overridden
+	// and ULandscapeHeightmapNormalsEditLayerRenderer is omitted.  The completion callback is very different.
+	// Note: not using SelectiveBounds in the actual Merge process, just to retrieve the data.  The render region is based only on the
+	// components from InUpdateLayersContentContext.
+
+	FMergeContext MergeContext(/*InLandscape = */this, /*bInIsHeightmapMerge = */true, /*bInSkipProceduralRenderers = */false);
+	TArray<FEditLayerRendererState> RendererStates = GetEditLayerRendererStatesEnableOverride(&MergeContext, InSelectiveRenderParams.ActiveEditLayers);
+
+	FMergeRenderParams MergeRenderParams(InUpdateLayersContentContext.LandscapeComponentsHeightmapsToRender, RendererStates);
+	FMergeRenderContext MergeRenderContext = PrepareEditLayersMergeRenderContext(MergeContext, MergeRenderParams);
+	if (!MergeRenderContext.IsValid())
+	{
+		// Nothing to do 
+		return false;
+	}
+
+	// Set up a FRHIGPUTextureReadback for each merge batch
+	struct FSyncRenderBatchResult
+	{
+		FRHIGPUTextureReadback Readback = { "SyncRenderResult" };
+		FIntRect SectionRect;
+		FIntPoint Size;  // RT size
+		TArray<FIntRect> ComponentRects;
+	};
+	TArray<FSyncRenderBatchResult> SyncRenderBatchResults;
+	SyncRenderBatchResults.AddDefaulted(MergeRenderContext.GetRenderBatches().Num());  // allocate in advance
+
+	// Batch completion callback
+	auto EnqueueBatchReadback = [&InUpdateLayersContentContext, &InSelectiveRenderParams, &SyncRenderBatchResults]
+	(const FMergeRenderContext::FOnRenderBatchTargetGroupDoneParams& InParams, UE::Landscape::FRDGBuilderRecorder& RDGBuilderRecorder)
+		{
+			TRACE_CPUPROFILER_EVENT_SCOPE(ExtractSelectiveMergeData);
+			RHI_BREADCRUMB_EVENT_GAMETHREAD("ExtractSelectiveMergeData");
+
+			ALandscape* Landscape = InParams.MergeRenderContext->GetLandscape();
+			const FMergeRenderBatch* RenderBatch = InParams.MergeRenderContext->GetCurrentRenderBatch();
+			int32 CurrentBatch = InParams.MergeRenderContext->GetCurrentRenderBatchIdx();
+			FSyncRenderBatchResult& BatchResult = SyncRenderBatchResults[CurrentBatch];
+
+			BatchResult.SectionRect = RenderBatch->SectionRect;
+			BatchResult.Size = RenderBatch->Resolution;  //Only used for sanity check on bounds
+			BatchResult.ComponentRects.Reserve(RenderBatch->ComponentsToRender.Num());
+			for (ULandscapeComponent* Component : RenderBatch->ComponentsToRender)
+			{
+				FIntRect ComponentRect = RenderBatch->ComputeSectionRect(Component, /*bInWithDuplicateBorders = */false);
+				BatchResult.ComponentRects.Add(ComponentRect);
+			}
+
+			InParams.MergeRenderContext->CycleBlendRenderTargets(RDGBuilderRecorder);
+			ULandscapeScratchRenderTarget* ReadRT = InParams.MergeRenderContext->GetBlendRenderTargetRead();
+
+			// Enqueue async readback direct from the current scratch read target.
+			ENQUEUE_RENDER_COMMAND(LandscapeLayers_Cmd_SyncExtractResults)(
+				[&BatchResult,
+				SourceResource = ReadRT->GetRenderTarget2D()->GetResource()]
+				(FRHICommandListImmediate& InRHICmdList)
+				{
+					FRDGBuilder GraphBuilder(InRHICmdList, RDG_EVENT_NAME("HeightmapsExtractResult"));
+
+					FRDGTextureRef SourceTextureRef = GraphBuilder.RegisterExternalTexture(CreateRenderTarget(SourceResource->TextureRHI, TEXT("SourceTexture")));
+					AddEnqueueCopyPass(GraphBuilder, &BatchResult.Readback, SourceTextureRef);
+
+					// We need to specify the final state of the external texture to prevent the graph builder from transitioning it to SRVMask :
+					GraphBuilder.SetTextureAccessFinal(SourceTextureRef, ERHIAccess::RTV);
+					GraphBuilder.Execute();
+				});
+		};
+
+	// Render everything now. Every time a group from a batch is done (there's only one group per batch for heightmaps), SyncExtractResults is called : 
+	MergeRenderContext.Render(EnqueueBatchReadback);
+
+	// Pick up readback results from SyncRenderBatchResults.  Enqueue this on the render thread because FRHIGPUTextureReadback can only be locked there.
+	ENQUEUE_RENDER_COMMAND(LandscapeLayers_Cmd_SyncFetchResults)(
+		[&SyncRenderBatchResults,
+		DestRect = InSelectiveRenderParams.Bounds,
+		Dest = InSelectiveRenderParams.CpuResult]
+		(FRHICommandListImmediate& InRHICmdList)
+		{
+			TRACE_CPUPROFILER_EVENT_SCOPE(LockAndFetchReadbackData);
+			const FBlitBuffer2DDesc DestDesc{ Dest.Num(), DestRect.Width(), DestRect };
+
+			for (FSyncRenderBatchResult& BatchResult : SyncRenderBatchResults)
+			{
+				int32 BatchStride = 0;
+				int32 BatchHeight = 0;
+				uint32* SrcData = static_cast<uint32*>(BatchResult.Readback.Lock(BatchStride, &BatchHeight));
+				int32 SrcBufferSize = BatchStride * BatchHeight; //the implied buffer size of the locked readback buffer, not necessarily of the native render target
+				check(BatchResult.Size.X <= BatchStride && BatchResult.Size.Y <= BatchHeight);
+
+				// The full batch rects are not all valid.  Only rects of each component in the batch are valid.  Blit one component at a time.
+				for (const FIntRect& ComponentBatchRect : BatchResult.ComponentRects)
+				{
+					// ComponentBatchRect is relative to the batch RT.  Convert back into overall landscape coords.
+					FIntRect ComponentRect = ComponentBatchRect + BatchResult.SectionRect.Min;
+					ComponentRect.Max += FIntPoint(1, 1);  //quads to verts
+					check(BatchResult.SectionRect.Contains(ComponentRect));
+					
+					const FBlitBuffer2DDesc SrcDesc{ SrcBufferSize, BatchStride, BatchResult.SectionRect };
+					int32 CopiedValues = BlitHeightChannelsToUint16(Dest.GetData(), DestDesc, SrcData, SrcDesc, ComponentRect);
+					ensure(CopiedValues > 0);  
+				}
+
+				BatchResult.Readback.Unlock();
+			}
+		});
+
+	// Sync the game thread with all of the above.  After this, the results are reliably accessible from the game thread.
+	FlushRenderingCommands();
+
+	return true;
+}
+
+bool ALandscape::PerformSelectiveRenderEditLayersHeightmapsRT(const FUpdateLayersContentContext& InUpdateLayersContentContext, const FLandscapeEditLayerRenderHeightParams& InSelectiveRenderParams)
+{
+	TRACE_CPUPROFILER_EVENT_SCOPE(PerformSelectiveRenderEditLayersHeightmapsRT);
+	RHI_BREADCRUMB_EVENT_GAMETHREAD("PerformSelectiveRenderEditLayersHeightmapsRT");
+
+	using namespace UE::Landscape;
+	using namespace UE::Landscape::Private;
+	using namespace UE::Landscape::EditLayers;
+	using namespace UE::Landscape::EditLayers::Private;
+
+	check(InSelectiveRenderParams.RTResult);
+
+	FMergeContext MergeContext(/*InLandscape = */this, /*bInIsHeightmapMerge = */true, /*bInSkipProceduralRenderers = */false);
+	TArray<FEditLayerRendererState> RendererStates = GetEditLayerRendererStatesEnableOverride(&MergeContext, InSelectiveRenderParams.ActiveEditLayers);
+
+	FMergeRenderParams MergeRenderParams(InUpdateLayersContentContext.LandscapeComponentsHeightmapsToRender, RendererStates);
+	FMergeRenderContext MergeRenderContext = PrepareEditLayersMergeRenderContext(MergeContext, MergeRenderParams);
+	if (!MergeRenderContext.IsValid())
+	{
+		// Nothing to do 
+		return false;
+	}
+
+	auto EnqueueCopyToResultRT = [&InUpdateLayersContentContext, &InSelectiveRenderParams, DestRT = InSelectiveRenderParams.RTResult]
+	(const FMergeRenderContext::FOnRenderBatchTargetGroupDoneParams& InParams, UE::Landscape::FRDGBuilderRecorder& RDGBuilderRecorder)
+		{
+			TRACE_CPUPROFILER_EVENT_SCOPE(ExtractSelectiveMergeBatchData);
+			RHI_BREADCRUMB_EVENT_GAMETHREAD("ExtractSelectiveMergeBatchData");
+
+			ALandscape* Landscape = InParams.MergeRenderContext->GetLandscape();
+			const FMergeRenderBatch* RenderBatch = InParams.MergeRenderContext->GetCurrentRenderBatch();
+
+			InParams.MergeRenderContext->CycleBlendRenderTargets(RDGBuilderRecorder);
+			ULandscapeScratchRenderTarget* SrcScratchRT = InParams.MergeRenderContext->GetBlendRenderTargetRead();
+			UTextureRenderTarget2D* ReadRT = SrcScratchRT->GetRenderTarget2D();
+			FIntRect DestRect = InSelectiveRenderParams.Bounds;
+
+			TArray<FRHICopyTextureInfo> ComponentCopyInfos;
+			ComponentCopyInfos.Reserve(RenderBatch->ComponentsToRender.Num());
+
+			for (ULandscapeComponent* Component : RenderBatch->ComponentsToRender)
+			{
+				FIntRect ComponentRect = RenderBatch->ComputeSectionRect(Component, /*bInWithDuplicateBorders = */false);  //component coords relative to the batch
+				ComponentRect.Max += FIntPoint(1, 1);  //quads to verts
+
+				FIntRect CopyRect = ComponentRect + RenderBatch->SectionRect.Min;  //coords in the landscape
+				CopyRect.Clip(RenderBatch->SectionRect);    //  limit to the intersection of the component, the batch, and the destination buffer
+				CopyRect.Clip(DestRect);
+				FIntRect SrcRect = CopyRect - RenderBatch->SectionRect.Min;  // relative to the batch
+
+				if (SrcRect.IsEmpty())
+				{
+					continue;
+				}
+
+				check(SrcRect.Min.X >= 0 && SrcRect.Min.Y >= 0);
+				check(SrcRect.Size().X <= RenderBatch->SectionRect.Size().X && SrcRect.Size().Y <= RenderBatch->SectionRect.Size().Y);
+				FIntPoint DestPosition = CopyRect.Min - DestRect.Min;
+
+				FRHICopyTextureInfo Info;
+				Info.Size = { SrcRect.Size().X, SrcRect.Size().Y, 1 };
+				Info.SourcePosition = { SrcRect.Min.X, SrcRect.Min.Y, 0 };
+				Info.DestPosition = { DestPosition.X, DestPosition.Y, 0 };
+
+				ComponentCopyInfos.Add(Info);
+			}
+
+			// Enqueue async texture copy into the destination render target
+			ENQUEUE_RENDER_COMMAND(LandscapeLayers_Cmd_ExtractSelectiveMergeBatchData)(
+				[ReadRT, DestRT, CopyInfos = MoveTemp(ComponentCopyInfos)]
+				(FRHICommandListImmediate& InRHICmdList)
+				{
+					FRDGBuilder GraphBuilder(InRHICmdList, RDG_EVENT_NAME("HeightmapsCopyComponentsToRT"));
+
+					FTextureRenderTargetResource* SrcResource = ReadRT->GetRenderTargetResource();
+					FTextureRenderTargetResource* DestResource = DestRT->GetRenderTargetResource();
+					FRDGTextureRef SourceTextureRef = GraphBuilder.RegisterExternalTexture(CreateRenderTarget(SrcResource->GetTextureRHI(), TEXT("CopySourceTexture")));
+					FRDGTextureRef DestTextureRef = GraphBuilder.RegisterExternalTexture(CreateRenderTarget(DestResource->GetTextureRHI(), TEXT("CopyDestTexture")));
+
+					for (const FRHICopyTextureInfo& CopyInfo : CopyInfos)
+					{
+						AddCopyTexturePass(GraphBuilder, SourceTextureRef, DestTextureRef, CopyInfo);
+					}
+
+					// Leave the scratch target in RTV, as it requires.
+					GraphBuilder.SetTextureAccessFinal(SourceTextureRef, ERHIAccess::RTV);
+					GraphBuilder.Execute();
+				});
+		};
+
+	// Render everything now. Every time a group from a batch is done (there's only one group per batch for heightmaps), the OnRenderBatchGroupDone callback is called : 
+	MergeRenderContext.Render(EnqueueCopyToResultRT);
+
+	// No flush.  Return without blocking.  The caller may need to flush if it wants to access the render target from the game thread.
+
+	return true;
+}
+
+
 FLandscapeBrushParameters ALandscape::BuildLandscapeBrushParameters(bool bInIsHeightmapMerge, const FIntRect& InRenderAreaExtents, UTextureRenderTarget2D* InRenderTarget, FName InWeightmapLayerName)
 {
 	FTransform RenderAreaWorldTransform = GetTransform();
@@ -11682,6 +11989,77 @@ bool ALandscape::CanUpdateLayersContent() const
 		Subsystem && Subsystem->GetTextureStreamingManager();
 }
 
+// Render the edit layers for an explicit set of components, for a subset of layers.  Do not modify the regular persistent state of edit layer data.
+// Empty InLayerVisibility implies using each layer's default visibility status.
+bool ALandscape::SelectiveRenderEditLayersHeightmaps(const FLandscapeEditLayerRenderHeightParams& InSelectiveRenderParams)
+{
+	check((InSelectiveRenderParams.CpuResult.Num() > 0) ^ (InSelectiveRenderParams.RTResult != nullptr));  //Expecting either a cpu result buffer or a render target.  Not both.
+
+	if (InSelectiveRenderParams.Bounds.IsEmpty())
+	{
+		return true;  //nothing to render
+	}
+	
+	// Expecting to be already initialized for BatchMerge.
+	check(CurrentEditLayersMergeMode == ELandscapeEditLayersMergeMode::BatchedMerge);
+	check(GetEditLayersMergeMode() == ELandscapeEditLayersMergeMode::BatchedMerge);
+	check(bLandscapeLayersAreInitialized);
+
+	if (!CanUpdateLayersContent())
+	{
+		return false;
+	}
+
+	UWorld* World = GetWorld();
+	check(World != nullptr);
+	ULandscapeSubsystem* LandscapeSubsystem = World->GetSubsystem<ULandscapeSubsystem>();
+	check(LandscapeSubsystem != nullptr);
+	ULandscapeInfo* LandscapeInfo = GetLandscapeInfo();
+	check(LandscapeInfo);
+
+	const bool bWaitForStreaming = true;
+	if (!PrepareTextureResources(bWaitForStreaming) ||
+		!PrepareLayersTextureResources(bWaitForStreaming) ||
+		!PrepareLayersResources(GetFeatureLevelShaderPlatform_Checked(World->GetFeatureLevel()), bWaitForStreaming))
+	{
+		return false;
+	}
+
+	TSet<ULandscapeComponent*> AffectedComponents;
+	LandscapeInfo->GetComponentsInRegion(InSelectiveRenderParams.Bounds.Min.X, InSelectiveRenderParams.Bounds.Min.Y, InSelectiveRenderParams.Bounds.Max.X, InSelectiveRenderParams.Bounds.Max.Y, AffectedComponents);
+	if (AffectedComponents.Num() == 0)
+	{
+		// No loaded components in region
+		UE_LOG(LogLandscape, Display, TEXT("SelectiveRenderLayers with no affected components"));
+		return true;
+	}
+
+	FTextureToComponentHelper Helper(*LandscapeInfo, AffectedComponents.Array(), FTextureToComponentHelper::ERefreshFlags::RefreshHeightmaps);
+	FUpdateLayersContentContext LayerContext(Helper, FUpdateLayersContentContext::FSelectiveRenderTag(), FUpdateLayersContentContext::ERefreshFlags::RefreshHeightmapInfos);
+
+	{
+		RenderCaptureInterface::FScopedCapture RenderCapture((RenderCaptureLayersNextHeightmapDraws != 0), TEXT("SelectiveRenderEditLayersHeightmaps"));
+		RenderCaptureLayersNextHeightmapDraws = FMath::Max(0, RenderCaptureLayersNextHeightmapDraws - 1);
+
+		if (!InSelectiveRenderParams.CpuResult.IsEmpty())
+		{
+			PerformSelectiveRenderEditLayersHeightmapsCPU(LayerContext, InSelectiveRenderParams);
+		}
+		else
+		{
+			PerformSelectiveRenderEditLayersHeightmapsRT(LayerContext, InSelectiveRenderParams);
+		}
+	}
+
+	return true;
+}
+
+bool ALandscape::SelectiveRenderEditLayersWeightmaps(const FLandscapeEditLayerRenderWeightParams& InSelectiveRenderParams)
+{
+	unimplemented();
+	return false;
+}
+
 void ALandscape::UpdateLayersContent(bool bInWaitForStreaming, bool bInSkipMonitorLandscapeEdModeChanges, bool bIntermediateRender, bool bFlushRender)
 {
 	TRACE_CPUPROFILER_EVENT_SCOPE(LandscapeLayers_UpdateLayersContent);
@@ -13867,6 +14245,232 @@ UTextureRenderTarget2D* FLandscapeLayerBrush::RenderLayer(const FLandscapeBrushP
 
 }
 
+void LandscapeDumpSelectiveLayerRender(const TArray<FString>& Args)
+{
+	// Example usage:  landscape.DumpSelectiveLayerRender -landscape Landscape -h -layers 255 -region 0 0 1000 500
+	// override default layer visibility with a bitmask:  -layers 255
+	// render a region of the landscape based on landscape coords (integers):  -region minX minY maxX maxY
+	// All params optional.
+
+	const TCHAR* LandscapeNameArg = TEXT("-landscape");
+	const TCHAR* LayerMaskArg = TEXT("-layermask");
+	const TCHAR* RegionArg = TEXT("-region");
+	const TCHAR* HeightmapArg = TEXT("-h");
+	const TCHAR* RTArg = TEXT("-rt");
+
+	ALandscape* Landscape = nullptr;
+	FString LandscapeName;
+	bool bHasLayerMask = false;
+	uint32 LayerMask = ~0;
+	bool bHasRegion = false;
+	FIntRect Region;
+	bool bUseRenderTarget = false;
+
+	bool bRenderHeight = true;  // Assuming heightmap render for now.  Weightmap support still to be added.
+
+	UWorld* World = GEditor->GetEditorWorldContext().World();
+	if (!World)
+	{
+		return;
+	}
+
+	for (int32 i = 0; i < Args.Num(); ++i)
+	{
+		const FString& Arg = Args[i];
+
+		if (Arg == LandscapeNameArg)
+		{
+			++i;
+			if (i < Args.Num())
+			{
+				LandscapeName = Args[i];
+			}
+		}
+		else if (Arg == LayerMaskArg)
+		{
+			++i;
+			if (i < Args.Num())
+			{
+				LexFromString(LayerMask, Args[i]);
+				bHasLayerMask = true;
+			}
+		}
+		else if (Arg == RegionArg)
+		{
+			if (i + 4 < Args.Num())
+			{
+				LexFromString(Region.Min.X, Args[i + 1]);
+				LexFromString(Region.Min.Y, Args[i + 2]);
+				LexFromString(Region.Max.X, Args[i + 3]);
+				LexFromString(Region.Max.Y, Args[i + 4]);
+				bHasRegion = true;
+			}
+			i += 4;
+		}
+		else if (Arg == HeightmapArg)
+		{
+			bRenderHeight = true;
+		}
+		else if (Arg == RTArg)
+		{
+			bUseRenderTarget = true;
+		}
+		else
+		{
+			UE_LOG(LogLandscape, Error, TEXT("LandscapeDumpSelectiveLayerRender: Unknown parameter %s"), *Arg);
+		}
+	}
+
+	// If landscape name not specified, just take the first one
+	for (ALandscape* Actor : TActorRange<ALandscape>(World))
+	{
+		if (LandscapeName.IsEmpty() || LandscapeName == Actor->GetActorNameOrLabel())
+		{
+			Landscape = Actor;
+			break;
+		}
+	}
+
+	if (!Landscape)
+	{
+		UE_LOG(LogLandscape, Error, TEXT("LandscapeDumpSelectiveLayerRender: landscape not found"));
+		return;
+	}
+
+	if (!bRenderHeight)
+	{
+		UE_LOG(LogLandscape, Error, TEXT("LandscapeDumpSelectiveLayerRender: must specify either -h or -w"));
+		return;
+	}
+
+	if (!bHasRegion)
+	{
+		if (ULandscapeInfo* LandscapeInfo = ULandscapeInfo::Find(World, Landscape->GetLandscapeGuid()))
+		{
+			bHasRegion = LandscapeInfo->GetLandscapeExtent(Region);
+		}
+	}
+
+	if (!bHasRegion)
+	{
+		UE_LOG(LogLandscape, Warning, TEXT("LandscapeDumpSelectiveLayerRender: couldn't determine landscape extent.  Nothing loaded?"));
+		return;
+	}
+
+	if (!Landscape->AllowSelectiveRenderEditLayers())
+	{
+		UE_LOG(LogLandscape, Warning, TEXT("LandscapeDumpSelectiveLayerRender: can't use Partial renders"));
+		return;
+	}
+
+	TBitArray<> LayerVis;  // leave empty for default layer visibility
+	if (bHasLayerMask)
+	{
+		int32 NumEditLayers = Landscape->GetEditLayersConst().Num();
+		for (int32 i = 0; i < NumEditLayers; ++i)
+		{
+			LayerVis.Add((LayerMask & (1 << i)) != 0);
+		}
+	}
+
+	ULandscapeSubsystem* Subsystem = World->GetSubsystem<ULandscapeSubsystem>();
+	const FDateTime CurrentTime = Subsystem->GetAppCurrentDateTime();
+	FString FilePattern = FString::Format(TEXT("{0}/Landscape/LandscapeRender{1}"), { FPaths::ProjectSavedDir(), CurrentTime.ToString() });
+
+	if (bRenderHeight)
+	{
+		UE_LOG(LogLandscape, Display, TEXT("LandscapeDumpSelectiveLayerRender: SelectiveRenderLayers heightmap, landscape=%s region=(%d,%d)-(%d,%d) layers=%x"),
+			*Landscape->GetActorNameOrLabel(), Region.Min.X, Region.Min.Y, Region.Max.X, Region.Max.Y, LayerMask);
+
+		TArray<uint16> HeightData;
+		HeightData.AddZeroed(Region.Area());
+		TObjectPtr<UTextureRenderTarget2D> HeightRT;
+
+		FLandscapeEditLayerRenderHeightParams RenderParams;
+		RenderParams.Bounds = Region;
+		RenderParams.ActiveEditLayers = LayerVis;
+
+		if (bUseRenderTarget)
+		{
+			// Prepare the render target
+			FName RTName = MakeUniqueObjectName(GetTransientPackage(), UTextureRenderTarget2D::StaticClass(), TEXT("LandscapeDumpHeightRT"));
+			HeightRT = NewObject<UTextureRenderTarget2D>(GetTransientPackage(), RTName, RF_Transient);
+			HeightRT->bCanCreateUAV = false;
+			HeightRT->RenderTargetFormat = RTF_RGBA8;
+			HeightRT->ClearColor = FColor(0);
+			FIntPoint Size = Region.Size();
+			HeightRT->InitAutoFormat(Size.X, Size.Y);
+			HeightRT->UpdateResourceImmediate(/*bClearRenderTarget = */true);
+
+			RenderParams.RTResult = HeightRT;
+		}
+		else
+		{
+			RenderParams.CpuResult = HeightData;
+		}
+
+		bool bRenderRet = Landscape->SelectiveRenderEditLayersHeightmaps(RenderParams);
+		if (!bRenderRet)
+		{
+			UE_LOG(LogLandscape, Error, TEXT("LandscapeDumpSelectiveLayerRender: SelectiveRenderEditLayersHeightmaps failed"));
+			return;
+		}
+
+		if (bUseRenderTarget)
+		{
+			TRACE_CPUPROFILER_EVENT_SCOPE(BatchedMergeHeightmapsSyncResults_ReadPixels);
+			FlushRenderingCommands();
+
+			// Note:  GameThread_GetRenderTargetResource is documented as returning a pointer which cannot be used from the game thread,
+			// but other callers say this can be used if you don't care about performance and flush the render thread first.
+			FTextureRenderTargetResource* RenderTargetResource = HeightRT->GameThread_GetRenderTargetResource();
+
+			// Make a blocking read from the render target.
+			TArray<FColor> RGBABuffer;
+			FIntRect ReadRect = Region;  // destination rect in landscape coords
+			RGBABuffer.Reserve(Region.Area());
+			FIntPoint Offset = ReadRect.Min;
+			ReadRect -= Offset;  //offset so that the first pixel of the buffer is (0,0).
+			FReadSurfaceDataFlags ReadFlags(RCM_MinMax);
+			bool bRet = RenderTargetResource->ReadPixels(RGBABuffer, ReadFlags, ReadRect);
+
+			// Convert FColor (R+G) to uint16
+			TArrayView<uint16> DestBuffer = HeightData;
+			ensure(bRet);
+			ensure(RGBABuffer.Num() == DestBuffer.Num());
+			if (bRet)
+			{
+				for (int32 i = 0; i < RGBABuffer.Num(); ++i)
+				{
+					FColor HeightAndNormal = RGBABuffer[i];
+					DestBuffer[i] = (static_cast<uint16>(HeightAndNormal.R) << 8) | static_cast<uint16>(HeightAndNormal.G);
+				}
+			}
+			else
+			{
+				UE_LOG(LogLandscape, Error, TEXT("LandscapeDumpSelectiveLayerRender: ReadPixels failed"));
+				return;
+			}
+
+			HeightRT->ReleaseResource();
+		}
+
+		// Write the data to an image file
+		FImageInfo ImageInfo(Region.Size().X, Region.Size().Y, 1, ERawImageFormat::G16, EGammaSpace::Linear);
+		FImageView ImageView(ImageInfo, HeightData.GetData());
+		FString Filename = FilePattern + "Heightmap.png";
+		if (FImageUtils::SaveImageByExtension(*Filename, ImageView))
+		{
+			UE_LOG(LogLandscape, Display, TEXT("LandscapeDumpSelectiveLayerRender: wrote heightmap data to %s"), *Filename);
+		}
+		else
+		{
+			UE_LOG(LogLandscape, Error, TEXT("LandscapeDumpSelectiveLayerRender: failed to write heightmap data to %s"), *Filename);
+		}
+	}
+}
+
+
 #endif // WITH_EDITOR
 
 
diff --git a/Engine/Source/Runtime/Landscape/Private/LandscapeUtilsPrivate.cpp b/Engine/Source/Runtime/Landscape/Private/LandscapeUtilsPrivate.cpp
--- a/Engine/Source/Runtime/Landscape/Private/LandscapeUtilsPrivate.cpp
+++ b/Engine/Source/Runtime/Landscape/Private/LandscapeUtilsPrivate.cpp
@@ -91,4 +91,48 @@ FLandscapeComponent2DIndexer CreateLandscapeComponent2DIndexer(const ULandscapeI
 	return FLandscapeComponent2DIndexer(AllValidComponents);
 }
 
+int32 BlitHeightChannelsToUint16 (uint16* DestBuffer, const FBlitBuffer2DDesc& Dest, uint32* SrcBuffer, const FBlitBuffer2DDesc& Src, FIntRect ClipRect)
+{
+	check(Dest.Rect.Height() * Dest.Stride <= Dest.BufferSize);
+	check(Src.Rect.Height() * Src.Stride <= Src.BufferSize);
+	FIntRect BlitRect = Dest.Rect;
+	BlitRect.Clip(Src.Rect);
+	BlitRect.Clip(ClipRect);
+	if (BlitRect.IsEmpty())
+	{
+		return 0;
+	}
+
+	FIntPoint ReadOffset = BlitRect.Min - Src.Rect.Min;
+	FIntPoint WriteOffset = BlitRect.Min - Dest.Rect.Min;
+	FIntPoint BlitSize = BlitRect.Size();
+
+	for (int32 Y = BlitRect.Min.Y; Y < BlitRect.Max.Y; ++Y)
+	{
+		// Set up the pointers that describe the start and end of each scanline to copy.
+		int32 SrcY = Y - Src.Rect.Min.Y;
+		uint32* Read = SrcBuffer + (SrcY * Src.Stride) + ReadOffset.X;
+		uint32* ReadEnd = Read + BlitSize.X;
+		int32 DestY = Y - Dest.Rect.Min.Y;
+		uint16* Write = DestBuffer + (DestY * Dest.Stride) + WriteOffset.X;
+		uint16* WriteEnd = Write + BlitSize.X;
+
+		check(Read >= SrcBuffer);
+		check(ReadEnd <= SrcBuffer + (Src.Stride * Src.BufferSize));
+		check(Read <= ReadEnd);
+		check(Write >= DestBuffer);
+		check(WriteEnd <= DestBuffer + Dest.BufferSize);
+		check(Write <= WriteEnd);
+
+		for (; Read < ReadEnd && Write < WriteEnd; ++Read, ++Write)
+		{
+			FColor HeightAndNormal(*Read);
+			uint16 HeightValue = (static_cast<uint16>(HeightAndNormal.R) << 8) | static_cast<uint16>(HeightAndNormal.G);
+			*Write = HeightValue;
+		}
+		check(Read == ReadEnd && Write == WriteEnd);
+	}
+	return BlitRect.Area();
+};
+
 } // end namespace UE::Landscape::Private
\ No newline at end of file
diff --git a/Engine/Source/Runtime/Landscape/Private/LandscapeUtilsPrivate.h b/Engine/Source/Runtime/Landscape/Private/LandscapeUtilsPrivate.h
--- a/Engine/Source/Runtime/Landscape/Private/LandscapeUtilsPrivate.h
+++ b/Engine/Source/Runtime/Landscape/Private/LandscapeUtilsPrivate.h
@@ -306,4 +306,26 @@ namespace UE::Landscape::Private
 
 	using FLandscapeComponent2DIndexer = T2DIndexer<ULandscapeComponent*, FLandscapeComponent2DIndexerKeyFuncs>;
 	FLandscapeComponent2DIndexer CreateLandscapeComponent2DIndexer(const ULandscapeInfo* InInfo);
-} // end namespace UE::Landscape
+
+
+	struct FBlitBuffer2DDesc
+	{
+		int32 BufferSize = 0;  // Buffer size in elements.
+		int32 Stride = 0;      // Buffer stride in elements.  i.e. the distance between the start of each row.
+		FIntRect Rect;         // The 2D region represented by this buffer.  Rect.Area() should be <= BufferSize.
+	};
+
+	/* Copy height values from the R+G channels of the source buffer into a plain uint16 dest buffer.  The source and destination buffers
+	 * represent sub-regions of a shared 2D coordinate system expressed in the Rect parameters.  ClipRect is the desired area of that
+	 * shared coordinate space to limit the operation to.  Copies only values in the the intersection of Src.Rect, Dest.Rect, and ClipRect.
+	 * All FIntRect params are treated as half-open.  .Min is included.  .Max is excluded.
+	 * @param DestBuffer destination buffer
+	 * @param Dest description of the destination buffer
+	 * @param SrcBuffer source buffer
+	 * @param Src description of the source buffer
+	 * @param ClipRect The operation will be limited to this region, in the same shared coordinate space of the rects from the buffer descriptions.
+	 * @return number of elements copied.
+	 */
+	int32 BlitHeightChannelsToUint16(uint16* DestBuffer, const FBlitBuffer2DDesc& Dest, uint32* SrcBuffer, const FBlitBuffer2DDesc& Src, FIntRect ClipRect);
+
+} // end namespace UE::Landscape::Private
diff --git a/Engine/Source/Runtime/Landscape/Public/LandscapeEditLayerMergeRenderContext.h b/Engine/Source/Runtime/Landscape/Public/LandscapeEditLayerMergeRenderContext.h
--- a/Engine/Source/Runtime/Landscape/Public/LandscapeEditLayerMergeRenderContext.h
+++ b/Engine/Source/Runtime/Landscape/Public/LandscapeEditLayerMergeRenderContext.h
@@ -361,6 +361,8 @@ public:
 	inline const TArray<FMergeRenderBatch>& GetRenderBatches() const { return RenderBatches; }
 	LANDSCAPE_EDIT_LAYERS_BATCHED_MERGE_EXPERIMENTAL
 	LANDSCAPE_API const FMergeRenderBatch* GetCurrentRenderBatch() const;
+	LANDSCAPE_EDIT_LAYERS_BATCHED_MERGE_EXPERIMENTAL
+	LANDSCAPE_API const int32 GetCurrentRenderBatchIdx() const;
 
 	LANDSCAPE_EDIT_LAYERS_BATCHED_MERGE_EXPERIMENTAL
 	LANDSCAPE_API FTransform ComputeVisualLogTransform(const FTransform& InTransform) const;
