Landscape groups handle changing transforms

If the landscape component is out of alignment with the grid when it is mapped, it first tries to fix the transform.  This happens when the transform on the ALandscape has been changed and you load a landscape streaming proxy from disk.

When an already-mapped landscape component is transformed, we add it to a list, and remap it during the next tick.  It's deferred so we can change the grid for all components at once.

diff --git a/Engine/Source/Runtime/Landscape/Classes/LandscapeComponent.h b/Engine/Source/Runtime/Landscape/Classes/LandscapeComponent.h
--- a/Engine/Source/Runtime/Landscape/Classes/LandscapeComponent.h
+++ b/Engine/Source/Runtime/Landscape/Classes/LandscapeComponent.h
@@ -775,6 +775,7 @@ public:
 	virtual ELightMapInteractionType GetStaticLightingType() const override;
 	virtual void GetStreamingRenderAssetInfo(FStreamingTextureLevelContext& LevelContext, TArray<FStreamingRenderAssetPrimitiveInfo>& OutStreamingRenderAssets) const override;
 	virtual bool IsPrecomputedLightingValid() const override;
+	virtual void OnUpdateTransform(EUpdateTransformFlags UpdateTransformFlags, ETeleportType TeleportType) override;
 
 	virtual TArray<URuntimeVirtualTexture*> const& GetRuntimeVirtualTextures() const override;
 	virtual ERuntimeVirtualTextureMainPassType GetVirtualTextureRenderPassType() const override;
diff --git a/Engine/Source/Runtime/Landscape/Private/Landscape.cpp b/Engine/Source/Runtime/Landscape/Private/Landscape.cpp
--- a/Engine/Source/Runtime/Landscape/Private/Landscape.cpp
+++ b/Engine/Source/Runtime/Landscape/Private/Landscape.cpp
@@ -2014,6 +2014,15 @@ const FMeshMapBuildData* ULandscapeComponent::GetMeshMapBuildData() const
 	return NULL;
 }
 
+void ULandscapeComponent::OnUpdateTransform(EUpdateTransformFlags UpdateTransformFlags, ETeleportType TeleportType)
+{
+	// if registered with a group, we need to update our registration to reflect the new transform
+	if (RegisteredLandscapeGroup != nullptr)
+	{
+		RegisteredLandscapeGroup->OnTransformUpdated(this);
+	}
+}
+
 bool ULandscapeComponent::IsPrecomputedLightingValid() const
 {
 	return GetMeshMapBuildData() != NULL;
diff --git a/Engine/Source/Runtime/Landscape/Private/LandscapeGroup.cpp b/Engine/Source/Runtime/Landscape/Private/LandscapeGroup.cpp
--- a/Engine/Source/Runtime/Landscape/Private/LandscapeGroup.cpp
+++ b/Engine/Source/Runtime/Landscape/Private/LandscapeGroup.cpp
@@ -2,6 +2,7 @@
 
 #include "LandscapeGroup.h"
 #include "CoreMinimal.h"
+#include "Engine/Engine.h"
 #include "Engine/World.h"
 #include "Landscape.h"
 #include "LandscapeStreamingProxy.h"
@@ -283,6 +284,24 @@ void FLandscapeGroup::UnregisterComponent(ULandscapeComponent* Component)
 	}
 }
 
+void FLandscapeGroup::OnTransformUpdated(ULandscapeComponent* Component)
+{
+	ULandscapeHeightmapTextureEdgeFixup* EdgeFixup = Component->RegisteredEdgeFixup;
+	check(EdgeFixup != nullptr);
+	check(Component->RegisteredLandscapeGroup == this);
+
+	if (!EdgeFixup->IsComponentActive(Component))
+	{
+		// a disabled component is not mapped, so we don't need to fix anything
+	}
+	else
+	{
+		// flag it as moved, so we can fix it on the next tick
+		// (this must be deferred so we can move all components simultaneously, and redefine the grid if necessary)
+		HeightmapsMoved.Add(EdgeFixup);
+	}
+}
+
 void FLandscapeGroup::DisableAndUnmap(ULandscapeHeightmapTextureEdgeFixup* Fixup)
 {
 	// should be registered before it can be disabled
@@ -295,20 +314,107 @@ void FLandscapeGroup::DisableAndUnmap(ULandscapeHeightmapTextureEdgeFixup* Fixup
 	Unmap(Fixup);
 }
 
+#if WITH_EDITOR
+namespace UE::Landscape::Private
+{
+	bool FixupProxyTransform(ALandscapeProxy* Proxy)
+	{
+		// make sure we've run shared property sync before moving proxies
+		if (!Proxy->bIsRegisteredWithLandscapeInfo)
+		{
+			return false;
+		}
+
+		ALandscape* Landscape = Proxy->GetLandscapeActor();
+		check(Landscape);
+
+		// Make sure section offset of all proxies is multiple of ALandscapeProxy::ComponentSizeQuads
+		{
+			FIntPoint LandscapeSectionOffset = Proxy->LandscapeSectionOffset - Landscape->LandscapeSectionOffset;
+			FIntPoint LandscapeSectionOffsetRem(
+				LandscapeSectionOffset.X % Proxy->ComponentSizeQuads,
+				LandscapeSectionOffset.Y % Proxy->ComponentSizeQuads);
+			if (LandscapeSectionOffsetRem.X != 0 || LandscapeSectionOffsetRem.Y != 0)
+			{
+				FIntPoint NewLandscapeSectionOffset = Proxy->LandscapeSectionOffset - LandscapeSectionOffsetRem;
+
+				UE_LOG(LogLandscape, Warning, TEXT("Landscape section base is not multiple of component size, attempted automated fix: '%s', %d,%d vs %d,%d."),
+					*Proxy->GetFullName(), Proxy->LandscapeSectionOffset.X, Proxy->LandscapeSectionOffset.Y, NewLandscapeSectionOffset.X, NewLandscapeSectionOffset.Y);
+
+				Proxy->SetAbsoluteSectionBase(NewLandscapeSectionOffset);
+			}
+		}
+
+		// Update transformations of all linked landscape proxies
+		FTransform LandscapeTM = Landscape->LandscapeActorToWorld();
+		{
+			FTransform ProxyRelativeTM(FVector(Proxy->LandscapeSectionOffset));
+			FTransform ProxyTransform = ProxyRelativeTM * LandscapeTM;
+
+			if (!Proxy->GetTransform().Equals(ProxyTransform))
+			{
+				Proxy->SetActorTransform(ProxyTransform);
+
+				// Let other systems know that an actor was moved
+				GEngine->BroadcastOnActorMoved(Proxy);
+			}
+		}
+
+		return true;
+	}
+}
+#endif // WITH_EDITOR
+
 FIntPoint FLandscapeGroup::Map(ULandscapeHeightmapTextureEdgeFixup* Fixup, ULandscapeComponent* Component)
 {
 	check(Fixup->bMapped == false);
 
-	const FTransform& ComponentLocalToWorld = Component->GetComponentTransform();
+	ALandscapeProxy* Proxy = Component->GetLandscapeProxy();
+
+	FTransform ComponentLocalToWorld = Component->GetComponentTransform();
 	FBoxSphereBounds ComponentLocalBounds = Component->CalcBounds(FTransform::Identity);
 	FVector ComponentCenterWorldSpace = ComponentLocalToWorld.TransformPosition(ComponentLocalBounds.Origin);
 	FVector ComponentXVectorWorldSpace = ComponentLocalToWorld.TransformVector(FVector::XAxisVector);
 	FVector ComponentYVectorWorldSpace = ComponentLocalToWorld.TransformVector(FVector::YAxisVector);
-	FVector ComponentLandscapeGridScale = Component->GetLandscapeProxy()->GetRootComponent()->GetRelativeScale3D();
+	FVector ComponentLandscapeGridScale = Proxy->GetRootComponent()->GetRelativeScale3D();
+
+	// in general, the first registered section gets to set up our RenderCoord grid so it is located at the origin
+	// however, there is also the case where we change the landscape transform, which should cause the entire grid to be remapped
+	// it's a bit tricky to handle properly though, as all proxies are moved, then each proxy's components are mapped and unmapped.
+	// meaning there's no time when there is nothing in the fixup map.  Instead, we check if all remaining items in the fixup map are moved.
+	// in this case we can assume it's ok to redefine the grid.
+	bool bCanRedefineGrid = false;
+	if (HeightmapsMoved.Num() >= XYToEdgeFixupMap.Num())
+	{
+		bCanRedefineGrid = true;
+		for (auto MappedFixupPair : XYToEdgeFixupMap)
+		{
+			ULandscapeHeightmapTextureEdgeFixup* MappedFixup = MappedFixupPair.Value;
+			if (!HeightmapsMoved.Contains(MappedFixup))
+			{
+				// there is still a mapped fixup that wasn't moved, don't update the coords
+				bCanRedefineGrid = false;
+				break;
+			}
+		}
 
-	if (XYToEdgeFixupMap.Num() == 0)
+	}
+
+	if (bCanRedefineGrid)
 	{
-		// the first registered section gets to set up our RenderCoord grid so it is located at the origin
+		if (XYToEdgeFixupMap.Num() > 0)
+		{
+			// unmap all existing fixups prior to changing the rendercoord grid
+			// (we can only get here if all of these fixups are "moved", and will be remapped anyways)
+			TArray<TObjectPtr<ULandscapeHeightmapTextureEdgeFixup>> Fixups;
+			XYToEdgeFixupMap.GenerateValueArray(Fixups);
+			for (ULandscapeHeightmapTextureEdgeFixup* ExistingFixup : Fixups)
+			{
+				const bool bRemoveFromLists = false;
+				Unmap(ExistingFixup, bRemoveFromLists);
+			}
+		}
+
 		ComponentResolution = Component->ComponentSizeQuads;
 		GroupCoordOrigin = ComponentCenterWorldSpace;
 		GroupCoordXVector = ComponentXVectorWorldSpace;
@@ -323,11 +429,32 @@ FIntPoint FLandscapeGroup::Map(ULandscapeHeightmapTextureEdgeFixup* Fixup, ULand
 		bool bYVectorMatches = (ComponentYVectorWorldSpace - GroupCoordYVector).IsNearlyZero();
 		if (!(bResolutionMatches && bXVectorMatches && bYVectorMatches))
 		{
-			UE_LOG(LogLandscape, Warning, TEXT("Landscapes in LOD Group with Key %d do not have matching resolution (%d == %d), scale (%f == %f, %f == %f) and/or rotation; seam artifacts may appear."),
-				LandscapeGroupKey,
-				ComponentResolution, Component->ComponentSizeQuads,
-				GroupCoordXVector.Length(), ComponentXVectorWorldSpace.Length(),
-				GroupCoordYVector.Length(), ComponentYVectorWorldSpace.Length());
+#if WITH_EDITOR
+			// mismatch: first check if this component just needs its transform updated (it may have been recently loaded and not yet fixed)
+			if (UE::Landscape::Private::FixupProxyTransform(Proxy))
+			{
+				// fixup success: grab the new transform and check again
+				ComponentLocalToWorld = Component->GetComponentTransform();
+				ComponentLocalBounds = Component->CalcBounds(FTransform::Identity);
+				ComponentCenterWorldSpace = ComponentLocalToWorld.TransformPosition(ComponentLocalBounds.Origin);
+				ComponentXVectorWorldSpace = ComponentLocalToWorld.TransformVector(FVector::XAxisVector);
+				ComponentYVectorWorldSpace = ComponentLocalToWorld.TransformVector(FVector::YAxisVector);
+				ComponentLandscapeGridScale = Proxy->GetRootComponent()->GetRelativeScale3D();
+				
+				bResolutionMatches = (ComponentResolution == Component->ComponentSizeQuads);
+				bXVectorMatches = (ComponentXVectorWorldSpace - GroupCoordXVector).IsNearlyZero();
+				bYVectorMatches = (ComponentYVectorWorldSpace - GroupCoordYVector).IsNearlyZero();
+			}
+#endif // WITH_EDITOR
+
+			if (!(bResolutionMatches && bXVectorMatches && bYVectorMatches))
+			{
+				UE_LOG(LogLandscape, Warning, TEXT("Landscapes in LOD Group with Key %d do not have matching resolution (%d == %d), scale (%f == %f, %f == %f) and/or rotation; seam artifacts may appear."),
+					LandscapeGroupKey,
+					ComponentResolution, Component->ComponentSizeQuads,
+					GroupCoordXVector.Length(), ComponentXVectorWorldSpace.Length(),
+					GroupCoordYVector.Length(), ComponentYVectorWorldSpace.Length());
+			}
 		}
 
 		if (!(LandscapeGridScale - ComponentLandscapeGridScale).IsNearlyZero())
@@ -384,9 +511,12 @@ FIntPoint FLandscapeGroup::Map(ULandscapeHeightmapTextureEdgeFixup* Fixup, ULand
 	return GroupCoord;
 }
 
-void FLandscapeGroup::Unmap(ULandscapeHeightmapTextureEdgeFixup* Fixup)
+void FLandscapeGroup::Unmap(ULandscapeHeightmapTextureEdgeFixup* Fixup, const bool bRemoveFromLists)
 {
-	check(Fixup->bMapped == true);
+	if (!Fixup->bMapped)
+	{
+		return;
+	}
 	check(Fixup->ActiveGroup == this);
 
 	FIntPoint GroupCoord = Fixup->GroupCoord;
@@ -394,10 +524,14 @@ void FLandscapeGroup::Unmap(ULandscapeHeightmapTextureEdgeFixup* Fixup)
 	XYToEdgeFixupMap.RemoveAndCopyValue(GroupCoord, RemovedEdgeFixup);
 	check(Fixup == RemovedEdgeFixup);
 
+	if (bRemoveFromLists)
+	{
 #if WITH_EDITOR
-	HeightmapsNeedingEdgeSnapshotCapture.Remove(Fixup);
+		HeightmapsNeedingEdgeSnapshotCapture.Remove(Fixup);
 #endif // WITH_EDITOR
-	HeightmapsNeedingEdgeTexturePatching.Remove(Fixup);
+		HeightmapsNeedingEdgeTexturePatching.Remove(Fixup);
+		HeightmapsMoved.Remove(Fixup);
+	}
 
 	Fixup->bMapped = false;
 }
@@ -507,6 +641,37 @@ void FLandscapeGroup::TickEdgeFixup(ULandscapeSubsystem* LandscapeSubsystem, boo
 		return;
 	}
 
+	if (HeightmapsMoved.Num() > 0)
+	{
+		// unmap the moved components
+		for (ULandscapeHeightmapTextureEdgeFixup* EdgeFixup : HeightmapsMoved)
+		{
+			check(EdgeFixup);
+			const bool bRemoveFromLists = false; // everything will be remapped below, no need to remove
+			Unmap(EdgeFixup, bRemoveFromLists);
+		}
+
+		// remap the moved components
+		for (ULandscapeHeightmapTextureEdgeFixup* EdgeFixup : HeightmapsMoved)
+		{
+			check(EdgeFixup);
+			Map(EdgeFixup, EdgeFixup->ActiveComponent);
+
+			// Mark it, and all neighbors, as needing edge texture patching.
+			HeightmapsNeedingEdgeTexturePatching.Add(EdgeFixup);
+			for (UE::Landscape::ENeighborIndex NeighborIndex : TEnumRange<UE::Landscape::ENeighborIndex>())
+			{
+				FIntPoint NeighborCoord = EdgeFixup->GroupCoord + GetNeighborRelativePosition(NeighborIndex);
+				if (TObjectPtr<ULandscapeHeightmapTextureEdgeFixup>* NeighborPtrPtr = XYToEdgeFixupMap.Find(NeighborCoord))
+				{
+					HeightmapsNeedingEdgeTexturePatching.Add(NeighborPtrPtr->Get());
+				}
+			}
+		}
+
+		HeightmapsMoved.Empty();
+	}
+
 #if !UE_BUILD_SHIPPING
 	if (GAmortizedGroupValidation)
 	{
diff --git a/Engine/Source/Runtime/Landscape/Private/LandscapeGroup.h b/Engine/Source/Runtime/Landscape/Private/LandscapeGroup.h
--- a/Engine/Source/Runtime/Landscape/Private/LandscapeGroup.h
+++ b/Engine/Source/Runtime/Landscape/Private/LandscapeGroup.h
@@ -46,6 +46,9 @@ struct FLandscapeGroup
 	// this is only used to double check correct behavior
 	TSet<TObjectPtr<ULandscapeHeightmapTextureEdgeFixup>> AllRegisteredFixups;
 
+	// set of fixups that were moved and may need to modify their mapped position (or if all are moved, we may need to reset the map grid)
+	TSet<TObjectPtr<ULandscapeHeightmapTextureEdgeFixup>> HeightmapsMoved;
+
 #if WITH_EDITOR
 	// fixups that need to capture new edge snapshots
 	TSet<TObjectPtr<ULandscapeHeightmapTextureEdgeFixup>> HeightmapsNeedingEdgeSnapshotCapture;
@@ -66,6 +69,7 @@ struct FLandscapeGroup
 
 	void RegisterComponent(ULandscapeComponent* Component);
 	void UnregisterComponent(ULandscapeComponent* Component);
+	void OnTransformUpdated(ULandscapeComponent* Component);
 
 	static void RegisterAllComponentsOnStreamingProxy(ALandscapeStreamingProxy* StreamingProxy);
 	static void UnregisterAllComponentsOnStreamingProxy(ALandscapeStreamingProxy* StreamingProxy);
@@ -92,5 +96,5 @@ struct FLandscapeGroup
 private:
 	void DisableAndUnmap(ULandscapeHeightmapTextureEdgeFixup* Fixup);
 	FIntPoint Map(ULandscapeHeightmapTextureEdgeFixup* Fixup, ULandscapeComponent* Component);
-	void Unmap(ULandscapeHeightmapTextureEdgeFixup* Fixup);
+	void Unmap(ULandscapeHeightmapTextureEdgeFixup* Fixup, const bool bRemoveFromLists = true);
 };
