PR #13415: Landscape Resize - Fixed incorrect Y offset calculation when expanding landscapes

diff --git a/Engine/Source/Runtime/Landscape/Public/LandscapeConfigHelper.h b/Engine/Source/Runtime/Landscape/Public/LandscapeConfigHelper.h
--- a/Engine/Source/Runtime/Landscape/Public/LandscapeConfigHelper.h
+++ b/Engine/Source/Runtime/Landscape/Public/LandscapeConfigHelper.h
@@ -101,7 +101,7 @@ public:
 		const int32 DstWidth = InDestRegion.Width() + 1;
 		const int32 DstHeight = InDestRegion.Height() + 1;
 		const int32 OffsetX = bOffset ? InDestRegion.Min.X - InSrcRegion.Min.X : 0;
-		const int32 OffsetY = bOffset ? InDestRegion.Min.Y - InSrcRegion.Min.X : 0;
+		const int32 OffsetY = bOffset ? InDestRegion.Min.Y - InSrcRegion.Min.Y : 0;
 		
 		OutData.Empty(DstWidth * DstHeight);
 		OutData.AddUninitialized(DstWidth * DstHeight);
