Disable landscape.BatchedMerge.SelectiveRenderEditLayers while investigating an issue with it.

diff --git a/Engine/Source/Runtime/Landscape/Private/LandscapeEditLayers.cpp b/Engine/Source/Runtime/Landscape/Private/LandscapeEditLayers.cpp
--- a/Engine/Source/Runtime/Landscape/Private/LandscapeEditLayers.cpp
+++ b/Engine/Source/Runtime/Landscape/Private/LandscapeEditLayers.cpp
@@ -272,7 +272,7 @@ TAutoConsoleVariable<bool> CVarSilenceMergeBatchResolutionWarning(
 // TODO:  remove with IntermediateRender.
 TAutoConsoleVariable<bool> CVarLandscapeSelectiveRenderEditLayers(
 	TEXT("landscape.BatchedMerge.SelectiveRenderEditLayers"),
-	true,
+	false,
 	TEXT("When true, use SelectiveRenderEditLayers codepath instead of intermediate renders"));
 
 
