Fix world partition commandlet build interacting with landscape texture controls.  Be prepared to run in a commandlet, which doesn't increment GFrameNumber.  When filtering out redundant calls to PrepareTextureResources, remember the last real result and return that.

diff --git a/Engine/Source/Runtime/Landscape/Classes/Landscape.h b/Engine/Source/Runtime/Landscape/Classes/Landscape.h
--- a/Engine/Source/Runtime/Landscape/Classes/Landscape.h
+++ b/Engine/Source/Runtime/Landscape/Classes/Landscape.h
@@ -827,6 +827,7 @@ private:
 
 	uint32 LastFlushedLayerUpdateFrame = 0;
 	uint32 LastPrepareTextureResourcesCalled = 0;
+	bool bLastPrepareTextureResourcesResult = false;
 	bool bHasWarnedAboutInvalidShadingModel = false;
 
 	// Counter to detect re-entrance
diff --git a/Engine/Source/Runtime/Landscape/Private/LandscapeEditLayers.cpp b/Engine/Source/Runtime/Landscape/Private/LandscapeEditLayers.cpp
--- a/Engine/Source/Runtime/Landscape/Private/LandscapeEditLayers.cpp
+++ b/Engine/Source/Runtime/Landscape/Private/LandscapeEditLayers.cpp
@@ -4156,13 +4156,13 @@ void ALandscape::PrintLayersDebugTextureResource(const FString& InContext, FText
 
 bool ALandscape::PrepareTextureResourcesLimited(bool bInWaitForStreaming)
 {
-	// Rate-limit to once per frame
-	if (GFrameNumber != LastPrepareTextureResourcesCalled)
+	// Rate-limit to once per frame.  Note that commandlets don't advance the frame counter.
+	if (GFrameNumber != LastPrepareTextureResourcesCalled || IsRunningCommandlet())
 	{
 		LastPrepareTextureResourcesCalled = GFrameNumber;
-		return PrepareTextureResources(bInWaitForStreaming);
+		bLastPrepareTextureResourcesResult = PrepareTextureResources(bInWaitForStreaming);
 	}
-	return true;
+	return bLastPrepareTextureResourcesResult;
 }
 
 bool ALandscape::PrepareTextureResources(bool bInWaitForStreaming)
