Landscape Editor - Updating FWeightmapLayerAllocationInfo validation to fix crash in UpdateCollisionLayerData
* Fixes crash when clearing a target layer that is already empty

diff --git a/Engine/Source/Runtime/Landscape/Private/LandscapeEdit.cpp b/Engine/Source/Runtime/Landscape/Private/LandscapeEdit.cpp
--- a/Engine/Source/Runtime/Landscape/Private/LandscapeEdit.cpp
+++ b/Engine/Source/Runtime/Landscape/Private/LandscapeEdit.cpp
@@ -2018,7 +2018,7 @@ void ULandscapeComponent::UpdateCollisionLayerData(const FColor* const* const We
 		{
 			const FWeightmapLayerAllocationInfo& AllocInfo = ComponentWeightmapLayerAllocations[AllocIdx];
 			ULandscapeLayerInfoObject* LayerInfo = AllocInfo.LayerInfo;
-			if (LayerInfo == ALandscapeProxy::VisibilityLayer || LayerInfo != nullptr)
+			if (LayerInfo == ALandscapeProxy::VisibilityLayer || AllocInfo.IsAllocated())
 			{
 				int32 Idx = CandidateLayers.Add(LayerInfo);
 				CandidateDataPtrs.Add(((uint8*)WeightmapTextureMipData[AllocInfo.WeightmapTextureIndex]) + ChannelOffsets[AllocInfo.WeightmapTextureChannel]);
@@ -5178,7 +5178,8 @@ void ULandscapeInfo::DeleteLayer(ULandscapeLayerInfoObject* LayerInfo, const FNa
 
 void ULandscapeInfo::ReplaceLayer(ULandscapeLayerInfoObject* FromLayerInfo, ULandscapeLayerInfoObject* ToLayerInfo)
 {
-	if (ensure(FromLayerInfo != ToLayerInfo))
+	// From and To may both be null if user tries to clear an already empty layer
+	if (FromLayerInfo != ToLayerInfo)
 	{
 		GWarn->BeginSlowTask(LOCTEXT("BeginReplacingLayerTask", "Replacing Layer"), true);
 
