[Landscape] Hook to customize ULandscapeComponent scene proxy

Landscape module now exposes an API to set an optional delegate factory to create a primitive scene proxy for landscape components.

This allows to customize landscape components rendering in plugins.

#rnx

diff --git a/Engine/Source/Runtime/Landscape/Private/Landscape.cpp b/Engine/Source/Runtime/Landscape/Private/Landscape.cpp
--- a/Engine/Source/Runtime/Landscape/Private/Landscape.cpp
+++ b/Engine/Source/Runtime/Landscape/Private/Landscape.cpp
@@ -31,6 +31,7 @@ Landscape.cpp: Terrain rendering
 #include "LandscapeMeshProxyComponent.h"
 #include "LandscapeNaniteComponent.h"
 #include "LandscapeRender.h"
+#include "LandscapeModule.h"
 #include "LandscapePrivate.h"
 #include "Logging/StructuredLog.h"
 #include "Logging/TokenizedMessage.h"
@@ -104,6 +105,7 @@ Landscape.cpp: Terrain rendering
 #include "WorldPartition/Landscape/LandscapeActorDesc.h"
 #include "WorldPartition/WorldPartitionActorDescInstance.h"
 #include "Engine/Texture2DArray.h"
+#include "Modules/ModuleManager.h"
 
 #if WITH_EDITOR
 #include "Cooker/CookEvents.h"
@@ -595,7 +597,7 @@ FGuid ALandscapeProxy::GetNaniteContentId() const
 				return true;
 			}
 			if (!B)
-		{
+			{
 				return false;
 			}
 			// Sort components based on their SectionBase (i.e. 2D index relative to the entire landscape) to ensure stable ID generation
@@ -2102,7 +2104,20 @@ void ULandscapeComponent::BeginDestroy()
 
 FPrimitiveSceneProxy* ULandscapeComponent::CreateSceneProxy()
 {
-	return new FLandscapeComponentSceneProxy(this);
+	FPrimitiveSceneProxy* LocalSceneProxy = nullptr;
+
+	ILandscapeModule& LandscapeModule = FModuleManager::GetModuleChecked<ILandscapeModule>("Landscape");
+	const UE::Landscape::FCreateLandscapeComponentSceneProxyDelegate& CreateProxyDelegate = LandscapeModule.GetCreateLandscapeComponentSceneProxyDelegate();
+	if (CreateProxyDelegate.IsBound())
+	{
+		LocalSceneProxy = CreateProxyDelegate.Execute(this);
+	}
+	
+	if (LocalSceneProxy == nullptr)
+	{
+		LocalSceneProxy = new FLandscapeComponentSceneProxy(this);
+	}
+	return LocalSceneProxy;
 }
 
 bool ULandscapeComponent::IsShown(const FEngineShowFlags& ShowFlags) const
diff --git a/Engine/Source/Runtime/Landscape/Private/LandscapeModule.cpp b/Engine/Source/Runtime/Landscape/Private/LandscapeModule.cpp
--- a/Engine/Source/Runtime/Landscape/Private/LandscapeModule.cpp
+++ b/Engine/Source/Runtime/Landscape/Private/LandscapeModule.cpp
@@ -44,6 +44,8 @@ public:
 
 	virtual void SetLandscapeEditorServices(ILandscapeEditorServices* InLandscapeEditorServices) override { LandscapeEditorServices = InLandscapeEditorServices; }
 	virtual ILandscapeEditorServices* GetLandscapeEditorServices() const override { return LandscapeEditorServices; }
+	virtual void SetCreateLandscapeComponentSceneProxyDelegate(const UE::Landscape::FCreateLandscapeComponentSceneProxyDelegate& InDelegate) override { CreateLandscapeComponentSceneProxyDelegate = InDelegate; }
+	virtual const UE::Landscape::FCreateLandscapeComponentSceneProxyDelegate& GetCreateLandscapeComponentSceneProxyDelegate() const override { return CreateLandscapeComponentSceneProxyDelegate; }
 
 private:
 	void OnPostEngineInit();
@@ -56,6 +58,7 @@ private:
 private:
 	TSharedPtr<FLandscapeSceneViewExtension, ESPMode::ThreadSafe> SceneViewExtension;
 	ILandscapeEditorServices* LandscapeEditorServices = nullptr;
+	UE::Landscape::FCreateLandscapeComponentSceneProxyDelegate CreateLandscapeComponentSceneProxyDelegate;
 
 	/** Event binding handles */
 #if ALLOW_LANDSCAPE_COLLISION_COMPONENT_DEBUG_DISPLAY
diff --git a/Engine/Source/Runtime/Landscape/Public/LandscapeModule.h b/Engine/Source/Runtime/Landscape/Public/LandscapeModule.h
--- a/Engine/Source/Runtime/Landscape/Public/LandscapeModule.h
+++ b/Engine/Source/Runtime/Landscape/Public/LandscapeModule.h
@@ -6,7 +6,14 @@
 #include "Modules/ModuleInterface.h"
 #include "LandscapeEditorServices.h"
 
+class FPrimitiveSceneProxy;
 class FLandscapeSceneViewExtension;
+class ULandscapeComponent;
+
+namespace UE::Landscape
+{
+	DECLARE_DELEGATE_RetVal_OneParam(FPrimitiveSceneProxy*, FCreateLandscapeComponentSceneProxyDelegate, ULandscapeComponent* /*InLandscapeComponent*/);
+} // namespace UE::Landscape
 
 /**
 * Landscape module interface
@@ -16,6 +23,9 @@ class ILandscapeModule : public IModuleInterface
 public:
 	virtual TSharedPtr<FLandscapeSceneViewExtension, ESPMode::ThreadSafe> GetLandscapeSceneViewExtension() const = 0;
 	
-	virtual void SetLandscapeEditorServices(ILandscapeEditorServices* InLandscapeEditorServices) = 0;
-	virtual ILandscapeEditorServices* GetLandscapeEditorServices() const = 0;
+	virtual void SetLandscapeEditorServices(ILandscapeEditorServices* InLandscapeEditorServices) PURE_VIRTUAL(ILandscapeModule::SetLandscapeEditorServices, );
+	virtual ILandscapeEditorServices* GetLandscapeEditorServices() const PURE_VIRTUAL(ILandscapeModule::GetLandscapeEditorServices, return nullptr; );
+	virtual void SetCreateLandscapeComponentSceneProxyDelegate(const UE::Landscape::FCreateLandscapeComponentSceneProxyDelegate& InDelegate) PURE_VIRTUAL(ILandscapeModule::SetCreateLandscapeComponentSceneProxyDelegate, );
+	virtual const UE::Landscape::FCreateLandscapeComponentSceneProxyDelegate& GetCreateLandscapeComponentSceneProxyDelegate() const 
+		PURE_VIRTUAL(ILandscapeModule::GetCreateLandscapeComponentSceneProxyDelegate, static UE::Landscape::FCreateLandscapeComponentSceneProxyDelegate Dummy; return Dummy; );
 };
diff --git a/Engine/Source/Runtime/Landscape/Public/LandscapeRender.h b/Engine/Source/Runtime/Landscape/Public/LandscapeRender.h
--- a/Engine/Source/Runtime/Landscape/Public/LandscapeRender.h
+++ b/Engine/Source/Runtime/Landscape/Public/LandscapeRender.h
@@ -897,7 +897,7 @@ protected:
 	LANDSCAPE_API bool GetMeshElementForFixedGrid(int32 InLodIndex, FMaterialRenderProxy* InMaterialInterface, FMeshBatch& OutMeshBatch, TArray<FLandscapeBatchElementParams>& OutStaticBatchParamArray) const;
 	template<class ArrayType> bool GetStaticMeshElement(int32 LODIndex, bool bForToolMesh, FMeshBatch& MeshBatch, ArrayType& OutStaticBatchParamArray) const;
 
-	virtual void OnForceHiddenChanged() override;
+	LANDSCAPE_API virtual void OnForceHiddenChanged() override;
 
 public:
 	// constructor
@@ -947,9 +947,9 @@ public:
 	LANDSCAPE_API virtual int32 GetComponentResolution() const override;
 	LANDSCAPE_API virtual double ComputeSectionResolution() const override;
 	LANDSCAPE_API virtual void GetSectionBoundsAndLocalToWorld(FBoxSphereBounds& LocalBounds, FMatrix& LocalToWorld) const override;
-	virtual void GetSectionCenterAndVectors(FVector& OutSectionCenterWorldSpace, FVector& OutSectionXVectorWorldSpace, FVector& OutSectionYVectorWorldSpace) const override;
-	virtual const FPrimitiveSceneInfo* GetPrimitiveSceneInfo() const override;
-	virtual bool ShouldInvalidateShadows(const FSceneView& InView, float InLODValue, float InLastShadowInvalidationLODValue) const override;
+	LANDSCAPE_API virtual void GetSectionCenterAndVectors(FVector& OutSectionCenterWorldSpace, FVector& OutSectionXVectorWorldSpace, FVector& OutSectionYVectorWorldSpace) const override;
+	LANDSCAPE_API virtual const FPrimitiveSceneInfo* GetPrimitiveSceneInfo() const override;
+	LANDSCAPE_API virtual bool ShouldInvalidateShadows(const FSceneView& InView, float InLODValue, float InLastShadowInvalidationLODValue) const override;
 };
 
 class FLandscapeDebugMaterialRenderProxy : public FMaterialRenderProxy
