Re-enable landscape selective edit layer renders and fix some asserts with that code.

diff --git a/Engine/Source/Editor/LandscapeEditor/Private/LandscapeEdModeTools.h b/Engine/Source/Editor/LandscapeEditor/Private/LandscapeEdModeTools.h
--- a/Engine/Source/Editor/LandscapeEditor/Private/LandscapeEdModeTools.h
+++ b/Engine/Source/Editor/LandscapeEditor/Private/LandscapeEdModeTools.h
@@ -1068,9 +1068,6 @@ public:
 				TSet<ULandscapeComponent*> AffectedComponents;
 				LandscapeInfo->GetComponentsInRegion(DesiredCacheBounds.Min.X, DesiredCacheBounds.Min.Y, DesiredCacheBounds.Max.X, DesiredCacheBounds.Max.Y, AffectedComponents);
 
-				UE_LOG(LogLandscapeTools, Warning, TEXT("IntRender (fr%u)   [%d, %d]-[%d, %d]  %d comps"), GFrameNumber,
-					DesiredCacheBounds.Min.X, DesiredCacheBounds.Min.Y, DesiredCacheBounds.Max.X, DesiredCacheBounds.Max.Y, AffectedComponents.Num());
-
 				SynchronousUpdateComponentVisibilityForHeight(AffectedComponents, NewLayerVisibility);
 
 				if (bVisibilityChanged)
diff --git a/Engine/Source/Runtime/Landscape/Private/LandscapeEditLayers.cpp b/Engine/Source/Runtime/Landscape/Private/LandscapeEditLayers.cpp
--- a/Engine/Source/Runtime/Landscape/Private/LandscapeEditLayers.cpp
+++ b/Engine/Source/Runtime/Landscape/Private/LandscapeEditLayers.cpp
@@ -272,7 +272,7 @@ TAutoConsoleVariable<bool> CVarSilenceMergeBatchResolutionWarning(
 // TODO:  remove with IntermediateRender.
 TAutoConsoleVariable<bool> CVarLandscapeSelectiveRenderEditLayers(
 	TEXT("landscape.BatchedMerge.SelectiveRenderEditLayers"),
-	false,
+	true,
 	TEXT("When true, use SelectiveRenderEditLayers codepath instead of intermediate renders"));
 
 
@@ -8158,8 +8158,7 @@ bool ALandscape::PerformSelectiveRenderEditLayersHeightmapsCPU(const FUpdateLaye
 					check(BatchResult.SectionRect.Contains(ComponentRect));
 					
 					const FBlitBuffer2DDesc SrcDesc{ SrcBufferSize, BatchStride, BatchResult.SectionRect };
-					int32 CopiedValues = BlitHeightChannelsToUint16(Dest.GetData(), DestDesc, SrcData, SrcDesc, ComponentRect);
-					ensure(CopiedValues > 0);  
+					BlitHeightChannelsToUint16(Dest.GetData(), DestDesc, SrcData, SrcDesc, ComponentRect);
 				}
 
 				BatchResult.Readback.Unlock();
diff --git a/Engine/Source/Runtime/Landscape/Private/LandscapeUtilsPrivate.cpp b/Engine/Source/Runtime/Landscape/Private/LandscapeUtilsPrivate.cpp
--- a/Engine/Source/Runtime/Landscape/Private/LandscapeUtilsPrivate.cpp
+++ b/Engine/Source/Runtime/Landscape/Private/LandscapeUtilsPrivate.cpp
@@ -98,7 +98,7 @@ int32 BlitHeightChannelsToUint16 (uint16* DestBuffer, const FBlitBuffer2DDesc& D
 	FIntRect BlitRect = Dest.Rect;
 	BlitRect.Clip(Src.Rect);
 	BlitRect.Clip(ClipRect);
-	if (BlitRect.IsEmpty())
+	if (BlitRect.Area() <= 0)
 	{
 		return 0;
 	}
@@ -118,7 +118,7 @@ int32 BlitHeightChannelsToUint16 (uint16* DestBuffer, const FBlitBuffer2DDesc& D
 		uint16* WriteEnd = Write + BlitSize.X;
 
 		check(Read >= SrcBuffer);
-		check(ReadEnd <= SrcBuffer + (Src.Stride * Src.BufferSize));
+		check(ReadEnd <= SrcBuffer + Src.BufferSize);
 		check(Read <= ReadEnd);
 		check(Write >= DestBuffer);
 		check(WriteEnd <= DestBuffer + Dest.BufferSize);
