World Partition - HLOD - Landscape: Disable distance field generation on landscape HLODs
* Rebuilt HLODs for the template open world map

#virtualized

diff --git a/Engine/Source/Runtime/Landscape/Private/LandscapeHLODBuilder.cpp b/Engine/Source/Runtime/Landscape/Private/LandscapeHLODBuilder.cpp
--- a/Engine/Source/Runtime/Landscape/Private/LandscapeHLODBuilder.cpp
+++ b/Engine/Source/Runtime/Landscape/Private/LandscapeHLODBuilder.cpp
@@ -44,7 +44,7 @@ uint32 ULandscapeHLODBuilder::ComputeHLODHash(const UActorComponent* InSourceCom
 	FArchiveCrc32 Ar;
 
 	// Base lanscape HLOD key, changing this will force a rebuild of all landscape HLODs
-	FString HLODBaseKey = "38DC3700FC0742929BA00ACCF5B1B626";
+	FString HLODBaseKey = "51EDA4343E26409DA6ABEACE14C138D1";
 	Ar << HLODBaseKey;
 
 	if (const ULandscapeComponent* LSComponent = Cast<ULandscapeComponent>(InSourceComponent))
@@ -376,6 +376,7 @@ TArray<UActorComponent*> ULandscapeHLODBuilder::Build(const FHLODBuildContext& I
 		// Mesh
 		{
 			FStaticMeshSourceModel& SrcModel = StaticMesh->AddSourceModel();
+			
 			// Don't allow the engine to recalculate normals
 			SrcModel.BuildSettings.bRecomputeNormals = false;
 			SrcModel.BuildSettings.bRecomputeTangents = false;
@@ -383,6 +384,9 @@ TArray<UActorComponent*> ULandscapeHLODBuilder::Build(const FHLODBuildContext& I
 			SrcModel.BuildSettings.bUseHighPrecisionTangentBasis = false;
 			SrcModel.BuildSettings.bUseFullPrecisionUVs = false;
 
+			// Disable distance field generation
+			SrcModel.BuildSettings.DistanceFieldResolutionScale = 0;
+
 			MeshDescription = StaticMesh->CreateMeshDescription(0);
 	
 			ALandscapeProxy::FRawMeshExportParams ExportParams;
