Landscape Editor - Fix for nanite build missing components when proxy has more more than NaniteMaxComponents

diff --git a/Engine/Source/Runtime/Landscape/Classes/LandscapeProxy.h b/Engine/Source/Runtime/Landscape/Classes/LandscapeProxy.h
--- a/Engine/Source/Runtime/Landscape/Classes/LandscapeProxy.h
+++ b/Engine/Source/Runtime/Landscape/Classes/LandscapeProxy.h
@@ -1075,7 +1075,7 @@ PRAGMA_ENABLE_DEPRECATION_WARNINGS
 	static constexpr int32 NaniteMaxComponents = NaniteComponentMaxSide * NaniteComponentMaxSide;
 	int32 NumNaniteRequiredComponents() const { return FMath::DivideAndRoundUp(LandscapeComponents.Num(), NaniteMaxComponents); }
 
-	TArray<ULandscapeComponent*> GatherSourceComponentsForNaniteComponent(int32 NaniteComponentIndex) const;
+	void SetSourceComponentsForNaniteComponents(bool bSetEmptyComponentsOnly = false) const;
 #endif	//WITH_EDITOR
 
 	bool AreNaniteComponentsValid(const FGuid& InProxyContentId) const;
diff --git a/Engine/Source/Runtime/Landscape/Private/Landscape.cpp b/Engine/Source/Runtime/Landscape/Private/Landscape.cpp
--- a/Engine/Source/Runtime/Landscape/Private/Landscape.cpp
+++ b/Engine/Source/Runtime/Landscape/Private/Landscape.cpp
@@ -467,13 +467,16 @@ FGraphEventRef ALandscapeProxy::UpdateNaniteRepresentationAsync(const ITargetPla
 		}
 
 		FGraphEventArray UpdateDependencies;
+
+		// Set the source landscape components for each nanite component
+		SetSourceComponentsForNaniteComponents();
 		for (int32 i = 0; i < NumNewNaniteComponents; ++i)
 		{
 			FGraphEventArray SingleProxyDependencies;
 
 			if (bNaniteContentDirty)
 			{
-				FGraphEventRef ComponentProcessTask = NaniteComponents[i]->InitializeForLandscapeAsync(this, NaniteContentId, GatherSourceComponentsForNaniteComponent(i), i);
+				FGraphEventRef ComponentProcessTask = NaniteComponents[i]->InitializeForLandscapeAsync(this, NaniteContentId, NaniteComponents[i]->GetSourceLandscapeComponents(), i);
 				SingleProxyDependencies.Add(ComponentProcessTask);
 			}
 
@@ -4634,14 +4637,8 @@ PRAGMA_DISABLE_DEPRECATION_WARNINGS
 	}
 PRAGMA_ENABLE_DEPRECATION_WARNINGS
 
-	// Fix for nanite components that have lost their link to their source landscape components
-	for (int32 i = 0; i < NaniteComponents.Num(); ++i)
-	{
-		if (NaniteComponents[i] && NaniteComponents[i]->GetSourceLandscapeComponents().Num() == 0)
-		{
-			NaniteComponents[i]->SetSourceLandscapeComponents(GatherSourceComponentsForNaniteComponent(i));
-		}
-	}
+	// Fix for empty nanite components that have lost their link to their source landscape components
+	SetSourceComponentsForNaniteComponents(/*bSetEmptyComponentsOnly=*/ true);
 
 	// Handle Nanite representation invalidation on load: 
 	if (!HasAnyFlags(RF_ClassDefaultObject) && !FPlatformProperties::RequiresCookedData())
@@ -6860,29 +6857,38 @@ void ALandscapeProxy::UpdateNaniteSharedPropertiesFromActor()
 	}
 }
 
-TArray<ULandscapeComponent*> ALandscapeProxy::GatherSourceComponentsForNaniteComponent(int32 NaniteComponentIndex) const
+void ALandscapeProxy::SetSourceComponentsForNaniteComponents(bool bSetEmptyComponentsOnly) const
 {
-	const int32 StartComponentIndex = NaniteComponentIndex * NaniteMaxComponents;
-	const int32 EndComponentIndex = FMath::Min(LandscapeComponents.Num(), (NaniteComponentIndex + 1) * NaniteMaxComponents);
-	const int32 NumComponents = EndComponentIndex - StartComponentIndex;
+	check(IsNaniteEnabled() ? (NaniteComponents.Num() == NumNaniteRequiredComponents()) : NaniteComponents.IsEmpty());
+
+	TArray<ULandscapeComponent*> StableOrderComponents(LandscapeComponents);
 
-	if (NumComponents <= 0 || !LandscapeComponents.IsValidIndex(StartComponentIndex))
+	for (int32 NaniteComponentIndex = 0; NaniteComponentIndex < NaniteComponents.Num(); ++NaniteComponentIndex)
 	{
-		return TArray<ULandscapeComponent*>();
-	}
+		const int32 StartComponentIndex = NaniteComponentIndex * NaniteMaxComponents;
+		const int32 EndComponentIndex = FMath::Min(LandscapeComponents.Num(), (NaniteComponentIndex + 1) * NaniteMaxComponents);
+		const int32 NumComponents = EndComponentIndex - StartComponentIndex;
 
-	TArrayView<const TObjectPtr<ULandscapeComponent>> ComponentsView(&LandscapeComponents[StartComponentIndex], LandscapeComponents.Num() - StartComponentIndex);
-	TArray<ULandscapeComponent*> Result(ComponentsView);
+		ULandscapeNaniteComponent* NaniteComponent = NaniteComponents[NaniteComponentIndex];
+		// Optionally skip nanite components that already have valid source components set
+		if (NumComponents <= 0 || !LandscapeComponents.IsValidIndex(StartComponentIndex) || (!NaniteComponent->GetSourceLandscapeComponents().IsEmpty() && bSetEmptyComponentsOnly))
+		{
+			continue;
+		}
 
-	ULandscapeComponent**const MinComponent = Algo::MinElementBy(Result,
-		[](const ULandscapeComponent* Component) { return Component->GetSectionBase(); },
-		[](const FIntPoint& A, const FIntPoint& B) { return (A.Y == B.Y) ? (A.X < B.X) : (A.Y < B.Y); }
-	);
-	check(MinComponent);
+		// Select the next available components from our distance sorted component pool each iteration
+		TArrayView<ULandscapeComponent*> StableOrderComponentsView = TArrayView<ULandscapeComponent*>(&StableOrderComponents[StartComponentIndex], LandscapeComponents.Num() - StartComponentIndex);
 
-	Algo::Sort(Result, FCompareULandscapeComponentClosest((*MinComponent)->GetSectionBase()));
-	Result.SetNum(NumComponents);
-	return Result;
+		ULandscapeComponent** const MinComponent = Algo::MinElementBy(StableOrderComponentsView,
+			[](const ULandscapeComponent* Component) { return Component->GetSectionBase(); },
+			[](const FIntPoint& A, const FIntPoint& B) { return (A.Y == B.Y) ? (A.X < B.X) : (A.Y < B.Y); }
+		);
+		check(MinComponent);
+
+		Algo::Sort(StableOrderComponentsView, FCompareULandscapeComponentClosest((*MinComponent)->GetSectionBase()));
+		TArray<ULandscapeComponent*> Result(StableOrderComponentsView.GetData(), NumComponents);
+		NaniteComponent->SetSourceLandscapeComponents(Result);
+	}
 }
 
 void ALandscapeProxy::InvalidatePhysicalMaterial()
