Improved landscape.DumpLODs command : now it only displays the landscape LOD information by default and -detailed needs to be used to get the verbose mip-to-mip delta infos

diff --git a/Engine/Source/Runtime/Landscape/Private/LandscapeRender.cpp b/Engine/Source/Runtime/Landscape/Private/LandscapeRender.cpp
--- a/Engine/Source/Runtime/Landscape/Private/LandscapeRender.cpp
+++ b/Engine/Source/Runtime/Landscape/Private/LandscapeRender.cpp
@@ -248,9 +248,16 @@ extern int32 GGrassMapAlwaysBuildRuntimeGenerationResources;
 
 #if !UE_BUILD_SHIPPING
 uint32 GVarDumpLandscapeLODsCurrentFrame = 0;
+bool GVarDumpLandscapeLODsDetails = false;
 
 static void OnDumpLandscapeLODs(const TArray< FString >& Args)
 {
+	GVarDumpLandscapeLODsDetails = false;
+	if (!Args.IsEmpty() && FParse::Param(*Args[0], TEXT("detailed")))
+	{
+		GVarDumpLandscapeLODsDetails = true;
+	}
+
 	if (GVarDumpLandscapeLODsCurrentFrame == 0)
 	{
 		// Add some buffer to be able to correctly catch the frame during the rendering
@@ -260,7 +267,7 @@ static void OnDumpLandscapeLODs(const TArray< FString >& Args)
 
 static FAutoConsoleCommand CVarDumpLandscapeLODs(
 	TEXT("landscape.DumpLODs"),
-	TEXT("Will dump the current status of LOD value and current texture streaming status"),
+	TEXT("Will dump the current status of LOD value and current texture streaming status. Use the -detailed option to include all LOD-related infos"),
 	FConsoleCommandWithArgsDelegate::CreateStatic(&OnDumpLandscapeLODs)
 );
 #endif
@@ -3006,18 +3013,6 @@ void FLandscapeComponentSceneProxy::GetDynamicMeshElements(const TArray<const FS
 #if !UE_BUILD_SHIPPING
 	if (GVarDumpLandscapeLODsCurrentFrame == GFrameNumberRenderThread)
 	{
-		// Dump the mip-to-mip info for this component to evaluate how different mips are from one another: 
-		FString MipToMipInfoString;
-		for (int32 SourceMipIndex = 0; SourceMipIndex < NumRelevantMips - 1; ++SourceMipIndex)
-		{
-			for (int32 DestinationMipIndex = SourceMipIndex + 1; DestinationMipIndex < NumRelevantMips; ++DestinationMipIndex)
-			{
-				int32 MipToMipDeltaIndex = UE::Landscape::Private::ComputeMipToMipMaxDeltasIndex(SourceMipIndex, DestinationMipIndex, NumRelevantMips);
-				MipToMipInfoString += FString::Printf(TEXT("- %i->%i: %f\n"), SourceMipIndex, DestinationMipIndex, WorldSpaceMipToMipMaxDeltas[MipToMipDeltaIndex]);
-			}
-		}
-		UE_LOG(LogLandscape, Display, TEXT("%s, WorldSpaceMipToMipMaxDeltas:\n%s"), *DebugName.ToString(), *MipToMipInfoString);
-
 		for (const FSceneView* View : Views)
 		{
 			const float LODValue = ComputeLODForView(*View);
@@ -3048,9 +3043,24 @@ void FLandscapeComponentSceneProxy::GetDynamicMeshElements(const TArray<const FS
 				return MaxMipCount > 0 ? 1 << (MaxMipCount - 1) : 0;
 			}();
 
-			UE_LOG(LogLandscape, Display, TEXT("\nView: %d, %s, LODValue: %f (LOD: %d), Resolution: %d, LoadedHeightmapMIP: %d, LoadedWeightmapMIP: %d"),
+			UE_LOG(LogLandscape, Display, TEXT("\nView: %d, %s, LODValue: %f (LOD: %d), Resolution: %d, LoadedHeightmapResolution: %d, LoadedWeightmapResolution: %d"),
 			       View->GetViewKey(), *DebugName.ToString(), LODValue, LOD, Resolution, LoadedHeightmapResolution, LoadedWeightmapResolution);
 		}
+
+		if (GVarDumpLandscapeLODsDetails)
+		{
+			// Dump the mip-to-mip info for this component to evaluate how different mips are from one another: 
+			FString MipToMipInfoString;
+			for (int32 SourceMipIndex = 0; SourceMipIndex < NumRelevantMips - 1; ++SourceMipIndex)
+			{
+				for (int32 DestinationMipIndex = SourceMipIndex + 1; DestinationMipIndex < NumRelevantMips; ++DestinationMipIndex)
+				{
+					int32 MipToMipDeltaIndex = UE::Landscape::Private::ComputeMipToMipMaxDeltasIndex(SourceMipIndex, DestinationMipIndex, NumRelevantMips);
+					MipToMipInfoString += FString::Printf(TEXT("- %i->%i: %f\n"), SourceMipIndex, DestinationMipIndex, WorldSpaceMipToMipMaxDeltas[MipToMipDeltaIndex]);
+				}
+			}
+			UE_LOG(LogLandscape, Display, TEXT("%s, WorldSpaceMipToMipMaxDeltas:\n%s"), *DebugName.ToString(), *MipToMipInfoString);
+		}
 	}
 #endif // !UE_BUILD_SHIPPING
 }
