landscape.DumpLODs argument parsing improvement suggested by BugHawk

diff --git a/Engine/Source/Runtime/Landscape/Private/LandscapeRender.cpp b/Engine/Source/Runtime/Landscape/Private/LandscapeRender.cpp
--- a/Engine/Source/Runtime/Landscape/Private/LandscapeRender.cpp
+++ b/Engine/Source/Runtime/Landscape/Private/LandscapeRender.cpp
@@ -253,9 +253,12 @@ bool GVarDumpLandscapeLODsDetails = false;
 static void OnDumpLandscapeLODs(const TArray< FString >& Args)
 {
 	GVarDumpLandscapeLODsDetails = false;
-	if (!Args.IsEmpty() && FParse::Param(*Args[0], TEXT("detailed")))
+	for (const FString& Arg : Args)
 	{
-		GVarDumpLandscapeLODsDetails = true;
+		if (FParse::Param(*Arg, TEXT("detailed")))
+		{
+			GVarDumpLandscapeLODsDetails = true;
+		}
 	}
 
 	if (GVarDumpLandscapeLODsCurrentFrame == 0)
