Landscape Editor - Fix PIE when nanite landscape has no nanite components
* Ignores landscape componet assignment when there are no nanite components

diff --git a/Engine/Source/Runtime/Landscape/Private/Landscape.cpp b/Engine/Source/Runtime/Landscape/Private/Landscape.cpp
--- a/Engine/Source/Runtime/Landscape/Private/Landscape.cpp
+++ b/Engine/Source/Runtime/Landscape/Private/Landscape.cpp
@@ -6859,7 +6859,13 @@ void ALandscapeProxy::UpdateNaniteSharedPropertiesFromActor()
 
 void ALandscapeProxy::SetSourceComponentsForNaniteComponents(bool bSetEmptyComponentsOnly) const
 {
-	check(IsNaniteEnabled() ? (NaniteComponents.Num() == NumNaniteRequiredComponents()) : NaniteComponents.IsEmpty());
+	// Nanite components are reset after invalidation, early out until new components are created
+	if (!IsNaniteEnabled() || NaniteComponents.IsEmpty())
+	{
+		return;
+	}
+	// Ensure we have the correct number of nanite components based on the proxy component count
+	check(NaniteComponents.Num() == NumNaniteRequiredComponents());
 
 	TArray<ULandscapeComponent*> StableOrderComponents(LandscapeComponents);
 
