[Backout] - CL44564328
#rnx
Original CL Desc
-----------------------------------------------------------------
Landscape Editor - Fix for nanite build missing components when proxy has more more than NaniteMaxComponents

diff --git a/Engine/Source/Runtime/Landscape/Classes/LandscapeProxy.h b/Engine/Source/Runtime/Landscape/Classes/LandscapeProxy.h
--- a/Engine/Source/Runtime/Landscape/Classes/LandscapeProxy.h
+++ b/Engine/Source/Runtime/Landscape/Classes/LandscapeProxy.h
@@ -1075,7 +1075,7 @@ PRAGMA_ENABLE_DEPRECATION_WARNINGS
 	static constexpr int32 NaniteMaxComponents = NaniteComponentMaxSide * NaniteComponentMaxSide;
 	int32 NumNaniteRequiredComponents() const { return FMath::DivideAndRoundUp(LandscapeComponents.Num(), NaniteMaxComponents); }
 
-	void SetSourceComponentsForNaniteComponents(bool bSetEmptyComponentsOnly = false) const;
+	TArray<ULandscapeComponent*> GatherSourceComponentsForNaniteComponent(int32 NaniteComponentIndex) const;
 #endif	//WITH_EDITOR
 
 	bool AreNaniteComponentsValid(const FGuid& InProxyContentId) const;
diff --git a/Engine/Source/Runtime/Landscape/Private/Landscape.cpp b/Engine/Source/Runtime/Landscape/Private/Landscape.cpp
--- a/Engine/Source/Runtime/Landscape/Private/Landscape.cpp
+++ b/Engine/Source/Runtime/Landscape/Private/Landscape.cpp
@@ -467,16 +467,13 @@ FGraphEventRef ALandscapeProxy::UpdateNaniteRepresentationAsync(const ITargetPla
 		}
 
 		FGraphEventArray UpdateDependencies;
-
-		// Set the source landscape components for each nanite component
-		SetSourceComponentsForNaniteComponents();
 		for (int32 i = 0; i < NumNewNaniteComponents; ++i)
 		{
 			FGraphEventArray SingleProxyDependencies;
 
 			if (bNaniteContentDirty)
 			{
-				FGraphEventRef ComponentProcessTask = NaniteComponents[i]->InitializeForLandscapeAsync(this, NaniteContentId, NaniteComponents[i]->GetSourceLandscapeComponents(), i);
+				FGraphEventRef ComponentProcessTask = NaniteComponents[i]->InitializeForLandscapeAsync(this, NaniteContentId, GatherSourceComponentsForNaniteComponent(i), i);
 				SingleProxyDependencies.Add(ComponentProcessTask);
 			}
 
@@ -4637,8 +4634,14 @@ PRAGMA_DISABLE_DEPRECATION_WARNINGS
 	}
 PRAGMA_ENABLE_DEPRECATION_WARNINGS
 
-	// Fix for empty nanite components that have lost their link to their source landscape components
-	SetSourceComponentsForNaniteComponents(/*bSetEmptyComponentsOnly=*/ true);
+	// Fix for nanite components that have lost their link to their source landscape components
+	for (int32 i = 0; i < NaniteComponents.Num(); ++i)
+	{
+		if (NaniteComponents[i] && NaniteComponents[i]->GetSourceLandscapeComponents().Num() == 0)
+		{
+			NaniteComponents[i]->SetSourceLandscapeComponents(GatherSourceComponentsForNaniteComponent(i));
+		}
+	}
 
 	// Handle Nanite representation invalidation on load: 
 	if (!HasAnyFlags(RF_ClassDefaultObject) && !FPlatformProperties::RequiresCookedData())
@@ -6857,44 +6860,29 @@ void ALandscapeProxy::UpdateNaniteSharedPropertiesFromActor()
 	}
 }
 
-void ALandscapeProxy::SetSourceComponentsForNaniteComponents(bool bSetEmptyComponentsOnly) const
+TArray<ULandscapeComponent*> ALandscapeProxy::GatherSourceComponentsForNaniteComponent(int32 NaniteComponentIndex) const
 {
-	// Nanite components are reset after invalidation, early out until new components are created
-	if (!IsNaniteEnabled() || NaniteComponents.IsEmpty())
-	{
-		return;
-	}
-	// Ensure we have the correct number of nanite components based on the proxy component count
-	check(NaniteComponents.Num() == NumNaniteRequiredComponents());
-
-	TArray<ULandscapeComponent*> StableOrderComponents(LandscapeComponents);
+	const int32 StartComponentIndex = NaniteComponentIndex * NaniteMaxComponents;
+	const int32 EndComponentIndex = FMath::Min(LandscapeComponents.Num(), (NaniteComponentIndex + 1) * NaniteMaxComponents);
+	const int32 NumComponents = EndComponentIndex - StartComponentIndex;
 
-	for (int32 NaniteComponentIndex = 0; NaniteComponentIndex < NaniteComponents.Num(); ++NaniteComponentIndex)
+	if (NumComponents <= 0 || !LandscapeComponents.IsValidIndex(StartComponentIndex))
 	{
-		const int32 StartComponentIndex = NaniteComponentIndex * NaniteMaxComponents;
-		const int32 EndComponentIndex = FMath::Min(LandscapeComponents.Num(), (NaniteComponentIndex + 1) * NaniteMaxComponents);
-		const int32 NumComponents = EndComponentIndex - StartComponentIndex;
-
-		ULandscapeNaniteComponent* NaniteComponent = NaniteComponents[NaniteComponentIndex];
-		// Optionally skip nanite components that already have valid source components set
-		if (NumComponents <= 0 || !LandscapeComponents.IsValidIndex(StartComponentIndex) || (!NaniteComponent->GetSourceLandscapeComponents().IsEmpty() && bSetEmptyComponentsOnly))
-		{
-			continue;
-		}
+		return TArray<ULandscapeComponent*>();
+	}
 
-		// Select the next available components from our distance sorted component pool each iteration
-		TArrayView<ULandscapeComponent*> StableOrderComponentsView = TArrayView<ULandscapeComponent*>(&StableOrderComponents[StartComponentIndex], LandscapeComponents.Num() - StartComponentIndex);
+	TArrayView<const TObjectPtr<ULandscapeComponent>> ComponentsView(&LandscapeComponents[StartComponentIndex], LandscapeComponents.Num() - StartComponentIndex);
+	TArray<ULandscapeComponent*> Result(ComponentsView);
 
-		ULandscapeComponent** const MinComponent = Algo::MinElementBy(StableOrderComponentsView,
-			[](const ULandscapeComponent* Component) { return Component->GetSectionBase(); },
-			[](const FIntPoint& A, const FIntPoint& B) { return (A.Y == B.Y) ? (A.X < B.X) : (A.Y < B.Y); }
-		);
-		check(MinComponent);
+	ULandscapeComponent**const MinComponent = Algo::MinElementBy(Result,
+		[](const ULandscapeComponent* Component) { return Component->GetSectionBase(); },
+		[](const FIntPoint& A, const FIntPoint& B) { return (A.Y == B.Y) ? (A.X < B.X) : (A.Y < B.Y); }
+	);
+	check(MinComponent);
 
-		Algo::Sort(StableOrderComponentsView, FCompareULandscapeComponentClosest((*MinComponent)->GetSectionBase()));
-		TArray<ULandscapeComponent*> Result(StableOrderComponentsView.GetData(), NumComponents);
-		NaniteComponent->SetSourceLandscapeComponents(Result);
-	}
+	Algo::Sort(Result, FCompareULandscapeComponentClosest((*MinComponent)->GetSectionBase()));
+	Result.SetNum(NumComponents);
+	return Result;
 }
 
 void ALandscapeProxy::InvalidatePhysicalMaterial()
