Adding GrassMap RenderEveryFrame command

command:   grass.GrassMap.RenderEveryFrame <component or proxy name>

This will render the grass maps for any landscape components with matching names every frame.
This is intended to be able to render-capture the results of grass map rendering at runtime, for debugging purposes.

Also fixes a race on the AsyncReadbackPtr.

diff --git a/Engine/Source/Runtime/Landscape/Internal/LandscapeGrassWeightExporter.h b/Engine/Source/Runtime/Landscape/Internal/LandscapeGrassWeightExporter.h
--- a/Engine/Source/Runtime/Landscape/Internal/LandscapeGrassWeightExporter.h
+++ b/Engine/Source/Runtime/Landscape/Internal/LandscapeGrassWeightExporter.h
@@ -87,13 +87,14 @@ private:
 	float PassOffsetX;
 	FVector ViewOrigin;
 
-	FLandscapeAsyncTextureReadback* AsyncReadbackPtr = nullptr;
+	// game thread synchronous, do not access directly from render thread
+	FLandscapeAsyncTextureReadback* GameThreadAsyncReadbackPtr = nullptr;
 
 	FMatrix ViewRotationMatrix;
 	FMatrix ProjectionMatrix;
 
 	/** Creates a texture and renders the component to it, and then and triggers a readback of the texture. */
-	void RenderLandscapeComponentToTexture_RenderThread(FRHICommandListImmediate& RHICmdList);
+	void RenderLandscapeComponentToTexture_RenderThread(FRHICommandListImmediate& RHICmdList, FLandscapeAsyncTextureReadback* AsyncReadbackPtr);
 };
 
 class FLandscapeGrassWeightExporter : public FLandscapeGrassWeightExporter_RenderThread
diff --git a/Engine/Source/Runtime/Landscape/Private/LandscapeGrassMapsBuilder.cpp b/Engine/Source/Runtime/Landscape/Private/LandscapeGrassMapsBuilder.cpp
--- a/Engine/Source/Runtime/Landscape/Private/LandscapeGrassMapsBuilder.cpp
+++ b/Engine/Source/Runtime/Landscape/Private/LandscapeGrassMapsBuilder.cpp
@@ -52,6 +52,20 @@ FAutoConsoleVariableRef CVarGrassMapEvictAll(
 	GGrassMapEvictAll,
 	TEXT("Remove all grass maps. In editor, or if runtime generation is enabled, the grass maps will repopulate automatically."));
 
+#if !UE_BUILD_SHIPPING
+FString GGrassMapRenderEveryFrame;
+FAutoConsoleVariableRef CVarGrassMapRenderEveryFrame(
+	TEXT("grass.GrassMap.RenderEveryFrame"),
+	GGrassMapRenderEveryFrame,
+	TEXT("Forces re-rendering the grassmap for the landscape component(s) on any actors with a name containing the given string."));
+
+int32 GGrassMapRenderEveryFrameDelay = 0;
+FAutoConsoleVariableRef CVarGrassMapRenderEveryFrameDelay(
+	TEXT("grass.GrassMap.RenderEveryFrameDelay"),
+	GGrassMapRenderEveryFrameDelay,
+	TEXT("When using grass.GrassMap.RenderEveryFrame, delays this many frames in between renders.  Default 0."));
+#endif // !UE_BUILD_SHIPPING
+
 int32 GGrassMapUseAsyncFetch = 0;
 FAutoConsoleVariableRef CVarGrassMapUseAsyncFetch(
 	TEXT("grass.GrassMap.UseAsyncFetch"),
@@ -335,7 +349,8 @@ FLandscapeGrassMapsBuilder::~FLandscapeGrassMapsBuilder()
 
 		const bool bCancelAndEvictAllImmediately = true;
 		const bool bEvictWhenBeyondEvictionRange = false;
-		UpdateTrackedComponents(EmptyCamerasArray, 0, UpdateAllComponentCount, bCancelAndEvictAllImmediately, bEvictWhenBeyondEvictionRange);
+		const int32 DeltaTicks = 0;
+		UpdateTrackedComponents(EmptyCamerasArray, 0, UpdateAllComponentCount, bCancelAndEvictAllImmediately, bEvictWhenBeyondEvictionRange, DeltaTicks);
 
 		ensure(NotReadyCount == 0 && StreamingCount == 0 && RenderingCount == 0 && AsyncFetchCount == 0 && PopulatedCount == 0);
 
@@ -352,9 +367,9 @@ FLandscapeGrassMapsBuilder::~FLandscapeGrassMapsBuilder()
 
 			if (State->Stage == EComponentStage::Rendering)
 			{
-				if ((State->ActiveRender != nullptr) && (State->ActiveRender->AsyncReadbackPtr != nullptr))
+				if ((State->ActiveRender != nullptr) && (State->ActiveRender->GameThreadAsyncReadbackPtr != nullptr))
 				{
-					UE_LOG(LogGrass, Warning, TEXT("  %s"), *State->ActiveRender->AsyncReadbackPtr->ToString());
+					UE_LOG(LogGrass, Warning, TEXT("  %s"), *State->ActiveRender->GameThreadAsyncReadbackPtr->ToString());
 				}
 			}
 			else if (State->Stage == EComponentStage::AsyncFetch)
@@ -411,7 +426,7 @@ void FAsyncFetchTask::DoWork()
 	Results = ActiveRender->FetchResults(bFreeAsyncReadback);
 }
 
-bool FLandscapeGrassMapsBuilder::UpdateTrackedComponents(const TArray<FVector>& Cameras, int32 LocalMaxRendering, int32 MaxExpensiveUpdateChecksToPerform, bool bCancelAndEvictAllImmediately, bool bEvictWhenBeyondEvictionRange)
+bool FLandscapeGrassMapsBuilder::UpdateTrackedComponents(const TArray<FVector>& Cameras, int32 LocalMaxRendering, int32 MaxExpensiveUpdateChecksToPerform, bool bCancelAndEvictAllImmediately, bool bEvictWhenBeyondEvictionRange, int32 DeltaTicks)
 {
 	TRACE_CPUPROFILER_EVENT_SCOPE(FLandscapeGrassMapsBuilder::UpdateTrackedComponents);
 	SCOPE_CYCLE_COUNTER(STAT_UpdateComponentGrassMaps);
@@ -439,7 +454,7 @@ bool FLandscapeGrassMapsBuilder::UpdateTrackedComponents(const TArray<FVector>&
 	{
 		FComponentState* State = It.Value();
 		ULandscapeComponent* Component = State->Component;
-		State->TickCount++;
+		State->TickCount += DeltaTicks;
 
 		if (bCancelAndEvictAllImmediately || (Component == nullptr))
 		{
@@ -841,6 +856,58 @@ void FLandscapeGrassMapsBuilder::AmortizedUpdateGrassMaps(
 		return; // if we can never ever render, don't bother to do anything here
 	}
 
+#if !UE_BUILD_SHIPPING
+	// determine if we have a debug render component
+	if (!GGrassMapRenderEveryFrame.IsEmpty())
+	{
+		bool bRenderThisFrame = true;
+		static int32 DelayTick = 0;
+		if (GGrassMapRenderEveryFrameDelay > 0)
+		{
+			DelayTick++;
+			bRenderThisFrame = false;
+			if (DelayTick > GGrassMapRenderEveryFrameDelay)
+			{
+				bRenderThisFrame = true;
+				DelayTick = 0;
+			}
+		}
+
+		if (bRenderThisFrame)
+		{
+			// trigger a debug render of the grass maps for the specified components
+			TArray<TObjectPtr<ULandscapeComponent>> DebugComponents;
+			for (auto It = ComponentStates.CreateIterator(); It; ++It)
+			{
+				FComponentState* State = It.Value();
+				if (ULandscapeComponent* Component = State->Component)
+				{
+					bool bDebugThisComponent = Component->GetName().Contains(GGrassMapRenderEveryFrame);
+					if (!bDebugThisComponent)
+					{
+						ALandscapeProxy* Proxy = Component->GetLandscapeProxy();
+						if (Proxy && Proxy->GetName().Contains(GGrassMapRenderEveryFrame))
+						{
+							bDebugThisComponent = true;
+						}
+					}
+
+					if (bDebugThisComponent)
+					{
+						UE_LOG(LogGrass, Display, TEXT("Issuing debug render for the grassmap on component'%s'"), *State->Component->GetName());
+						DebugComponents.Add(Component);;
+					}
+				}
+			}
+
+			if (DebugComponents.Num() > 0)
+			{
+				DebugRenderComponents(DebugComponents);
+			}
+		}
+	}
+#endif // !UE_BUILD_SHIPPING
+
 	int32 AmortizedMaxStreaming = GGrassMapMaxComponentsStreaming;
 	int32 AmortizedMaxRendering = GGrassMapMaxComponentsRendering;
 
@@ -862,7 +929,8 @@ void FLandscapeGrassMapsBuilder::AmortizedUpdateGrassMaps(
 	const bool bEvictWhenBeyondEvictionRange = (GGrassMapUseRuntimeGeneration != 0);
 #endif // WITH_EDITOR
 	GGrassMapEvictAll = 0;
-	UpdateTrackedComponents(Cameras, AmortizedMaxRendering, GGrassMapMaxDiscardChecksPerFrame, bCancelAndEvictAllImmediately, bEvictWhenBeyondEvictionRange);
+	const int32 DeltaTicks = 1;
+	UpdateTrackedComponents(Cameras, AmortizedMaxRendering, GGrassMapMaxDiscardChecksPerFrame, bCancelAndEvictAllImmediately, bEvictWhenBeyondEvictionRange, DeltaTicks);
 
 	// no point in looking to start new grass map generation if nothing is pending, if grass is disabled or there are no cameras
 	if (bAllowStartGrassMapGeneration && PendingCount > 0 && GGrassEnable && Cameras.Num() > 0)
@@ -875,6 +943,152 @@ void FLandscapeGrassMapsBuilder::AmortizedUpdateGrassMaps(
 	}
 }
 
+#if !UE_BUILD_SHIPPING
+void FLandscapeGrassMapsBuilder::DebugRenderComponents(
+	TArrayView<TObjectPtr<ULandscapeComponent>> LandscapeComponents)
+{
+	TRACE_CPUPROFILER_EVENT_SCOPE(FLandscapeGrassMapsBuilder::DebugRenderComponents);
+	const int32 MaxStreamingRendering = FMath::Max(GGrassMapMaxComponentsForBlockingUpdate, 1);
+
+	// the goal of this function is to get every input Landscape Component to transition from the TextureStreaming to the Rendering state
+	// (i.e. issue GPU render commands for those components)
+
+	// the structure of this code is similar to BuildGrassMapsNow, except:
+	//    1) we don't need to complete the full process, we just need to ensure we trigger a GPU render of each LandscapeComponent
+	//    2) we need to evict any component that is already past the TextureStreaming->Rendering transition
+	//    3) we should skip fast paths, as we don't want to use pre-existing data, we want to render from scratch
+
+	if (LandscapeComponents.IsEmpty())
+	{
+		return;
+	}
+
+	if (!CanCurrentlyRender())
+	{
+		return;
+	}
+
+	// using an empty cameras array causes all distance checks to return 0, so grass maps won't be evicted
+	const TArray<FVector> EmptyCamerasArray;
+
+	TSet<ULandscapeComponent*> PreRenderComponents;			// components that have been in a pre-Rendering state inside this function (to know we did the transition)
+	TSet<ULandscapeComponent*> RenderedComponents;			// components that we have issued GPU renders for inside this function
+	TSet<ULandscapeComponent*> FailedComponents;			// components that have failed to render (for any reason)
+	TSet<ULandscapeComponent*> RemainingComponents;			// components that are not yet reached Rendered or Failed state, and still need to be checked
+	TArray<ULandscapeComponent*> ComponentsCompletedThisIteration;
+
+	RemainingComponents.Append(LandscapeComponents);
+	
+	while (RemainingComponents.Num() > 0)					// RenderedComponents.Num() + FailedComponents.Num() < LandscapeComponents.Num())
+	{
+		int32 AvailableStreamingSlots = MaxStreamingRendering - StreamingCount; // here we don't limit by overall population count
+		for (ULandscapeComponent* Component : RemainingComponents)
+		{
+			FComponentState* State = ComponentStates.FindRef(Component);
+			if (State == nullptr)
+			{
+				FailedComponents.Add(Component);
+				ComponentsCompletedThisIteration.Add(Component);
+				continue;
+			}
+			check(State->Component != nullptr);
+
+			switch(State->Stage)
+			{
+				case EComponentStage::Pending:
+					{
+						if (AvailableStreamingSlots > 0)
+						{
+							constexpr bool bForceCompileShaders = true;
+							constexpr bool bTryFastPaths = false;		// we want to force a render, so skip the fast paths
+							PreRenderComponents.Add(Component);
+							StartGrassMapGeneration(*State, bForceCompileShaders, bTryFastPaths);
+							if (State->Stage != EComponentStage::NotReady &&
+								State->Stage != EComponentStage::Pending)
+							{
+								AvailableStreamingSlots--;
+							}
+						}
+					}
+					break;
+
+				case EComponentStage::NotReady:
+					{
+						// if it's not ready because of shader reasons
+						if (!FailedComponents.Contains(Component) && !UE::Landscape::CanRenderGrassMap(Component))
+						{
+	#if WITH_EDITOR
+							// in editor, try to force compilation to complete
+							UE::Landscape::CompileGrassMapShader(State->Component);
+							if (!UE::Landscape::CanRenderGrassMap(Component))
+	#endif // WITH_EDITOR
+							{
+								// failed to compile shaders... we won't be able to build the grass map for this component
+								FailedComponents.Add(Component);
+								ComponentsCompletedThisIteration.Add(Component);
+							}
+						}
+					}
+					break;
+
+				case EComponentStage::TextureStreaming:
+					PreRenderComponents.Add(Component);
+					break;
+
+				case EComponentStage::Rendering:
+				case EComponentStage::AsyncFetch:
+				case EComponentStage::GrassMapsPopulated:
+					if (!PreRenderComponents.Contains(Component))
+					{
+						// this was not yet in the streaming state in this function, which means the render happened outside of this function.
+						// evict to restart the process
+						constexpr bool bCancelImmediately = true;
+						CancelAndEvict(*State, bCancelImmediately);
+						PreRenderComponents.Add(Component);
+					}
+					else
+					{
+						if (State->Stage == EComponentStage::Rendering)
+						{
+							// yay we issued the render commands for this component
+							RenderedComponents.Add(Component);
+							ComponentsCompletedThisIteration.Add(Component);
+						}
+					}
+					break;
+			}			
+		}
+
+		for (ULandscapeComponent* Component : ComponentsCompletedThisIteration)
+		{
+			RemainingComponents.Remove(Component);
+		}
+		ComponentsCompletedThisIteration.Reset();
+
+		if (RemainingComponents.Num() <= 0)
+		{
+			break;
+		}
+
+		// update all components that are tracked (without evicting)
+		const int32 UpdateAllComponentCount = ComponentStates.Num();
+		const bool bCancelAndEvictAllImmediately = false;
+		const bool bEvictWhenBeyondEvictionRange = false;
+		const int32 DeltaTicks = 0;
+		bool bChanged = UpdateTrackedComponents(EmptyCamerasArray, MaxStreamingRendering, UpdateAllComponentCount, bCancelAndEvictAllImmediately, bEvictWhenBeyondEvictionRange, DeltaTicks);
+
+		// no need to tick SubmitGPUCommands or wait for AsyncFetch tasks to complete
+		// because we don't need to block on them (we can let the async stuff just pile up til we exit this function)
+
+		// If any streaming is in flight, we do need to ensure that completes, via a blocking texture streaming update.
+		// TODO [chris.tchou] : ideally this would be a non-blocking streaming update tick, so we can react to other updates finishing
+		if (StreamingCount > 0)
+		{
+			TextureStreamingManager.WaitForTextureStreaming();
+		}
+	}
+}
+#endif // !UE_BUILD_SHIPPING
 
 bool FLandscapeGrassMapsBuilder::BuildGrassMapsNowForComponents(
 	TArrayView<TObjectPtr<ULandscapeComponent>> LandscapeComponents, FScopedSlowTask* SlowTask, bool bMarkDirty)
@@ -923,10 +1137,10 @@ bool FLandscapeGrassMapsBuilder::BuildGrassMapsNowForComponents(
 
 		// update all components that are tracked (without evicting)
 		const int32 UpdateAllComponentCount = ComponentStates.Num();
-
 		const bool bCancelAndEvictAllImmediately = false;
 		const bool bEvictWhenBeyondEvictionRange = false;
-		bool bChanged = UpdateTrackedComponents(EmptyCamerasArray, MaxStreamingRendering, UpdateAllComponentCount, bCancelAndEvictAllImmediately, bEvictWhenBeyondEvictionRange);
+		const int32 DeltaTicks = 1; // we need a positive delta tick to force completion of GPU readbacks on D3D11
+		bool bChanged = UpdateTrackedComponents(EmptyCamerasArray, MaxStreamingRendering, UpdateAllComponentCount, bCancelAndEvictAllImmediately, bEvictWhenBeyondEvictionRange, DeltaTicks);
 
 		UpToDateCount = 0;
 		int32 AvailableStreamingSlots = MaxStreamingRendering - StreamingCount; // here we don't limit by overall population count
@@ -945,7 +1159,8 @@ bool FLandscapeGrassMapsBuilder::BuildGrassMapsNowForComponents(
 				constexpr bool bForceCompileShaders = true;
 				if (AvailableStreamingSlots > 0)
 				{
-					StartGrassMapGeneration(*State, bForceCompileShaders);
+					constexpr bool bTryFastPaths = true;
+					StartGrassMapGeneration(*State, bForceCompileShaders, bTryFastPaths);
 					if (State->Stage != EComponentStage::NotReady &&
 						State->Stage != EComponentStage::Pending)
 					{
@@ -959,6 +1174,7 @@ bool FLandscapeGrassMapsBuilder::BuildGrassMapsNowForComponents(
 					}
 				}
 			}
+
 			if (State->Stage == EComponentStage::GrassMapsPopulated)
 			{
 				check(Component->GrassData->HasValidData()); // guaranteed by UpdateTrackedComponents(), as long as we update all of the components
@@ -992,6 +1208,7 @@ bool FLandscapeGrassMapsBuilder::BuildGrassMapsNowForComponents(
 				}
 				UpToDateCount++;
 			}
+
 			if (State->Stage == EComponentStage::NotReady)
 			{
 				// if it's not ready because of shader reasons
@@ -1152,7 +1369,7 @@ bool FLandscapeGrassMapsBuilder::CancelAndEvict(FComponentState& State, bool bCa
 					return false;
 				}
 
-				State.ActiveRender->FreeAsyncReadback();
+				State.ActiveRender->FreeAsyncReadback(); // TODO [chris.tchou] : I think this is redundant with Reset() below
 				State.ActiveRender.Reset();
 				State.AsyncFetchTask.Reset();
 			}
@@ -1258,14 +1475,14 @@ bool FLandscapeGrassMapsBuilder::TryFastpathsFromPending(FComponentState& State,
 	return false;
 }
 
-bool FLandscapeGrassMapsBuilder::StartGrassMapGeneration(FComponentState& State, bool bForceCompileShaders)
+bool FLandscapeGrassMapsBuilder::StartGrassMapGeneration(FComponentState& State, bool bForceCompileShaders, bool bTryFastpaths)
 {
 	TRACE_CPUPROFILER_EVENT_SCOPE(FLandscapeGrassMapsBuilder::StartGrassMapGeneration);
 	check(State.Stage == EComponentStage::Pending);
 
 	ULandscapeComponent* Component = State.Component;
 
-	if (TryFastpathsFromPending(State, /* bRecalculateHashes = */ true))
+	if (bTryFastpaths && TryFastpathsFromPending(State, /* bRecalculateHashes = */ true))
 	{
 		return false;
 	}
diff --git a/Engine/Source/Runtime/Landscape/Private/LandscapeGrassMapsBuilder.h b/Engine/Source/Runtime/Landscape/Private/LandscapeGrassMapsBuilder.h
--- a/Engine/Source/Runtime/Landscape/Private/LandscapeGrassMapsBuilder.h
+++ b/Engine/Source/Runtime/Landscape/Private/LandscapeGrassMapsBuilder.h
@@ -77,6 +77,12 @@ public:
 	// returns true if it successfully makes all LandscapeComponents have up to date Grass Maps.
 	bool BuildGrassMapsNowForComponents(TArrayView<TObjectPtr<ULandscapeComponent>> LandscapeComponents, FScopedSlowTask* SlowTask, bool bMarkDirty);
 
+#if !UE_BUILD_SHIPPING
+	// Force stream textures and issue draw calls for the specified components -- mainly for capturing the rendering behavior and debugging issues
+	// (this will evict existing populated grass maps on those components)
+	void DebugRenderComponents(TArrayView<TObjectPtr<ULandscapeComponent>> LandscapeComponents);
+#endif // !UE_BUILD_SHIPPING
+
 private:
 
 	// false if this program instance will never be able to render grass
@@ -90,7 +96,8 @@ private:
 	// MaxExpensiveUpdateChecksToPerform controls how many expensive component updates are performed (pass ComponentStates.Num() to process all)
 	// bCancelAndEvictAllImmediately will move all tracked component states to the pending state and block until any in-flight processing is cancelled and cleaned up
 	// bEvictWhenBeyondEvictionRange will evict any populated components that are beyond the eviction range distance (used to reclaim memory when using runtime generation)
-	bool UpdateTrackedComponents(const TArray<FVector>& Cameras, int32 LocalMaxRendering, int32 MaxExpensiveUpdateChecksToPerform, bool bCancelAndEvictAllImmediately, bool bEvictWhenBeyondEvictionRange);
+	// DeltaTicks specifies how many ticks to advance the tick counter on each tracked component.  In general 1 tick, though in some cases where we are iterating within a single frame, and don't need to trigger delayed-tick events (i.e. D3D11 readback), we can pass 0.
+	bool UpdateTrackedComponents(const TArray<FVector>& Cameras, int32 LocalMaxRendering, int32 MaxExpensiveUpdateChecksToPerform, bool bCancelAndEvictAllImmediately, bool bEvictWhenBeyondEvictionRange, int32 DeltaTicks);
 
 	// Start the grass map generation process on pending components in priority order
 	// (based on distance from the given Camera set) -- Cameras must not be empty.
@@ -259,7 +266,7 @@ private:
 	bool CancelAndEvict(FComponentState& State, bool bCancelImmediately);
 
 	// try to kick off the grass map generation pipeline -- returns true if it started the amortized update path, false otherwise.
-	bool StartGrassMapGeneration(FComponentState& State, bool bForceCompileShaders);
+	bool StartGrassMapGeneration(FComponentState& State, bool bForceCompileShaders, bool bTryFastPaths = true);
 
 	// try to apply fast path transitions to a pending component -- returns true if a fastpath was taken, and the component is no longer pending.
 	bool TryFastpathsFromPending(FComponentState& State, bool bRecalculateHashes);
diff --git a/Engine/Source/Runtime/Landscape/Private/LandscapeGrassWeightExporter.cpp b/Engine/Source/Runtime/Landscape/Private/LandscapeGrassWeightExporter.cpp
--- a/Engine/Source/Runtime/Landscape/Private/LandscapeGrassWeightExporter.cpp
+++ b/Engine/Source/Runtime/Landscape/Private/LandscapeGrassWeightExporter.cpp
@@ -402,23 +402,28 @@ FLandscapeGrassWeightExporter_RenderThread::FLandscapeGrassWeightExporter_Render
 	if (bInReadbackToCPU)
 	{
 		// even when doing a synchronous readback, we use the async readback structure
-		AsyncReadbackPtr = new FLandscapeAsyncTextureReadback();
+		GameThreadAsyncReadbackPtr = new FLandscapeAsyncTextureReadback();
 	}
 }
 
 FLandscapeGrassWeightExporter_RenderThread::~FLandscapeGrassWeightExporter_RenderThread()
 {
-	if (AsyncReadbackPtr != nullptr)
+	if (GameThreadAsyncReadbackPtr != nullptr)
 	{
-		AsyncReadbackPtr->QueueDeletionFromGameThread();
-		AsyncReadbackPtr = nullptr;
+		GameThreadAsyncReadbackPtr->QueueDeletionFromGameThread();
+		GameThreadAsyncReadbackPtr = nullptr;
 	}
 }
 
-void FLandscapeGrassWeightExporter_RenderThread::RenderLandscapeComponentToTexture_RenderThread(FRHICommandListImmediate& RHICmdList)
+DECLARE_GPU_STAT(LandscapeGrassMaps);
+
+void FLandscapeGrassWeightExporter_RenderThread::RenderLandscapeComponentToTexture_RenderThread(FRHICommandListImmediate& RHICmdList, FLandscapeAsyncTextureReadback* AsyncReadbackPtr)
 {
 	FRDGBuilder GraphBuilder(RHICmdList);
 
+	RHI_BREADCRUMB_EVENT_STAT(RHICmdList, LandscapeGrassMaps, "LandscapeGrassMaps");
+	SCOPED_GPU_STAT(RHICmdList, LandscapeGrassMaps);
+
 	FRDGTextureDesc TextureDesc = FRDGTextureDesc::Create2D(
 		TargetSize,
 		PF_B8G8R8A8,
@@ -429,10 +434,8 @@ void FLandscapeGrassWeightExporter_RenderThread::RenderLandscapeComponentToTextu
 
 	RenderLandscapeComponentToTexture_RenderThread(GraphBuilder, OutputTexture);
 
-	if (AsyncReadbackPtr != nullptr)
-	{
-		AsyncReadbackPtr->StartReadback_RenderThread(GraphBuilder, OutputTexture);
-	}
+	check(AsyncReadbackPtr != nullptr);
+	AsyncReadbackPtr->StartReadback_RenderThread(GraphBuilder, OutputTexture);
 
 	GraphBuilder.Execute();
 }
@@ -557,9 +560,9 @@ FLandscapeGrassWeightExporter::FLandscapeGrassWeightExporter(ALandscapeProxy* In
 		// render
 		FLandscapeGrassWeightExporter_RenderThread* Exporter = this;
 		ENQUEUE_RENDER_COMMAND(FDrawSceneCommand)(
-			[Exporter](FRHICommandListImmediate& RHICmdList)
+			[Exporter, AsyncReadbackPtr = GameThreadAsyncReadbackPtr](FRHICommandListImmediate& RHICmdList)
 			{
-				Exporter->RenderLandscapeComponentToTexture_RenderThread(RHICmdList);
+				Exporter->RenderLandscapeComponentToTexture_RenderThread(RHICmdList, AsyncReadbackPtr);
 				RHICmdList.ImmediateFlush(EImmediateFlushType::FlushRHIThreadFlushResources);
 			});
 	}
@@ -567,14 +570,14 @@ FLandscapeGrassWeightExporter::FLandscapeGrassWeightExporter(ALandscapeProxy* In
 
 bool FLandscapeGrassWeightExporter::CheckAndUpdateAsyncReadback(bool& bOutRenderCommandsQueued, const bool bInForceFinish)
 {
-	check(AsyncReadbackPtr != nullptr);
-	return AsyncReadbackPtr->CheckAndUpdate(bOutRenderCommandsQueued, bInForceFinish);
+	check(GameThreadAsyncReadbackPtr != nullptr);
+	return GameThreadAsyncReadbackPtr->CheckAndUpdate(bOutRenderCommandsQueued, bInForceFinish);
 }
 
 bool FLandscapeGrassWeightExporter::IsAsyncReadbackComplete()
 {
-	check(AsyncReadbackPtr != nullptr);
-	return AsyncReadbackPtr->IsComplete();
+	check(GameThreadAsyncReadbackPtr != nullptr);
+	return GameThreadAsyncReadbackPtr->IsComplete();
 }
 
 struct FByteBuffer2DView : public IBuffer2DView<uint8>
@@ -678,25 +681,24 @@ struct FHeightBuffer2DView : IBuffer2DView<uint16>
 
 void FLandscapeGrassWeightExporter::FreeAsyncReadback()
 {
-	check(AsyncReadbackPtr != nullptr);
-	AsyncReadbackPtr->QueueDeletionFromGameThread();
-	AsyncReadbackPtr = nullptr;
+	check(GameThreadAsyncReadbackPtr != nullptr);
+	GameThreadAsyncReadbackPtr->QueueDeletionFromGameThread();
+	GameThreadAsyncReadbackPtr = nullptr;
 }
 
 void FLandscapeGrassWeightExporter::CancelAndSelfDestruct()
 {
-	check(AsyncReadbackPtr != nullptr);
+	check(GameThreadAsyncReadbackPtr != nullptr);
 
 	// Cancel the readback, and queue destruction on the render thread
-	AsyncReadbackPtr->CancelAndSelfDestruct();
-	AsyncReadbackPtr = nullptr;
+	GameThreadAsyncReadbackPtr->CancelAndSelfDestruct();
+	GameThreadAsyncReadbackPtr = nullptr;
 
 	// Queue destruction of FLandscapeGrassWeightExporter, also on the render thread
 	FLandscapeGrassWeightExporter* Exporter = this;
 	ENQUEUE_RENDER_COMMAND(FCancelAndDestructCommand)(
 		[Exporter](FRHICommandListImmediate& RHICmdList)
 		{
-			check(Exporter->AsyncReadbackPtr == nullptr);
 			delete Exporter;
 		});
 }
@@ -707,11 +709,11 @@ TMap<ULandscapeComponent*, TUniquePtr<FLandscapeComponentGrassData>, TInlineSetA
 	TMap<ULandscapeComponent*, TUniquePtr<FLandscapeComponentGrassData>, TInlineSetAllocator<1>> Results;
 	TArray<FColor> Samples;
 
-	check(AsyncReadbackPtr != nullptr);
+	check(GameThreadAsyncReadbackPtr != nullptr);
 
 	{
 		FIntPoint Size;
-		Samples = AsyncReadbackPtr->TakeResults(&Size);
+		Samples = GameThreadAsyncReadbackPtr->TakeResults(&Size);
 		if (bFreeAsyncReadback)
 		{
 			FreeAsyncReadback();
