Fix landscape flatten tool conflict with undo.  Avoid processing edge fixups while batchmerge readbacks are pending to avoid stomping on pending merge evaluations.

diff --git a/Engine/Source/Runtime/Landscape/Private/LandscapeGroup.cpp b/Engine/Source/Runtime/Landscape/Private/LandscapeGroup.cpp
--- a/Engine/Source/Runtime/Landscape/Private/LandscapeGroup.cpp
+++ b/Engine/Source/Runtime/Landscape/Private/LandscapeGroup.cpp
@@ -9,6 +9,7 @@
 #include "LandscapeEdgeFixup.h"
 #include "LandscapePrivate.h"
 #include "LandscapeSubsystem.h"
+#include "LandscapeEditReadback.h"
 
 // enables verbose debug spew
 //#define ENABLE_LANDSCAPE_EDGE_FIXUP_DEBUG_SPEW 1
@@ -791,55 +792,65 @@ void FLandscapeGroup::TickEdgeFixup(ULandscapeSubsystem* LandscapeSubsystem, boo
 	{
 		TRACE_CPUPROFILER_EVENT_SCOPE(HeightmapsNeedingEdgeTexturePatching);
 
-		for (auto It = HeightmapsNeedingEdgeTexturePatching.CreateIterator(); It; ++It)
+
+#if WITH_EDITORONLY_DATA
+		if (FLandscapeEditLayerReadback::HasWork())
+		{
+			GROUP_DEBUG_LOG_DETAIL(TEXT("-- Patch (All Components) -------- PAUSED (deferred while any readbacks pending)"));
+		}
+		else
+#endif
 		{
-			TRACE_CPUPROFILER_EVENT_SCOPE(CheckHeightmap);
+			for (auto It = HeightmapsNeedingEdgeTexturePatching.CreateIterator(); It; ++It)
+			{
+				TRACE_CPUPROFILER_EVENT_SCOPE(CheckHeightmap);
 
-			ULandscapeHeightmapTextureEdgeFixup* EdgeFixup = *It;
-			check(EdgeFixup->ActiveGroup == this); // should never be in HeightmapsNeedingEdgeTexturePatching unless active in this group
-			check(AllRegisteredFixups.Contains(EdgeFixup));
-			check(EdgeFixup->bMapped);
+				ULandscapeHeightmapTextureEdgeFixup* EdgeFixup = *It;
+				check(EdgeFixup->ActiveGroup == this); // should never be in HeightmapsNeedingEdgeTexturePatching unless active in this group
+				check(AllRegisteredFixups.Contains(EdgeFixup));
+				check(EdgeFixup->bMapped);
 
 #if WITH_EDITORONLY_DATA
-			ALandscape* Landscape = EdgeFixup->ActiveComponent->GetLandscapeActor();
-			if (!Landscape || !Landscape->bGrassUpdateEnabled)
-			{
-				GROUP_DEBUG_LOG_DETAIL(TEXT("-- Patch Component %p (%d,%d) -------- PAUSED (editing or landscape unregistered)"), EdgeFixup->ActiveComponent.Get(), EdgeFixup->GroupCoord.X, EdgeFixup->GroupCoord.Y);
-				continue;
-			}
+				ALandscape* Landscape = EdgeFixup->ActiveComponent->GetLandscapeActor();
+				if (!Landscape || !Landscape->bGrassUpdateEnabled)
+				{
+					GROUP_DEBUG_LOG_DETAIL(TEXT("-- Patch Component %p (%d,%d) -------- PAUSED (editing or landscape unregistered)"), EdgeFixup->ActiveComponent.Get(), EdgeFixup->GroupCoord.X, EdgeFixup->GroupCoord.Y);
+					continue;
+				}
 #endif // WITH_EDITORONLY_DATA
 
-			if (EdgeFixup->HeightmapTexture->HasPendingInitOrStreaming())
-			{
-				GROUP_DEBUG_LOG_DETAIL(TEXT("-- Patch Component %p (%d,%d) -------- PAUSED (pending)"), EdgeFixup->ActiveComponent.Get(), EdgeFixup->GroupCoord.X, EdgeFixup->GroupCoord.Y);
-				continue;
-			}
+				if (EdgeFixup->HeightmapTexture->HasPendingInitOrStreaming())
+				{
+					GROUP_DEBUG_LOG_DETAIL(TEXT("-- Patch Component %p (%d,%d) -------- PAUSED (pending)"), EdgeFixup->ActiveComponent.Get(), EdgeFixup->GroupCoord.X, EdgeFixup->GroupCoord.Y);
+					continue;
+				}
 
 #if WITH_EDITOR
-			if (EdgeFixup->IsTextureEdgePatchingPaused())
-			{
-				GROUP_DEBUG_LOG_DETAIL(TEXT("-- Patch Component %p (%d,%d) -------- PAUSED (readback)"), EdgeFixup->ActiveComponent.Get(), EdgeFixup->GroupCoord.X, EdgeFixup->GroupCoord.Y);
-				continue;
-			}
+				if (EdgeFixup->IsTextureEdgePatchingPaused())
+				{
+					GROUP_DEBUG_LOG_DETAIL(TEXT("-- Patch Component %p (%d,%d) -------- PAUSED (readback)"), EdgeFixup->ActiveComponent.Get(), EdgeFixup->GroupCoord.X, EdgeFixup->GroupCoord.Y);
+					continue;
+				}
 
-			// wait til we are not using the default texture to apply edge patching
-			if (EdgeFixup->GetHeightmapTexture()->IsDefaultTexture())
-			{
-				GROUP_DEBUG_LOG_DETAIL(TEXT("-- Patch Component %p (%d,%d) -------- WAITING (texture load)"), EdgeFixup->ActiveComponent.Get(), EdgeFixup->GroupCoord.X, EdgeFixup->GroupCoord.Y);
-				continue;
-			}
+				// wait til we are not using the default texture to apply edge patching
+				if (EdgeFixup->GetHeightmapTexture()->IsDefaultTexture())
+				{
+					GROUP_DEBUG_LOG_DETAIL(TEXT("-- Patch Component %p (%d,%d) -------- WAITING (texture load)"), EdgeFixup->ActiveComponent.Get(), EdgeFixup->GroupCoord.X, EdgeFixup->GroupCoord.Y);
+					continue;
+				}
 #endif // WITH_EDITOR
 
-			check(EdgeFixup->ActiveGroup == this);
-			int32 PatchedEdgeCount = EdgeFixup->CheckAndPatchTextureEdgesFromEdgeSnapshots();
-			if (PatchedEdgeCount > 0)
-			{
-				TexturesPatched++;
-			}
-			EdgesPatched += PatchedEdgeCount;
+				check(EdgeFixup->ActiveGroup == this);
+				int32 PatchedEdgeCount = EdgeFixup->CheckAndPatchTextureEdgesFromEdgeSnapshots();
+				if (PatchedEdgeCount > 0)
+				{
+					TexturesPatched++;
+				}
+				EdgesPatched += PatchedEdgeCount;
 
-			It.RemoveCurrent();
-			GROUP_DEBUG_LOG_DETAIL(TEXT("-- Patch Component %p (%d,%d) -------- PATCHED %d edges"), EdgeFixup->ActiveComponent.Get(), EdgeFixup->GroupCoord.X, EdgeFixup->GroupCoord.Y, PatchedEdgeCount);
+				It.RemoveCurrent();
+				GROUP_DEBUG_LOG_DETAIL(TEXT("-- Patch Component %p (%d,%d) -------- PATCHED %d edges"), EdgeFixup->ActiveComponent.Get(), EdgeFixup->GroupCoord.X, EdgeFixup->GroupCoord.Y, PatchedEdgeCount);
+			}
 		}
 	}
 
