Exposed landscape component's GetGrassTypes to BP

diff --git a/Engine/Source/Runtime/Landscape/Classes/LandscapeComponent.h b/Engine/Source/Runtime/Landscape/Classes/LandscapeComponent.h
--- a/Engine/Source/Runtime/Landscape/Classes/LandscapeComponent.h
+++ b/Engine/Source/Runtime/Landscape/Classes/LandscapeComponent.h
@@ -875,6 +875,9 @@ public:
 	/** Gets the landscape info object for this landscape */
 	LANDSCAPE_API ULandscapeInfo* GetLandscapeInfo() const;
 
+	UFUNCTION(BlueprintCallable, Category = "Landscape|Grass", meta=(DisplayName = "GetGrassTypes", Tooltip = "Returns the grass types used by the landscape material on this component"))
+	const TArray<ULandscapeGrassType*>& GetGrassTypesBP() const { return ObjectPtrDecay(GrassTypes); }
+
 	/** Returns the array of grass types used by the landscape material. Call UpdateGrassTypes first to ensure this array is up to date. */
 	const TArray<TObjectPtr<ULandscapeGrassType>>& GetGrassTypes() const { return GrassTypes; }
 
