Fixed landscape legacy weight blending trying to blend layers that were not rendered in a given batch

diff --git a/Engine/Shaders/Private/Landscape/LandscapeEditLayersWeightmaps.usf b/Engine/Shaders/Private/Landscape/LandscapeEditLayersWeightmaps.usf
--- a/Engine/Shaders/Private/Landscape/LandscapeEditLayersWeightmaps.usf
+++ b/Engine/Shaders/Private/Landscape/LandscapeEditLayersWeightmaps.usf
@@ -25,6 +25,7 @@
 /** EWeightmapTargetLayerFlags enum */
 #define EWEIGHTMAPTARGETLAYERFLAGS_ISVISIBILITYLAYER (1 << 0)
 #define EWEIGHTMAPTARGETLAYERFLAGS_ISWEIGHTBLENDED (1 << 1)
+#define EWEIGHTMAPTARGETLAYERFLAGS_SKIP (1 << 2)
 
 #if MERGE_EDIT_LAYER
 // MergeEditLayerPS inputs/outputs :
@@ -81,9 +82,12 @@ Texture2D<float4> InSourceWeightmap; // Source weightmap (containing the current
 #if PERFORM_LEGACY_WEIGHT_BLENDING
 bool IsWeightBlendedTargetLayer(FWeightmapTargetLayerInfo InTargetLayerInfo)
 {
-	// For the visibility layer, deactivate weight blending altogether : 
-	return (((InTargetLayerInfo.Flags & EWEIGHTMAPTARGETLAYERFLAGS_ISVISIBILITYLAYER) == 0)
-		&& (InTargetLayerInfo.Flags & EWEIGHTMAPTARGETLAYERFLAGS_ISWEIGHTBLENDED));
+	
+	return (InTargetLayerInfo.Flags & EWEIGHTMAPTARGETLAYERFLAGS_ISWEIGHTBLENDED)
+		// For the visibility layer, deactivate weight blending altogether : 
+		&& ((InTargetLayerInfo.Flags & EWEIGHTMAPTARGETLAYERFLAGS_ISVISIBILITYLAYER) == 0)
+		// Some layers are absent from the batch, they should be skipped from blending too :
+		&& ((InTargetLayerInfo.Flags & EWEIGHTMAPTARGETLAYERFLAGS_SKIP) == 0); 
 }
 #endif // PERFORM_LEGACY_WEIGHT_BLENDING
 
diff --git a/Engine/Source/Runtime/Landscape/Private/LandscapeEditLayers.cpp b/Engine/Source/Runtime/Landscape/Private/LandscapeEditLayers.cpp
--- a/Engine/Source/Runtime/Landscape/Private/LandscapeEditLayers.cpp
+++ b/Engine/Source/Runtime/Landscape/Private/LandscapeEditLayers.cpp
@@ -6644,6 +6644,9 @@ bool ULandscapeWeightmapWeightBlendedLayersRenderer::RenderLayer(UE::Landscape::
 	checkf(RDGBuilderRecorder.IsRecording(), TEXT("ERenderFlags::RenderMode_Recorded means the command recorder should be recording at this point"));
 	check(!RenderParams.MergeRenderContext->IsHeightmapMerge());
 
+	const FMergeRenderBatch* RenderBatch = RenderParams.MergeRenderContext->GetCurrentRenderBatch();
+	check(RenderBatch != nullptr);
+
 	RenderParams.MergeRenderContext->CycleBlendRenderTargets(RDGBuilderRecorder);
 	ULandscapeScratchRenderTarget* WriteRT = RenderParams.MergeRenderContext->GetBlendRenderTargetWrite();
 	ULandscapeScratchRenderTarget* ReadRT = RenderParams.MergeRenderContext->GetBlendRenderTargetRead();
@@ -6653,17 +6656,26 @@ bool ULandscapeWeightmapWeightBlendedLayersRenderer::RenderLayer(UE::Landscape::
 
 	TArray<FWeightmapTargetLayerInfo> WeightmapTargetLayerInfos;
 	WeightmapTargetLayerInfos.Reserve(RenderParams.TargetLayerGroupLayerInfos.Num());
-	Algo::Transform(RenderParams.TargetLayerGroupLayerInfos, WeightmapTargetLayerInfos, [](ULandscapeLayerInfoObject* InLayerInfo)
+	Algo::Transform(RenderParams.TargetLayerGroupLayerInfos, WeightmapTargetLayerInfos, [RenderBatch, MergeRenderContext = RenderParams.MergeRenderContext](ULandscapeLayerInfoObject* InLayerInfo)
 		{
 			check(InLayerInfo != nullptr); // There should only be valid layer infos at this point
+
+			int32 TargetLayerIndex = MergeRenderContext->GetTargetLayerIndexForLayerInfoChecked(InLayerInfo);
+			// Skip weight-blended layers that are not rendered in the batch, because their texture slice might contain garbage :
+			const bool bRenderedInBatch = RenderBatch->TargetLayerBitIndices[TargetLayerIndex];
+
 			FWeightmapTargetLayerInfo WeightmapTargetLayerInfo;
+			if (!bRenderedInBatch)
+			{
+				WeightmapTargetLayerInfo.Flags |= EWeightmapTargetLayerFlags::Skip;
+			}
 			if (!InLayerInfo->bNoWeightBlend)
 			{
-				WeightmapTargetLayerInfo.Flags = EWeightmapTargetLayerFlags::IsWeightBlended;
+				WeightmapTargetLayerInfo.Flags |= EWeightmapTargetLayerFlags::IsWeightBlended;
 			}
 			if (InLayerInfo == ALandscapeProxy::VisibilityLayer)
 			{
-				WeightmapTargetLayerInfo.Flags = EWeightmapTargetLayerFlags::IsVisibilityLayer;
+				WeightmapTargetLayerInfo.Flags |= EWeightmapTargetLayerFlags::IsVisibilityLayer;
 			}
 			return WeightmapTargetLayerInfo;
 		});
@@ -6686,16 +6698,19 @@ bool ULandscapeWeightmapWeightBlendedLayersRenderer::RenderLayer(UE::Landscape::
 			const int32 NumTargetLayers = TargetLayerNames.Num();
 			for (int32 TargetLayerIndex = 0; TargetLayerIndex < NumTargetLayers; ++TargetLayerIndex)
 			{
-				FLandscapeEditLayersWeightmapsPerformLegacyWeightBlendingPS::FParameters* PSParams = GraphBuilder.AllocParameters<FLandscapeEditLayersWeightmapsPerformLegacyWeightBlendingPS::FParameters>();
-				PSParams->RenderTargets[0] = FRenderTargetBinding(OutputTextureRef, ERenderTargetLoadAction::ENoAction, /*InMipIndex = */0, /*InArraySlice = */TargetLayerIndex);
-				PSParams->InTargetLayerIndex = TargetLayerIndex;
-				PSParams->InNumTargetLayers = NumTargetLayers;
-				PSParams->InTargetLayerInfos = TargetLayerInfosBufferSRV;
-				PSParams->InCurrentEditLayerWeightmaps = CurrentEditLayerTextureSRVRef;
+				if (!EnumHasAnyFlags(WeightmapTargetLayerInfos[TargetLayerIndex].Flags, EWeightmapTargetLayerFlags::Skip))
+				{ 
+					FLandscapeEditLayersWeightmapsPerformLegacyWeightBlendingPS::FParameters* PSParams = GraphBuilder.AllocParameters<FLandscapeEditLayersWeightmapsPerformLegacyWeightBlendingPS::FParameters>();
+					PSParams->RenderTargets[0] = FRenderTargetBinding(OutputTextureRef, ERenderTargetLoadAction::ENoAction, /*InMipIndex = */0, /*InArraySlice = */TargetLayerIndex);
+					PSParams->InTargetLayerIndex = TargetLayerIndex;
+					PSParams->InNumTargetLayers = NumTargetLayers;
+					PSParams->InTargetLayerInfos = TargetLayerInfosBufferSRV;
+					PSParams->InCurrentEditLayerWeightmaps = CurrentEditLayerTextureSRVRef;
 
-				FLandscapeEditLayersWeightmapsPerformLegacyWeightBlendingPS::PerformLegacyWeightBlendingPS(
-					RDG_EVENT_NAME("LegacyWeightBlending(%s) -> %s", *TargetLayerNames[TargetLayerIndex].ToString(), *OutputResourceName),
-					GraphBuilder, PSParams, EffectiveTextureSize);
+					FLandscapeEditLayersWeightmapsPerformLegacyWeightBlendingPS::PerformLegacyWeightBlendingPS(
+						RDG_EVENT_NAME("LegacyWeightBlending(%s) -> %s", *TargetLayerNames[TargetLayerIndex].ToString(), *OutputResourceName),
+						GraphBuilder, PSParams, EffectiveTextureSize);
+				}
 			}
 		};
 
@@ -11072,7 +11087,7 @@ int32 ALandscape::PerformLayersWeightmapsBatchedMerge(FUpdateLayersContentContex
 				{
 					// If the component is present in this batch, we can resolve its channels :
 					if (RenderBatch->ComponentsToRender.Contains(RenderedComponent) 
-						// The component was possibly rendered in a previous batch, in which case we don't have to resolved it again : 
+						// The component was possibly rendered in a previous batch, in which case we don't have to resolve it again : 
 						&& !ResolvedLandscapeComponents.Contains(RenderedComponent))
 					{
 						TArray<FIntRect, TInlineAllocator<4>> SourceSubsectionRects, DummySubsectionRects;
diff --git a/Engine/Source/Runtime/Landscape/Public/LandscapeEditLayerTypes.h b/Engine/Source/Runtime/Landscape/Public/LandscapeEditLayerTypes.h
--- a/Engine/Source/Runtime/Landscape/Public/LandscapeEditLayerTypes.h
+++ b/Engine/Source/Runtime/Landscape/Public/LandscapeEditLayerTypes.h
@@ -72,6 +72,7 @@ enum class EWeightmapTargetLayerFlags : uint32
 {
 	IsVisibilityLayer = (1 << 0), // This target layer is the visibility layer
 	IsWeightBlended = (1 << 1), // Blend the target layer's value with all the other target layers weights
+	Skip = (1 << 2), // This layer should be skipped from blending because it has not been rendered in this batch
 
 	None = 0
 };
