Turned assert about landscape physical material tasks failing to complete on save into an error, since that is something that can happen if the material fails to compile

diff --git a/Engine/Source/Runtime/Landscape/Private/Landscape.cpp b/Engine/Source/Runtime/Landscape/Private/Landscape.cpp
--- a/Engine/Source/Runtime/Landscape/Private/Landscape.cpp
+++ b/Engine/Source/Runtime/Landscape/Private/Landscape.cpp
@@ -3964,9 +3964,16 @@ void ALandscapeProxy::PreSave(FObjectPreSaveContext ObjectSaveContext)
 			{
 				LandscapeComponent->PhysicalMaterialTask.Flush();
 			}
-			check(LandscapeComponent->PhysicalMaterialTask.IsComplete());
-			// IsInProgress tests the render thread status.  Now finish the last steps on the game thread.
-			LandscapeComponent->UpdatePhysicalMaterialTasks();
+
+			if (LandscapeComponent->PhysicalMaterialTask.IsComplete())
+			{
+				// IsInProgress tests the render thread status.  Now finish the last steps on the game thread.
+				LandscapeComponent->UpdatePhysicalMaterialTasks();
+			}
+			else
+			{
+				UE_LOG(LogLandscape, Error, TEXT("Landscape physical material failed to complete before saving. That could be due to the landscape material failing to compile on this component (%s)"), *LandscapeComponent->GetFullName());
+			}
 		}
 
 		// Ensure the component's cached bounds are correct
