Change landscape CollapseLayer tool to use selective render weightmaps code instead of intermediate render code.  Add weightmap support for edit layer selective renders.

diff --git a/Engine/Source/Runtime/Landscape/Classes/Landscape.h b/Engine/Source/Runtime/Landscape/Classes/Landscape.h
--- a/Engine/Source/Runtime/Landscape/Classes/Landscape.h
+++ b/Engine/Source/Runtime/Landscape/Classes/Landscape.h
@@ -246,30 +246,28 @@ struct FOnLandscapeEditLayersMergedParams
 	bool bIsHeightmapMerge = false;
 };
 
-
-// Parameters to describe a subset of a landscape heightmap to render.  Can select a subset of edit layers and a subregion of the landscape coordinates.
-struct FLandscapeEditLayerRenderHeightParams
+struct FLandscapeEditLayerRenderCommonParams
 {
 	// Region to render, in Landscape coordinates
-	FIntRect Bounds;  
+	FIntRect Bounds;
 
 	// Which edit layers to include.  Order matches ALandscape::GetEditLayers().
 	TBitArray<> ActiveEditLayers;
 
+	bool bRenderBrushes = true;
+};
+
+// Parameters to describe a subset of a landscape heightmap to render.  Can select a subset of edit layers and a subregion of the landscape coordinates.
+struct FLandscapeEditLayerRenderHeightParams : public FLandscapeEditLayerRenderCommonParams
+{
 	// Provide either a cpu buffer or render target to receive the data.  Buffer should be allocated large enough to receive the samples in Bounds.
 	TArrayView<uint16> CpuResult;
 	UTextureRenderTarget2D* RTResult = nullptr;
 };
 
 // Parameters to describe a subset of a landscape weightmaps to render.  Can select a subset of edit layers and a subregion of the landscape coordinates, for a list of weightmaps.
-struct FLandscapeEditLayerRenderWeightParams
+struct FLandscapeEditLayerRenderWeightParams : public FLandscapeEditLayerRenderCommonParams
 {
-	// Region to render, in Landscape coordinates
-	FIntRect Bounds;
-
-	// Which edit layers to include.  Order matches ALandscape::GetEditLayers().
-	TBitArray<> ActiveEditLayers;
-
 	TArrayView<FName> WeightmapNames;
 
 	// Provide either a cpu buffer or render target to receive the data.  Buffer should be allocated large enough to receive the samples in Bounds.
@@ -497,6 +495,7 @@ PRAGMA_ENABLE_DEPRECATION_WARNINGS
 	
 	LANDSCAPE_API bool DeleteLayer(int32 InLayerIndex);
 	LANDSCAPE_API void CollapseLayer(int32 InLayerIndex);
+	void CollapseLayerIntermediateRender(int32 InLayerIndex);
 	LANDSCAPE_API void DeleteLayers();
 	LANDSCAPE_API void SetEditingLayer(const FGuid& InLayerGuid = FGuid());
 	LANDSCAPE_API void SetGrassUpdateEnabled(bool bInGrassUpdateEnabled);
@@ -604,6 +603,7 @@ private:
 	int32 PerformLayersWeightmapsLocalMerge(FUpdateLayersContentContext& InUpdateLayersContentContext, const FEditLayersWeightmapMergeParams& InMergeParams);
 	int32 PerformLayersWeightmapsGlobalMerge(FUpdateLayersContentContext& InUpdateLayersContentContext, const FEditLayersWeightmapMergeParams& InMergeParams);
 	int32 PerformLayersWeightmapsBatchedMerge(FUpdateLayersContentContext& InUpdateLayersContentContext, const FEditLayersWeightmapMergeParams& InMergeParams);
+	bool PerformSelectiveRenderEditLayersWeightmapsCPU(FUpdateLayersContentContext& InUpdateLayersContentContext, const FLandscapeEditLayerRenderWeightParams& InSelectiveRenderParams);
 	void ResolveLayersWeightmapTexture(const FTextureToComponentHelper& MapHelper, const TSet<UTexture2D*>& WeightmapsToResolve, bool bIntermediateRender, TArray<FLandscapeEditLayerComponentReadbackResult>& InOutComponentReadbackResults);
 
 	bool ResolveLayersTexture(FTextureToComponentHelper const& MapHelper, FLandscapeEditLayerReadback* InCPUReadBack, UTexture2D* InOutputTexture, bool bIntermediateRender,
diff --git a/Engine/Source/Runtime/Landscape/Private/LandscapeEditLayerMergeRenderContext.cpp b/Engine/Source/Runtime/Landscape/Private/LandscapeEditLayerMergeRenderContext.cpp
--- a/Engine/Source/Runtime/Landscape/Private/LandscapeEditLayerMergeRenderContext.cpp
+++ b/Engine/Source/Runtime/Landscape/Private/LandscapeEditLayerMergeRenderContext.cpp
@@ -309,6 +309,17 @@ FIntRect FMergeRenderBatch::ComputeSectionRect(ULandscapeComponent* InComponent,
 	return FIntRect(SectionMin, SectionMax);
 }
 
+TArray<FIntRect> FMergeRenderBatch::ComputeAllComponentSectionRects(bool bInWithDuplicateBorders) const
+{
+	TArray<FIntRect> ComponentRects;
+	ComponentRects.Reserve(ComponentsToRender.Num());
+	for (ULandscapeComponent* Component : ComponentsToRender)
+	{
+		ComponentRects.Add(ComputeSectionRect(Component, bInWithDuplicateBorders));
+	}
+	return ComponentRects;
+}
+
 void FMergeRenderBatch::ComputeAllSubsectionRects(TArray<FIntRect>& OutSubsectionRects, TArray<FIntRect>& OutSubsectionRectsWithDuplicateBorders) const
 {
 	const int32 NumSubsections = Landscape->NumSubsections;
diff --git a/Engine/Source/Runtime/Landscape/Private/LandscapeEditLayers.cpp b/Engine/Source/Runtime/Landscape/Private/LandscapeEditLayers.cpp
--- a/Engine/Source/Runtime/Landscape/Private/LandscapeEditLayers.cpp
+++ b/Engine/Source/Runtime/Landscape/Private/LandscapeEditLayers.cpp
@@ -5540,17 +5540,17 @@ TArray<UE::Landscape::EditLayers::FEditLayerRendererState> ALandscape::GetEditLa
 
 		if (bUseEnableOverride)
 		{
-			ELandscapeToolTargetType TargetType = InMergeContext->IsHeightmapMerge() ? ELandscapeToolTargetType::Heightmap : ELandscapeToolTargetType::Weightmap;
+			ELandscapeToolTargetTypeFlags TargetTypes = InMergeContext->IsHeightmapMerge() ? ELandscapeToolTargetTypeFlags::Heightmap : (ELandscapeToolTargetTypeFlags::Weightmap | ELandscapeToolTargetTypeFlags::Visibility);
 			bool bEnabled = InLayerEnableOverride[LayerIdx];
-			Algo::ForEach(LayerRendererStates, [TargetType, bEnabled](FEditLayerRendererState& State)
+			Algo::ForEach(LayerRendererStates, [TargetTypes, bEnabled](FEditLayerRendererState& State)
 				{
 					if (bEnabled)
 					{
-						State.EnableTargetType(TargetType);
+						State.EnableTargetTypeMask(TargetTypes);
 					}
 					else
 					{
-						State.DisableTargetType(TargetType);
+						State.DisableTargetTypeMask(TargetTypes);
 					}
 				});
 		}
@@ -8080,7 +8080,7 @@ bool ALandscape::PerformSelectiveRenderEditLayersHeightmapsCPU(const FUpdateLaye
 	// Note: not using SelectiveBounds in the actual Merge process, just to retrieve the data.  The render region is based only on the
 	// components from InUpdateLayersContentContext.
 
-	FMergeContext MergeContext(/*InLandscape = */this, /*bInIsHeightmapMerge = */true, /*bInSkipProceduralRenderers = */false);
+	FMergeContext MergeContext(/*InLandscape = */this, /*bInIsHeightmapMerge = */true, !InSelectiveRenderParams.bRenderBrushes);
 	TArray<FEditLayerRendererState> RendererStates = GetEditLayerRendererStatesEnableOverride(&MergeContext, InSelectiveRenderParams.ActiveEditLayers);
 
 	FMergeRenderParams MergeRenderParams(InUpdateLayersContentContext.LandscapeComponentsHeightmapsToRender, RendererStates);
@@ -8103,7 +8103,7 @@ bool ALandscape::PerformSelectiveRenderEditLayersHeightmapsCPU(const FUpdateLaye
 	SyncRenderBatchResults.AddDefaulted(MergeRenderContext.GetRenderBatches().Num());  // allocate in advance
 
 	// Batch completion callback
-	auto EnqueueBatchReadback = [&InUpdateLayersContentContext, &InSelectiveRenderParams, &SyncRenderBatchResults]
+	auto EnqueueBatchReadback = [&SyncRenderBatchResults]
 	(const FMergeRenderContext::FOnRenderBatchTargetGroupDoneParams& InParams, UE::Landscape::FRDGBuilderRecorder& RDGBuilderRecorder)
 		{
 			TRACE_CPUPROFILER_EVENT_SCOPE(ExtractSelectiveMergeData);
@@ -8116,12 +8116,7 @@ bool ALandscape::PerformSelectiveRenderEditLayersHeightmapsCPU(const FUpdateLaye
 
 			BatchResult.SectionRect = RenderBatch->SectionRect;
 			BatchResult.Size = RenderBatch->Resolution;  //Only used for sanity check on bounds
-			BatchResult.ComponentRects.Reserve(RenderBatch->ComponentsToRender.Num());
-			for (ULandscapeComponent* Component : RenderBatch->ComponentsToRender)
-			{
-				FIntRect ComponentRect = RenderBatch->ComputeSectionRect(Component, /*bInWithDuplicateBorders = */false);
-				BatchResult.ComponentRects.Add(ComponentRect);
-			}
+			BatchResult.ComponentRects = RenderBatch->ComputeAllComponentSectionRects(/*bInWithDuplicateBorders = */false);
 
 			InParams.MergeRenderContext->CycleBlendRenderTargets(RDGBuilderRecorder);
 			ULandscapeScratchRenderTarget* ReadRT = InParams.MergeRenderContext->GetBlendRenderTargetRead();
@@ -8198,7 +8193,7 @@ bool ALandscape::PerformSelectiveRenderEditLayersHeightmapsRT(const FUpdateLayer
 
 	check(InSelectiveRenderParams.RTResult);
 
-	FMergeContext MergeContext(/*InLandscape = */this, /*bInIsHeightmapMerge = */true, /*bInSkipProceduralRenderers = */false);
+	FMergeContext MergeContext(/*InLandscape = */this, /*bInIsHeightmapMerge = */true, !InSelectiveRenderParams.bRenderBrushes);
 	TArray<FEditLayerRendererState> RendererStates = GetEditLayerRendererStatesEnableOverride(&MergeContext, InSelectiveRenderParams.ActiveEditLayers);
 
 	FMergeRenderParams MergeRenderParams(InUpdateLayersContentContext.LandscapeComponentsHeightmapsToRender, RendererStates);
@@ -8209,7 +8204,7 @@ bool ALandscape::PerformSelectiveRenderEditLayersHeightmapsRT(const FUpdateLayer
 		return false;
 	}
 
-	auto EnqueueCopyToResultRT = [&InUpdateLayersContentContext, &InSelectiveRenderParams, DestRT = InSelectiveRenderParams.RTResult]
+	auto EnqueueCopyToResultRT = [&InSelectiveRenderParams, DestRT = InSelectiveRenderParams.RTResult]
 	(const FMergeRenderContext::FOnRenderBatchTargetGroupDoneParams& InParams, UE::Landscape::FRDGBuilderRecorder& RDGBuilderRecorder)
 		{
 			TRACE_CPUPROFILER_EVENT_SCOPE(ExtractSelectiveMergeBatchData);
@@ -8223,12 +8218,13 @@ bool ALandscape::PerformSelectiveRenderEditLayersHeightmapsRT(const FUpdateLayer
 			UTextureRenderTarget2D* ReadRT = SrcScratchRT->GetRenderTarget2D();
 			FIntRect DestRect = InSelectiveRenderParams.Bounds;
 
+			TArray<FIntRect> ComponentRects = RenderBatch->ComputeAllComponentSectionRects(/*bInWithDuplicateBorders = */false);
+
 			TArray<FRHICopyTextureInfo> ComponentCopyInfos;
-			ComponentCopyInfos.Reserve(RenderBatch->ComponentsToRender.Num());
+			ComponentCopyInfos.Reserve(ComponentRects.Num());
 
-			for (ULandscapeComponent* Component : RenderBatch->ComponentsToRender)
+			for (FIntRect ComponentRect : ComponentRects)
 			{
-				FIntRect ComponentRect = RenderBatch->ComputeSectionRect(Component, /*bInWithDuplicateBorders = */false);  //component coords relative to the batch
 				ComponentRect.Max += FIntPoint(1, 1);  //quads to verts
 
 				FIntRect CopyRect = ComponentRect + RenderBatch->SectionRect.Min;  //coords in the landscape
@@ -8236,7 +8232,7 @@ bool ALandscape::PerformSelectiveRenderEditLayersHeightmapsRT(const FUpdateLayer
 				CopyRect.Clip(DestRect);
 				FIntRect SrcRect = CopyRect - RenderBatch->SectionRect.Min;  // relative to the batch
 
-				if (SrcRect.IsEmpty())
+				if (SrcRect.Area() <= 0)
 				{
 					continue;
 				}
@@ -11296,7 +11292,168 @@ int32 ALandscape::PerformLayersWeightmapsBatchedMerge(FUpdateLayersContentContex
 	return InMergeParams.WeightmapUpdateModes;
 }
 
+bool ALandscape::PerformSelectiveRenderEditLayersWeightmapsCPU(FUpdateLayersContentContext& InUpdateLayersContentContext, const FLandscapeEditLayerRenderWeightParams& InSelectiveRenderParams)
+{
+	using namespace UE::Landscape;
+	using namespace UE::Landscape::Private;
+	using namespace UE::Landscape::EditLayers;
+	using namespace UE::Landscape::EditLayers::Private;
+
+	TRACE_CPUPROFILER_EVENT_SCOPE(ALandscape::PerformSelectiveRenderEditLayersWeightmapsCPU);
+	RHI_BREADCRUMB_EVENT_GAMETHREAD("PerformSelectiveRenderEditLayersWeightmapsCPU");
+
+	FMergeContext MergeContext(/*InLandscape = */this, /*bInIsHeightmapMerge = */false, !InSelectiveRenderParams.bRenderBrushes);
+	TArray<FEditLayerRendererState> RendererStates = GetEditLayerRendererStatesEnableOverride(&MergeContext, InSelectiveRenderParams.ActiveEditLayers);
+
+	TMap<FName, int32> RequestedNameToIndex;
+	for (int Idx = 0; Idx < InSelectiveRenderParams.WeightmapNames.Num(); ++Idx)
+	{
+		RequestedNameToIndex.Emplace(InSelectiveRenderParams.WeightmapNames[Idx], Idx);
+	}
+
+	// Prepare the merge : 
+	// Not yet ready to selectively render layers, so request all layers.  Asserts deep in in FMergeRenderContext::Render if this is false.
+	const bool bRequestAllLayers = true;
+	const TSet<FName> RequestedWeightmapLayerNames(InSelectiveRenderParams.WeightmapNames);
+	FMergeRenderParams MergeRenderParams(InUpdateLayersContentContext.LandscapeComponentsWeightmapsToRender, RendererStates, RequestedWeightmapLayerNames, bRequestAllLayers);
+	FMergeRenderContext MergeRenderContext = PrepareEditLayersMergeRenderContext(MergeContext, MergeRenderParams);
+	if (!MergeRenderContext.IsValid())
+	{
+		return false;
+	}
+
+	// Set up a FRHIGPUTextureReadback for each merge batch
+	struct FSyncReadback
+	{
+		// FRHIGPUTextureReadback only supports one slice each
+		FRHIGPUTextureReadback Readback = { "SyncRenderResult" };
+		FName TargetLayerName;
+		int32 SliceIdx = 0;
+	};
+	struct FSyncRenderBatchResult
+	{
+		FIntRect SectionRect;
+		FIntPoint Size;  // RT size
+		TArray<FIntRect> ComponentRects;
+		TArray<FSyncReadback> SliceData;  // one slice per target layer
+	};
+	TArray<FSyncRenderBatchResult> SyncRenderBatchResults;
+	SyncRenderBatchResults.AddDefaulted(MergeRenderContext.GetRenderBatches().Num());  // allocate in advance.  accessed from the render thread, so can't re-size.
+
+	// Batch completion callback
+	// New version, based on selective render heightmaps.  no resolve stuff.
+	auto EnqueueBatchReadback = [&RequestedWeightmapLayerNames, &SyncRenderBatchResults]
+	(const FMergeRenderContext::FOnRenderBatchTargetGroupDoneParams& InParams, UE::Landscape::FRDGBuilderRecorder& RDGBuilderRecorder)
+		{
+			TRACE_CPUPROFILER_EVENT_SCOPE(ExtractSelectiveMergeData);
+			RHI_BREADCRUMB_EVENT_GAMETHREAD("ExtractSelectiveMergeData");
+
+			InParams.MergeRenderContext->CycleBlendRenderTargets(RDGBuilderRecorder);
+			ULandscapeScratchRenderTarget* ReadRT = InParams.MergeRenderContext->GetBlendRenderTargetRead();
+			UTextureRenderTarget2DArray* SourceTexArray = ReadRT->GetRenderTarget2DArray();
+			check(SourceTexArray->GetSurfaceArraySize() == InParams.TargetLayerGroupLayerNames.Num());
+
+			ALandscape* Landscape = InParams.MergeRenderContext->GetLandscape();
+			const FMergeRenderBatch* RenderBatch = InParams.MergeRenderContext->GetCurrentRenderBatch();
+			int32 CurrentBatch = InParams.MergeRenderContext->GetCurrentRenderBatchIdx();
+			FSyncRenderBatchResult& BatchResult = SyncRenderBatchResults[CurrentBatch];
+
+			BatchResult.SectionRect = RenderBatch->SectionRect;
+			BatchResult.Size = RenderBatch->Resolution;  //Only used for sanity check on bounds
+			BatchResult.ComponentRects = RenderBatch->ComputeAllComponentSectionRects(/*bInWithDuplicateBorders = */false);
+
+			BatchResult.SliceData.Reserve(InParams.TargetLayerGroupLayerNames.Num());
+			for (int32 Idx = 0; Idx < InParams.TargetLayerGroupLayerNames.Num(); ++Idx)
+			{
+				// Only need the layers we asked for.  The merge code could include other layers based on internal requirements.
+				if (RequestedWeightmapLayerNames.Contains(InParams.TargetLayerGroupLayerNames[Idx]))
+				{
+					FSyncReadback& ReadbackData = BatchResult.SliceData.AddDefaulted_GetRef();
+					ReadbackData.TargetLayerName = InParams.TargetLayerGroupLayerNames[Idx];
+					ReadbackData.SliceIdx = Idx;
+				}
+			}
+
+			// Enqueue async readback direct from the current scratch read target.  One readback per texture slice, but only on the slices that we requested.
+			// Extra layers are present because bRequestAllLayers=true is required.
+			ENQUEUE_RENDER_COMMAND(LandscapeLayers_Cmd_SyncExtractResults)(
+				[&BatchResult,
+				SourceResource = SourceTexArray->GetResource()]
+				(FRHICommandListImmediate& InRHICmdList)
+				{
+					FRDGBuilder GraphBuilder(InRHICmdList, RDG_EVENT_NAME("WeightmapsExtractResult"));
+
+					FRDGTextureRef SourceTextureRef = GraphBuilder.RegisterExternalTexture(CreateRenderTarget(SourceResource->TextureRHI, TEXT("SourceTexture")));
+					for (FSyncReadback& SliceData : BatchResult.SliceData)
+					{
+						AddEnqueueCopyPass(GraphBuilder, &SliceData.Readback, SourceTextureRef, FResolveRect(), SliceData.SliceIdx);
+					}
+
+					// We need to specify the final state of the external texture to prevent the graph builder from transitioning it to SRVMask :
+					GraphBuilder.SetTextureAccessFinal(SourceTextureRef, ERHIAccess::RTV);
+					GraphBuilder.Execute();
+				});
+		};
+	
+	// Render everything now. Every time a group from a batch is done, the EnqueueBatchReadback callback is called : 
+	MergeRenderContext.Render(EnqueueBatchReadback);
+
+	// Pick up readback results from SyncRenderBatchResults.  Enqueue this on the render thread because FRHIGPUTextureReadback can only be locked there.
+	ENQUEUE_RENDER_COMMAND(LandscapeLayers_Cmd_SyncFetchResults)(
+		[&SyncRenderBatchResults,
+		&RequestedNameToIndex,
+		DestRect = InSelectiveRenderParams.Bounds,
+		Dest = InSelectiveRenderParams.CpuResult]
+		(FRHICommandListImmediate& InRHICmdList)
+		{
+			TRACE_CPUPROFILER_EVENT_SCOPE(LockAndFetchReadbackData);
+			const FBlitBuffer2DDesc DestDesc{ Dest.Num(), DestRect.Width(), DestRect };
+
+			int32 DestSliceStride = DestRect.Area();
 
+			for (FSyncRenderBatchResult& BatchResult : SyncRenderBatchResults)
+			{
+				int32 NumSlices = BatchResult.SliceData.Num();
+				for (FSyncReadback& SliceData : BatchResult.SliceData)
+				{
+					FName LayerName = SliceData.TargetLayerName;
+					int32* DestLayerIndexPtr = RequestedNameToIndex.Find(LayerName);
+					if (DestLayerIndexPtr == nullptr)
+					{
+						continue;  // didn't ask for this layer
+					}
+					int32 DestLayerIndex = *DestLayerIndexPtr;
+
+					int32 BatchStride = 0;
+					int32 BatchHeight = 0;
+					uint32* SrcData = static_cast<uint32*>(SliceData.Readback.Lock(BatchStride, &BatchHeight));
+					int32 SrcBufferSize = BatchStride * BatchHeight; //the implied buffer size of the locked readback buffer, not necessarily of the native render target
+					check(BatchResult.Size.X <= BatchStride && BatchResult.Size.Y <= BatchHeight);
+					uint8* DestPtr = Dest.GetData() + (DestSliceStride * DestLayerIndex);
+					check(DestSliceStride * (DestLayerIndex + 1) <= Dest.Num());
+
+					// The full batch rects are not all valid.  Only rects of each component in the batch are valid.  Blit one component at a time.
+					for (const FIntRect& ComponentBatchRect : BatchResult.ComponentRects)
+					{
+						// ComponentBatchRect is relative to the batch RT.  Convert back into overall landscape coords.
+						FIntRect ComponentRect = ComponentBatchRect + BatchResult.SectionRect.Min;
+						ComponentRect.Max += FIntPoint(1, 1);  //quads to verts
+						check(BatchResult.SectionRect.Contains(ComponentRect));
+
+						const FBlitBuffer2DDesc SrcDesc{ SrcBufferSize, BatchStride, BatchResult.SectionRect };
+						BlitWeightChannelsToUint8(DestPtr, DestDesc, SrcData, SrcDesc, ComponentRect);
+					}
+
+					SliceData.Readback.Unlock();
+				}
+			}
+		});
+
+	// Sync the game thread with all of the above.  After this, the results are reliably accessible from the game thread.
+	FlushRenderingCommands();
+
+	return true;
+}
 
 int32 ALandscape::RegenerateLayersWeightmaps(FUpdateLayersContentContext& InUpdateLayersContentContext)
 {
@@ -12009,7 +12166,7 @@ bool ALandscape::SelectiveRenderEditLayersHeightmaps(const FLandscapeEditLayerRe
 {
 	check((InSelectiveRenderParams.CpuResult.Num() > 0) ^ (InSelectiveRenderParams.RTResult != nullptr));  //Expecting either a cpu result buffer or a render target.  Not both.
 
-	if (InSelectiveRenderParams.Bounds.IsEmpty())
+	if (InSelectiveRenderParams.Bounds.Area() <= 0)
 	{
 		return true;  //nothing to render
 	}
@@ -12057,21 +12214,75 @@ bool ALandscape::SelectiveRenderEditLayersHeightmaps(const FLandscapeEditLayerRe
 
 		if (!InSelectiveRenderParams.CpuResult.IsEmpty())
 		{
-			PerformSelectiveRenderEditLayersHeightmapsCPU(LayerContext, InSelectiveRenderParams);
+			return PerformSelectiveRenderEditLayersHeightmapsCPU(LayerContext, InSelectiveRenderParams);
 		}
 		else
 		{
-			PerformSelectiveRenderEditLayersHeightmapsRT(LayerContext, InSelectiveRenderParams);
+			return PerformSelectiveRenderEditLayersHeightmapsRT(LayerContext, InSelectiveRenderParams);
 		}
 	}
-
-	return true;
 }
 
 bool ALandscape::SelectiveRenderEditLayersWeightmaps(const FLandscapeEditLayerRenderWeightParams& InSelectiveRenderParams)
 {
-	unimplemented();
-	return false;
+	check((InSelectiveRenderParams.CpuResult.Num() > 0) ^ (InSelectiveRenderParams.RTResult != nullptr));  //Expecting either a cpu result buffer or a render target.  Not both.
+
+	if (InSelectiveRenderParams.Bounds.Area() <= 0)
+	{
+		return true;  //nothing to render
+	}
+
+	// Expecting to be already initialized for BatchMerge.
+	check(CurrentEditLayersMergeMode == ELandscapeEditLayersMergeMode::BatchedMerge);
+	check(GetEditLayersMergeMode() == ELandscapeEditLayersMergeMode::BatchedMerge);
+	check(bLandscapeLayersAreInitialized);
+
+	if (!CanUpdateLayersContent())
+	{
+		return false;
+	}
+
+	UWorld* World = GetWorld();
+	check(World != nullptr);
+	ULandscapeSubsystem* LandscapeSubsystem = World->GetSubsystem<ULandscapeSubsystem>();
+	check(LandscapeSubsystem != nullptr);
+	ULandscapeInfo* LandscapeInfo = GetLandscapeInfo();
+	check(LandscapeInfo);
+
+	const bool bWaitForStreaming = true;
+	if (!PrepareTextureResources(bWaitForStreaming) ||
+		!PrepareLayersTextureResources(bWaitForStreaming) ||
+		!PrepareLayersResources(GetFeatureLevelShaderPlatform_Checked(World->GetFeatureLevel()), bWaitForStreaming))
+	{
+		return false;
+	}
+
+	TSet<ULandscapeComponent*> AffectedComponents;
+	LandscapeInfo->GetComponentsInRegion(InSelectiveRenderParams.Bounds.Min.X, InSelectiveRenderParams.Bounds.Min.Y, InSelectiveRenderParams.Bounds.Max.X, InSelectiveRenderParams.Bounds.Max.Y, AffectedComponents);
+	if (AffectedComponents.Num() == 0)
+	{
+		// No loaded components in region
+		UE_LOG(LogLandscape, Display, TEXT("SelectiveRenderLayers with no affected components"));
+		return true;
+	}
+
+	FTextureToComponentHelper Helper(*LandscapeInfo, AffectedComponents.Array(), FTextureToComponentHelper::ERefreshFlags::RefreshWeightmaps);
+	FUpdateLayersContentContext LayerContext(Helper, FUpdateLayersContentContext::FSelectiveRenderTag(), FUpdateLayersContentContext::ERefreshFlags::RefreshWeightmapInfos);
+
+	{
+		RenderCaptureInterface::FScopedCapture RenderCapture((RenderCaptureLayersNextHeightmapDraws != 0), TEXT("SelectiveRenderEditLayersWeightmaps"));
+		RenderCaptureLayersNextHeightmapDraws = FMath::Max(0, RenderCaptureLayersNextHeightmapDraws - 1);
+
+		if (!InSelectiveRenderParams.CpuResult.IsEmpty())
+		{
+			return PerformSelectiveRenderEditLayersWeightmapsCPU(LayerContext, InSelectiveRenderParams);
+		}
+		else
+		{
+			unimplemented();
+			return false;
+		}
+	}
 }
 
 void ALandscape::UpdateLayersContent(bool bInWaitForStreaming, bool bInSkipMonitorLandscapeEdModeChanges, bool bIntermediateRender, bool bFlushRender)
@@ -13348,6 +13559,115 @@ bool ALandscape::DeleteLayer(int32 InLayerIndex)
 
 void ALandscape::CollapseLayer(int32 InLayerIndex)
 {
+	if (!AllowSelectiveRenderEditLayers())
+	{
+		// TODO:  remove with IntermediateRender.
+		CollapseLayerIntermediateRender(InLayerIndex);
+		return;
+	}
+
+	check(InLayerIndex >= 1);
+
+	FScopedSlowTask SlowTask(static_cast<float>(GetLandscapeInfo()->XYtoComponentMap.Num()), LOCTEXT("Landscape_CollapseLayer_SlowWork", "Collapsing Layer..."));
+	SlowTask.MakeDialog();
+
+	ULandscapeInfo* LandscapeInfo = GetLandscapeInfo();
+	check(LandscapeInfo);
+
+	// Only render the two layers being combined
+	TBitArray<> EditLayerVis;
+	for (int32 i = 0; i < LandscapeEditLayers.Num(); ++i)
+	{
+		bool bRenderLayer = i == InLayerIndex || i == InLayerIndex - 1;
+		EditLayerVis.Add(bRenderLayer);
+	}
+
+	TArray<FName> AllTargetNames = GetTargetLayerNames(/*bInIncludeVisibilityLayer = */true);
+
+	FIntRect ExtentInc;
+	verify(LandscapeInfo->GetLandscapeExtent(ExtentInc));
+	FIntRect ExtentExcl = ExtentInc;
+	ExtentExcl.Max += FIntPoint(1, 1);  //convert to exclusive range
+	int32 LayerSamples = ExtentExcl.Area();
+	check(LayerSamples > 0);
+
+	// ForceUpdateLayersContent without requesting any updates to flush any pending evaluation.
+	const bool bLocalIntermediateRender = false;
+	ForceUpdateLayersContent(bLocalIntermediateRender);
+
+	TArray<uint16> HeightData;
+	HeightData.AddZeroed(LayerSamples);
+
+	TArray<uint8> WeightData;
+	bool bHasTargetLayers = !AllTargetNames.IsEmpty();
+	if (bHasTargetLayers)
+	{
+		WeightData.AddZeroed(LayerSamples * AllTargetNames.Num());
+	}
+
+	// Render heightmap
+	FLandscapeEditLayerRenderHeightParams RenderHeightParams;
+	RenderHeightParams.Bounds = ExtentExcl;
+	RenderHeightParams.ActiveEditLayers = EditLayerVis;
+	RenderHeightParams.bRenderBrushes = false;
+	RenderHeightParams.CpuResult = HeightData;
+	bool bRenderRet = SelectiveRenderEditLayersHeightmaps(RenderHeightParams);
+	check(bRenderRet);
+
+	if (bHasTargetLayers)
+	{
+		// Render weightmaps
+		FLandscapeEditLayerRenderWeightParams RenderWeightParams;
+		RenderWeightParams.Bounds = ExtentExcl;
+		RenderWeightParams.ActiveEditLayers = EditLayerVis;
+		RenderWeightParams.bRenderBrushes = false;
+		RenderWeightParams.WeightmapNames = AllTargetNames;
+		RenderWeightParams.CpuResult = WeightData;
+		bRenderRet = SelectiveRenderEditLayersWeightmaps(RenderWeightParams);
+		check(bRenderRet);
+	}
+
+	// Assign combined data to destination layer
+	{
+		FLandscapeEditDataInterface DataInterface(LandscapeInfo);
+		DataInterface.SetShouldDirtyPackage(true);
+		const ULandscapeEditLayerBase* DestLayer = GetEditLayerConst(InLayerIndex - 1);
+		FGuid DestLayerGuid = DestLayer->GetGuid();
+		DataInterface.SetEditLayer(DestLayerGuid);
+		int32 DataStride = 0;  //automatic
+
+		DataInterface.SetHeightData(ExtentInc.Min.X, ExtentInc.Min.Y, ExtentInc.Max.X, ExtentInc.Max.Y,
+			HeightData.GetData(), DataStride, /*InCalcNormals=*/ true);
+
+		for (int32 TargetLayerIdx = 0; TargetLayerIdx < AllTargetNames.Num(); ++TargetLayerIdx)
+		{
+			// Set weightmap data one target layer at a time.  There is a multi-layer version, but it expects interleaved channel data.
+			ULandscapeLayerInfoObject* LayerInfo = LandscapeInfo->GetLayerInfoByName(AllTargetNames[TargetLayerIdx]);
+			check(LayerInfo);
+			int32 DataOffset = LayerSamples * TargetLayerIdx;
+			check(WeightData.IsValidIndex(DataOffset + LayerSamples - 1));
+			DataInterface.SetAlphaData(LayerInfo, ExtentInc.Min.X, ExtentInc.Min.Y, ExtentInc.Max.X, ExtentInc.Max.Y,
+				WeightData.GetData() + DataOffset, DataStride);
+		}
+	}
+
+	// Note that this code theoretically supports brushes, to handle the non-brush persistent layer data that brush-containing layers
+	// are allowed to hold. But the UI doesn't allow CollapseLayer to be used if either layer contains brushes.
+	for (int32 i = 0; i < LandscapeEditLayers[InLayerIndex].Brushes.Num(); ++i)
+	{
+		ALandscapeBlueprintBrushBase* Brush = LandscapeEditLayers[InLayerIndex].Brushes[i].GetBrush();
+		RemoveBrushFromLayer(InLayerIndex, Brush);
+		AddBrushToLayer(InLayerIndex - 1, Brush);
+	}
+
+	DeleteLayer(InLayerIndex);
+
+	RequestLayersContentUpdateForceAll();
+}
+
+void ALandscape::CollapseLayerIntermediateRender(int32 InLayerIndex)
+{
+	// TODO:  remove with IntermediateRender.
 	FScopedSlowTask SlowTask(static_cast<float>(GetLandscapeInfo()->XYtoComponentMap.Num()), LOCTEXT("Landscape_CollapseLayer_SlowWork", "Collapsing Layer..."));
 	SlowTask.MakeDialog();
 	TArray<bool> BackupVisibility;
@@ -13366,10 +13686,10 @@ void ALandscape::CollapseLayer(int32 InLayerIndex)
 
 	// Call Request Update on all components...
 	GetLandscapeInfo()->ForAllLandscapeComponents([](ULandscapeComponent* LandscapeComponent)
-	{
-		LandscapeComponent->RequestWeightmapUpdate(false, false);
-		LandscapeComponent->RequestHeightmapUpdate(false, false);
-	});
+		{
+			LandscapeComponent->RequestWeightmapUpdate(false, false);
+			LandscapeComponent->RequestHeightmapUpdate(false, false);
+		});
 
 	const bool bLocalIntermediateRender = true;
 	ForceUpdateLayersContent(bLocalIntermediateRender);
@@ -13382,10 +13702,10 @@ void ALandscape::CollapseLayer(int32 InLayerIndex)
 		TSet<UTexture2D*> ProcessedHeightmaps;
 		FScopedSetLandscapeEditingLayer ScopeEditingLayer(this, LandscapeEditLayers[InLayerIndex - 1].EditLayer->GetGuid());
 		GetLandscapeInfo()->ForAllLandscapeComponents([&](ULandscapeComponent* LandscapeComponent)
-		{
-			SlowTask.EnterProgressFrame(1.f);
-			LandscapeComponent->CopyFinalLayerIntoEditingLayer(DataInterface, ProcessedHeightmaps);
-		});
+			{
+				SlowTask.EnterProgressFrame(1.f);
+				LandscapeComponent->CopyFinalLayerIntoEditingLayer(DataInterface, ProcessedHeightmaps);
+			});
 	}
 
 	TArray<ALandscapeBlueprintBrushBase*> BrushesToMove;
@@ -14261,15 +14581,20 @@ UTextureRenderTarget2D* FLandscapeLayerBrush::RenderLayer(const FLandscapeBrushP
 
 void LandscapeDumpSelectiveLayerRender(const TArray<FString>& Args)
 {
-	// Example usage:  landscape.DumpSelectiveLayerRender -landscape Landscape -h -layers 255 -region 0 0 1000 500
-	// override default layer visibility with a bitmask:  -layers 255
+	// Example usages:
+	//		landscape.DumpSelectiveLayerRender -landscape Landscape -h -layermask 255 -region 0 0 1000 500
+	//		landscape.DumpSelectiveLayerRender -landscape Landscape -w -wname Grass -layermask 255 -region 0 0 1000 500
+	// override default layer visibility with a bitmask:  -layermask 255
 	// render a region of the landscape based on landscape coords (integers):  -region minX minY maxX maxY
-	// All params optional.
+	// All params optional except -w or -h.
+	// -w with no -wname means all weightmaps.
 
 	const TCHAR* LandscapeNameArg = TEXT("-landscape");
 	const TCHAR* LayerMaskArg = TEXT("-layermask");
 	const TCHAR* RegionArg = TEXT("-region");
 	const TCHAR* HeightmapArg = TEXT("-h");
+	const TCHAR* WeightmapArg = TEXT("-w");
+	const TCHAR* WeightmapNameArg = TEXT("-wname");
 	const TCHAR* RTArg = TEXT("-rt");
 
 	ALandscape* Landscape = nullptr;
@@ -14280,7 +14605,9 @@ void LandscapeDumpSelectiveLayerRender(const TArray<FString>& Args)
 	FIntRect Region;
 	bool bUseRenderTarget = false;
 
-	bool bRenderHeight = true;  // Assuming heightmap render for now.  Weightmap support still to be added.
+	bool bRenderHeight = false;
+	bool bRenderWeight = false;
+	TArray<FName> WeightmapNames;
 
 	UWorld* World = GEditor->GetEditorWorldContext().World();
 	if (!World)
@@ -14288,6 +14615,8 @@ void LandscapeDumpSelectiveLayerRender(const TArray<FString>& Args)
 		return;
 	}
 
+	// Parse args
+
 	for (int32 i = 0; i < Args.Num(); ++i)
 	{
 		const FString& Arg = Args[i];
@@ -14325,6 +14654,18 @@ void LandscapeDumpSelectiveLayerRender(const TArray<FString>& Args)
 		{
 			bRenderHeight = true;
 		}
+		else if (Arg == WeightmapArg)
+		{
+			bRenderWeight = true;
+		}
+		else if (Arg == WeightmapNameArg)
+		{
+			++i;
+			if (i < Args.Num())
+			{
+				WeightmapNames.Add(FName(Args[i]));
+			}
+		}
 		else if (Arg == RTArg)
 		{
 			bUseRenderTarget = true;
@@ -14335,6 +14676,8 @@ void LandscapeDumpSelectiveLayerRender(const TArray<FString>& Args)
 		}
 	}
 
+	// Validate args and fill in implied params
+
 	// If landscape name not specified, just take the first one
 	for (ALandscape* Actor : TActorRange<ALandscape>(World))
 	{
@@ -14351,7 +14694,7 @@ void LandscapeDumpSelectiveLayerRender(const TArray<FString>& Args)
 		return;
 	}
 
-	if (!bRenderHeight)
+	if (!bRenderHeight && !bRenderWeight)
 	{
 		UE_LOG(LogLandscape, Error, TEXT("LandscapeDumpSelectiveLayerRender: must specify either -h or -w"));
 		return;
@@ -14362,6 +14705,7 @@ void LandscapeDumpSelectiveLayerRender(const TArray<FString>& Args)
 		if (ULandscapeInfo* LandscapeInfo = ULandscapeInfo::Find(World, Landscape->GetLandscapeGuid()))
 		{
 			bHasRegion = LandscapeInfo->GetLandscapeExtent(Region);
+			Region.Max += FIntPoint(1, 1);  //convert to exclusive range
 		}
 	}
 
@@ -14387,10 +14731,40 @@ void LandscapeDumpSelectiveLayerRender(const TArray<FString>& Args)
 		}
 	}
 
+	if (bRenderWeight)
+	{
+		TArray<FName> AllNames = Landscape->GetTargetLayerNames(/*bInIncludeVisibilityLayer = */true);
+		if (WeightmapNames.IsEmpty())
+		{
+			WeightmapNames = AllNames;
+		}
+		else
+		{
+			// Filter arg-provided names by available names.
+			WeightmapNames.SetNum(Algo::RemoveIf(WeightmapNames, [&AllNames](FName TestName)
+				{
+					if (!AllNames.Contains(TestName))
+					{
+						UE_LOG(LogLandscape, Warning, TEXT("LandscapeDumpSelectiveLayerRender: weightmap named %s not found."), *TestName.ToString());
+						return true;
+					}
+					return false;
+				}));
+		}
+
+		if (WeightmapNames.IsEmpty())
+		{
+			UE_LOG(LogLandscape, Warning, TEXT("LandscapeDumpSelectiveLayerRender: no weightmaps to render."));
+			bRenderWeight = false;
+		}
+	}
+
 	ULandscapeSubsystem* Subsystem = World->GetSubsystem<ULandscapeSubsystem>();
 	const FDateTime CurrentTime = Subsystem->GetAppCurrentDateTime();
 	FString FilePattern = FString::Format(TEXT("{0}/Landscape/LandscapeRender{1}"), { FPaths::ProjectSavedDir(), CurrentTime.ToString() });
 
+	// Render the heightmaps
+
 	if (bRenderHeight)
 	{
 		UE_LOG(LogLandscape, Display, TEXT("LandscapeDumpSelectiveLayerRender: SelectiveRenderLayers heightmap, landscape=%s region=(%d,%d)-(%d,%d) layers=%x"),
@@ -14482,6 +14856,52 @@ void LandscapeDumpSelectiveLayerRender(const TArray<FString>& Args)
 			UE_LOG(LogLandscape, Error, TEXT("LandscapeDumpSelectiveLayerRender: failed to write heightmap data to %s"), *Filename);
 		}
 	}
+
+	// Render the weightmaps
+
+	if (bRenderWeight)
+	{
+		UE_LOG(LogLandscape, Display, TEXT("LandscapeDumpSelectiveLayerRender: SelectiveRenderLayers weightmaps, landscape=%s region=(%d,%d)-(%d,%d) layers=%x"),
+			*Landscape->GetActorNameOrLabel(), Region.Min.X, Region.Min.Y, Region.Max.X, Region.Max.Y, LayerMask);
+
+		const int32 SingleWeightmapSizePixels = Region.Area();
+		TArray<uint8> WeightData;
+
+		WeightData.AddZeroed(SingleWeightmapSizePixels * WeightmapNames.Num());
+
+		FLandscapeEditLayerRenderWeightParams RenderParams;
+		RenderParams.Bounds = Region;
+		RenderParams.ActiveEditLayers = LayerVis;
+		RenderParams.WeightmapNames = WeightmapNames;
+
+		RenderParams.CpuResult = WeightData;
+
+		bool bRenderRet = Landscape->SelectiveRenderEditLayersWeightmaps(RenderParams);
+		if (!bRenderRet)
+		{
+			UE_LOG(LogLandscape, Error, TEXT("LandscapeDumpSelectiveLayerRender: SelectiveRenderEditLayersHeightmaps failed"));
+			return;
+		}
+
+		FImageInfo ImageInfo(Region.Size().X, Region.Size().Y, 1, ERawImageFormat::G8, EGammaSpace::Linear);
+
+		for (int32 NameIdx = 0; NameIdx < WeightmapNames.Num(); NameIdx++)
+		{
+			// Write each weightmap to an image file.
+			FName WeightmapName = WeightmapNames[NameIdx];
+			FImageView ImageView(ImageInfo, WeightData.GetData() + (SingleWeightmapSizePixels * NameIdx));  // Offset into image buffer
+			FString Filename = FilePattern + "Weightmap_" + WeightmapName.ToString() + ".png";
+			if (FImageUtils::SaveImageByExtension(*Filename, ImageView))
+			{
+				UE_LOG(LogLandscape, Display, TEXT("LandscapeDumpSelectiveLayerRender: wrote weightmap data to %s"), *Filename);
+			}
+			else
+			{
+				UE_LOG(LogLandscape, Error, TEXT("LandscapeDumpSelectiveLayerRender: failed to write weightmap data to %s"), *Filename);
+			}
+		}
+
+	}
 }
 
 
diff --git a/Engine/Source/Runtime/Landscape/Private/LandscapeUtilsPrivate.cpp b/Engine/Source/Runtime/Landscape/Private/LandscapeUtilsPrivate.cpp
--- a/Engine/Source/Runtime/Landscape/Private/LandscapeUtilsPrivate.cpp
+++ b/Engine/Source/Runtime/Landscape/Private/LandscapeUtilsPrivate.cpp
@@ -91,7 +91,8 @@ FLandscapeComponent2DIndexer CreateLandscapeComponent2DIndexer(const ULandscapeI
 	return FLandscapeComponent2DIndexer(AllValidComponents);
 }
 
-int32 BlitHeightChannelsToUint16 (uint16* DestBuffer, const FBlitBuffer2DDesc& Dest, uint32* SrcBuffer, const FBlitBuffer2DDesc& Src, FIntRect ClipRect)
+template<typename DestType, typename ConvertFromColorType>
+static int32 BlitRGChannelsToSingleTmpl(DestType* DestBuffer, const FBlitBuffer2DDesc& Dest, uint32* SrcBuffer, const FBlitBuffer2DDesc& Src, FIntRect ClipRect, ConvertFromColorType ConvertFromColor)
 {
 	check(Dest.Rect.Height() * Dest.Stride <= Dest.BufferSize);
 	check(Src.Rect.Height() * Src.Stride <= Src.BufferSize);
@@ -114,8 +115,8 @@ int32 BlitHeightChannelsToUint16 (uint16* DestBuffer, const FBlitBuffer2DDesc& D
 		uint32* Read = SrcBuffer + (SrcY * Src.Stride) + ReadOffset.X;
 		uint32* ReadEnd = Read + BlitSize.X;
 		int32 DestY = Y - Dest.Rect.Min.Y;
-		uint16* Write = DestBuffer + (DestY * Dest.Stride) + WriteOffset.X;
-		uint16* WriteEnd = Write + BlitSize.X;
+		DestType* Write = DestBuffer + (DestY * Dest.Stride) + WriteOffset.X;
+		DestType* WriteEnd = Write + BlitSize.X;
 
 		check(Read >= SrcBuffer);
 		check(ReadEnd <= SrcBuffer + Src.BufferSize);
@@ -126,13 +127,36 @@ int32 BlitHeightChannelsToUint16 (uint16* DestBuffer, const FBlitBuffer2DDesc& D
 
 		for (; Read < ReadEnd && Write < WriteEnd; ++Read, ++Write)
 		{
-			FColor HeightAndNormal(*Read);
-			uint16 HeightValue = (static_cast<uint16>(HeightAndNormal.R) << 8) | static_cast<uint16>(HeightAndNormal.G);
-			*Write = HeightValue;
+			FColor Color(*Read);
+			DestType Value = ConvertFromColor(Color);
+			*Write = Value;
 		}
 		check(Read == ReadEnd && Write == WriteEnd);
 	}
 	return BlitRect.Area();
+}
+
+
+int32 BlitHeightChannelsToUint16(uint16* DestBuffer, const FBlitBuffer2DDesc& Dest, uint32* SrcBuffer, const FBlitBuffer2DDesc& Src, FIntRect ClipRect)
+{
+	// Get 16 bit height value from R and G channels
+	auto GetHeight = [](const FColor& HeightAndNormal) { return (static_cast<uint16>(HeightAndNormal.R) << 8) | static_cast<uint16>(HeightAndNormal.G); };
+
+	return BlitRGChannelsToSingleTmpl(DestBuffer, Dest, SrcBuffer, Src, ClipRect, GetHeight);
+};
+
+int32 BlitWeightChannelsToUint8(uint8* DestBuffer, const FBlitBuffer2DDesc& Dest, uint32* SrcBuffer, const FBlitBuffer2DDesc& Src, FIntRect ClipRect)
+{
+	// 16 bit weight value in R and G channels represents 0..1.  65535 maps to 255.
+	auto GetWeight = [](const FColor& HeightAndNormal)
+		{
+			uint32 Weight = (static_cast<uint16>(HeightAndNormal.R) << 8) | static_cast<uint16>(HeightAndNormal.G);
+			Weight = Weight * 255 / 65535;
+			check(Weight <= 255);
+			return static_cast<uint8>(Weight);
+		};
+
+	return BlitRGChannelsToSingleTmpl(DestBuffer, Dest, SrcBuffer, Src, ClipRect, GetWeight);
 };
 
 } // end namespace UE::Landscape::Private
\ No newline at end of file
diff --git a/Engine/Source/Runtime/Landscape/Private/LandscapeUtilsPrivate.h b/Engine/Source/Runtime/Landscape/Private/LandscapeUtilsPrivate.h
--- a/Engine/Source/Runtime/Landscape/Private/LandscapeUtilsPrivate.h
+++ b/Engine/Source/Runtime/Landscape/Private/LandscapeUtilsPrivate.h
@@ -328,4 +328,18 @@ namespace UE::Landscape::Private
 	 */
 	int32 BlitHeightChannelsToUint16(uint16* DestBuffer, const FBlitBuffer2DDesc& Dest, uint32* SrcBuffer, const FBlitBuffer2DDesc& Src, FIntRect ClipRect);
 
+
+	/* Copy 16-bit weight values from the R+G channels of the source buffer, scaled down into a plain uint8 dest buffer.  The source and destination buffers
+	 * represent sub-regions of a shared 2D coordinate system expressed in the Rect parameters.  ClipRect is the desired area of that
+	 * shared coordinate space to limit the operation to.  Copies only values in the the intersection of Src.Rect, Dest.Rect, and ClipRect.
+	 * All FIntRect params are treated as half-open.  .Min is included.  .Max is excluded.
+	 * @param DestBuffer destination buffer
+	 * @param Dest description of the destination buffer
+	 * @param SrcBuffer source buffer
+	 * @param Src description of the source buffer
+	 * @param ClipRect The operation will be limited to this region, in the same shared coordinate space of the rects from the buffer descriptions.
+	 * @return number of elements copied.
+	 */
+	int32 BlitWeightChannelsToUint8(uint8* DestBuffer, const FBlitBuffer2DDesc& Dest, uint32* SrcBuffer, const FBlitBuffer2DDesc& Src, FIntRect ClipRect);
+
 } // end namespace UE::Landscape::Private
diff --git a/Engine/Source/Runtime/Landscape/Public/LandscapeEditLayerMergeRenderContext.h b/Engine/Source/Runtime/Landscape/Public/LandscapeEditLayerMergeRenderContext.h
--- a/Engine/Source/Runtime/Landscape/Public/LandscapeEditLayerMergeRenderContext.h
+++ b/Engine/Source/Runtime/Landscape/Public/LandscapeEditLayerMergeRenderContext.h
@@ -201,6 +201,12 @@ public:
 	 */
 	FIntRect ComputeSectionRect(ULandscapeComponent* InComponent, bool bInWithDuplicateBorders) const;
 
+	/**
+	 * Returns all subregions of the render batch representing valid data.
+	 * @param bInWithDuplicateBorders indicates whether the rect coordinates should include the duplicated columns/rows at the end of each subsection or not
+	 */
+	TArray<FIntRect> ComputeAllComponentSectionRects(bool bInWithDuplicateBorders) const;
+
 	/**
 	 * Compute the rects corresponding to the sub-sections that need to be read from and written to when expanding the render target (inclusive bounds)
 	 * @param OutSubsectionRects List of all subsection rects when *not* taking into account duplicate borders (inclusive bounds)
