Fixed mobile landscape data hash being zeroed when entering PIE. That was leading to all landscape mobile data being rebuilt when entering PIE while in mobile preview mode.

diff --git a/Engine/Source/Runtime/Landscape/Classes/LandscapeComponent.h b/Engine/Source/Runtime/Landscape/Classes/LandscapeComponent.h
--- a/Engine/Source/Runtime/Landscape/Classes/LandscapeComponent.h
+++ b/Engine/Source/Runtime/Landscape/Classes/LandscapeComponent.h
@@ -646,11 +646,11 @@ public:
 	FLandscapeEditToolRenderData EditToolRenderData;
 
 	/** Hash of source for mobile generated data. Used determine if we need to re-generate mobile pixel data. */
-	UPROPERTY(DuplicateTransient)
+	UPROPERTY(NonPIEDuplicateTransient)
 	FGuid MobileDataSourceHash;
 
 	/** Represent the chosen material for each LOD */
-	UPROPERTY(DuplicateTransient)
+	UPROPERTY(NonPIEDuplicateTransient)
 	TMap<TObjectPtr<UMaterialInterface>, int8> MaterialPerLOD;
 
 	/** Represents hash of last weightmap usage update */
