[HWRT] Fix issues in some GetDynamicRayTracingInstances(...) implementations.

diff --git a/Engine/Plugins/Experimental/Water/Source/Runtime/Private/WaterMeshSceneProxy.cpp b/Engine/Plugins/Experimental/Water/Source/Runtime/Private/WaterMeshSceneProxy.cpp
--- a/Engine/Plugins/Experimental/Water/Source/Runtime/Private/WaterMeshSceneProxy.cpp
+++ b/Engine/Plugins/Experimental/Water/Source/Runtime/Private/WaterMeshSceneProxy.cpp
@@ -1103,7 +1103,7 @@ void FWaterMeshSceneProxy::GetDynamicRayTracingInstances(FRayTracingInstanceColl
 						continue;
 					}
 
-					Collector.AddRayTracingInstance(ViewIndex, MoveTemp(RayTracingInstance));
+					Collector.AddRayTracingInstance(ViewIndex, RayTracingInstance);
 				}
 			}
 		}
diff --git a/Engine/Plugins/FX/Niagara/Source/Niagara/Private/NiagaraRendererRibbons.cpp b/Engine/Plugins/FX/Niagara/Source/Niagara/Private/NiagaraRendererRibbons.cpp
--- a/Engine/Plugins/FX/Niagara/Source/Niagara/Private/NiagaraRendererRibbons.cpp
+++ b/Engine/Plugins/FX/Niagara/Source/Niagara/Private/NiagaraRendererRibbons.cpp
@@ -1171,7 +1171,7 @@ void FNiagaraRendererRibbons::GetDynamicRayTracingInstances(FRayTracingInstanceC
 	
 	for (int32 ViewIndex = 0; ViewIndex < Views.Num(); ViewIndex++)
 	{
-		if (VisibilityMap & (1 << ViewIndex))
+		if ((VisibilityMap & (1 << ViewIndex)) == 0)
 		{
 			continue;
 		}
diff --git a/Engine/Source/Runtime/Landscape/Private/LandscapeRender.cpp b/Engine/Source/Runtime/Landscape/Private/LandscapeRender.cpp
--- a/Engine/Source/Runtime/Landscape/Private/LandscapeRender.cpp
+++ b/Engine/Source/Runtime/Landscape/Private/LandscapeRender.cpp
@@ -3120,13 +3120,13 @@ void FLandscapeComponentSceneProxy::GetDynamicRayTracingInstances(FRayTracingIns
 		// this is really not normal that we have no material at this point, so do not continue
 		if (SelectedMaterial == nullptr)
 		{
-			return;
+			continue;
 		}
 
 		int32 RayTracingSectionLODVertCount = GetRayTracingSectionVertCountForLOD(SubsectionSizeVerts, NumRayTracingSections, CurrentLODIndex);
 		if (RayTracingSectionLODVertCount <= 0)
 		{
-			return;
+			continue;
 		}
 
 		if (!RayTracingImpl.IsValid())
