Landscape Editor - Remove ULandscapeInfo entries when levels are removed to ensure landscape infos are not reused between level instances
* ALandscapeProxy::PostLoad expects that there is no level instance registered on initial level load. Without removing entires, shared data may incorrectly get fixed up before the ALandscape is registered

diff --git a/Engine/Source/Runtime/Landscape/Classes/LandscapeInfo.h b/Engine/Source/Runtime/Landscape/Classes/LandscapeInfo.h
--- a/Engine/Source/Runtime/Landscape/Classes/LandscapeInfo.h
+++ b/Engine/Source/Runtime/Landscape/Classes/LandscapeInfo.h
@@ -371,6 +371,7 @@ public:
 #endif
 	LANDSCAPE_API static ULandscapeInfo* Find(UWorld* InWorld, const FGuid& LandscapeGuid);
 	LANDSCAPE_API static ULandscapeInfo* FindOrCreate(UWorld* InWorld, const FGuid& LandscapeGuid);
+	LANDSCAPE_API static int32 RemoveLandscapeInfo(UWorld* InWorld, const FGuid& LandscapeGuid);
 
 	/** Called after creating object so that it can initialize its state */
 	void Initialize(UWorld* InWorld, const FGuid& InLandscapeGuid);
diff --git a/Engine/Source/Runtime/Landscape/Private/Landscape.cpp b/Engine/Source/Runtime/Landscape/Private/Landscape.cpp
--- a/Engine/Source/Runtime/Landscape/Private/Landscape.cpp
+++ b/Engine/Source/Runtime/Landscape/Private/Landscape.cpp
@@ -5724,7 +5724,7 @@ ULandscapeInfo* ULandscapeInfo::Find(UWorld* InWorld, const FGuid& LandscapeGuid
 	ULandscapeInfo* LandscapeInfo = nullptr;
 	if (InWorld != nullptr && LandscapeGuid.IsValid())
 	{
-		auto& LandscapeInfoMap = ULandscapeInfoMap::GetLandscapeInfoMap(InWorld);
+		ULandscapeInfoMap& LandscapeInfoMap = ULandscapeInfoMap::GetLandscapeInfoMap(InWorld);
 		LandscapeInfo = LandscapeInfoMap.Map.FindRef(LandscapeGuid);
 	}
 	return LandscapeInfo;
@@ -5737,7 +5737,7 @@ ULandscapeInfo* ULandscapeInfo::FindOrCreate(UWorld* InWorld, const FGuid& Lands
 	check(LandscapeGuid.IsValid());
 	check(InWorld);
 
-	auto& LandscapeInfoMap = ULandscapeInfoMap::GetLandscapeInfoMap(InWorld);
+	ULandscapeInfoMap& LandscapeInfoMap = ULandscapeInfoMap::GetLandscapeInfoMap(InWorld);
 	LandscapeInfo = LandscapeInfoMap.Map.FindRef(LandscapeGuid);
 
 	if (!LandscapeInfo)
@@ -5751,6 +5751,15 @@ ULandscapeInfo* ULandscapeInfo::FindOrCreate(UWorld* InWorld, const FGuid& Lands
 	return LandscapeInfo;
 }
 
+int32 ULandscapeInfo::RemoveLandscapeInfo(UWorld* InWorld, const FGuid& LandscapeGuid)
+{
+	check(LandscapeGuid.IsValid());
+	check(InWorld);
+
+	ULandscapeInfoMap& LandscapeInfoMap = ULandscapeInfoMap::GetLandscapeInfoMap(InWorld);
+	return LandscapeInfoMap.Map.Remove(LandscapeGuid);
+}
+
 void ULandscapeInfo::Initialize(UWorld* InWorld, const FGuid& InLandscapeGuid)
 {
 	LandscapeGuid = InLandscapeGuid;
@@ -5865,6 +5874,9 @@ void ULandscapeInfo::RegisterLandscapeActorWithProxyInternal(ALandscapeProxy* Pr
 
 	if (ALandscapeStreamingProxy* StreamingProxy = Cast<ALandscapeStreamingProxy>(Proxy))
 	{
+		// If Guids mismatch, PostLoad has failed to fixup properties and this proxy is being registered to the wrong ULandscapeInfo
+		check(Landscape->GetLandscapeGuid() == StreamingProxy->GetLandscapeGuid());
+
 		// streaming proxy specific setup here
 		StreamingProxy->SetLandscapeActor(Landscape);
 
diff --git a/Engine/Source/Runtime/Landscape/Private/LandscapeSubsystem.cpp b/Engine/Source/Runtime/Landscape/Private/LandscapeSubsystem.cpp
--- a/Engine/Source/Runtime/Landscape/Private/LandscapeSubsystem.cpp
+++ b/Engine/Source/Runtime/Landscape/Private/LandscapeSubsystem.cpp
@@ -451,6 +451,8 @@ void ULandscapeSubsystem::Initialize(FSubsystemCollectionBase& Collection)
 		}
 	}
 
+	FWorldDelegates::LevelRemovedFromWorld.AddUObject(this, &ULandscapeSubsystem::OnLevelRemovedFromWorld);
+
 	static IConsoleVariable* NaniteEnabledCVar = IConsoleManager::Get().FindConsoleVariable(TEXT("r.Nanite"));
 	if (NaniteEnabledCVar && !NaniteEnabledCVar->OnChangedDelegate().IsBoundToObject(this))
 	{
@@ -505,6 +507,8 @@ void ULandscapeSubsystem::Deinitialize()
 		OnNaniteWorldSettingsChangedHandle.Reset();
 	}
 
+	FWorldDelegates::LevelRemovedFromWorld.RemoveAll(this);
+
 	static IConsoleVariable* NaniteEnabledCVar = IConsoleManager::Get().FindConsoleVariable(TEXT("r.Nanite"));
 	if (NaniteEnabledCVar)
 	{
@@ -608,6 +612,21 @@ void ULandscapeSubsystem::UnregisterComponent(ULandscapeComponent* Component)
 	}
 }
 
+void ULandscapeSubsystem::OnLevelRemovedFromWorld(ULevel* InLevel, UWorld* InWorld)
+{
+	// Clean up info for any landscape actor in the removed level
+	// ALandscapeProxy::PostLoad expects no ULandscapeInfo to exist until registration during initial level loading
+	for (const ALandscape* Landscape : LandscapeActors)
+	{
+		if (InLevel == Landscape->GetLevel())
+		{
+			int32 NumRemovedPairs = ULandscapeInfo::RemoveLandscapeInfo(InWorld, Landscape->GetLandscapeGuid());
+			// Landscape actors in the level should still be registered and removable here
+			check(NumRemovedPairs != 0);
+		}
+	}
+}
+
 void ULandscapeSubsystem::RemoveGrassInstances(const TSet<ULandscapeComponent*>* ComponentsToRemoveGrassInstances)
 {
 	TRACE_CPUPROFILER_EVENT_SCOPE(ULandscapeSubsystem::RemoveGrassInstances);
diff --git a/Engine/Source/Runtime/Landscape/Public/LandscapeSubsystem.h b/Engine/Source/Runtime/Landscape/Public/LandscapeSubsystem.h
--- a/Engine/Source/Runtime/Landscape/Public/LandscapeSubsystem.h
+++ b/Engine/Source/Runtime/Landscape/Public/LandscapeSubsystem.h
@@ -141,6 +141,8 @@ public:
 	void RegisterComponent(ULandscapeComponent* Component);
 	void UnregisterComponent(ULandscapeComponent* Component);
 
+	void OnLevelRemovedFromWorld(ULevel* InLevel, UWorld* InWorld);
+
 #if WITH_EDITOR
 	UE_DEPRECATED(5.6, "Use the function with the EBuildFlags param")
 	LANDSCAPE_API void BuildAll();
