Landscape - Deprecating non-edit layer based landscapes
* Auto converts non-edit layer based landscapes on load
* Deprecates all ALandscape and ALandscapeProxy CanHaveLayersContent functions and UProperties
* Removes edit layer toggle options from details view and new landscape creation

diff --git a/Engine/Plugins/Experimental/Landmass/Source/Editor/Private/LandmassActor.cpp b/Engine/Plugins/Experimental/Landmass/Source/Editor/Private/LandmassActor.cpp
--- a/Engine/Plugins/Experimental/Landmass/Source/Editor/Private/LandmassActor.cpp
+++ b/Engine/Plugins/Experimental/Landmass/Source/Editor/Private/LandmassActor.cpp
@@ -183,7 +183,7 @@ void ALandmassActor::FindOrSpawnManager()
 
 	for (TActorIterator<ALandscape> LandscapeIterator(GetWorld()); LandscapeIterator; ++LandscapeIterator)
 	{
-		if (LandscapeIterator->CanHaveLayersContent())
+		if (BrushManager && !LandscapeIterator->IsTemplate())
 		{
 			BrushManager->SetTargetLandscape(*LandscapeIterator);
 			break;
diff --git a/Engine/Plugins/Experimental/Landmass/Source/Editor/Private/LandmassErosionBrushBase.cpp b/Engine/Plugins/Experimental/Landmass/Source/Editor/Private/LandmassErosionBrushBase.cpp
--- a/Engine/Plugins/Experimental/Landmass/Source/Editor/Private/LandmassErosionBrushBase.cpp
+++ b/Engine/Plugins/Experimental/Landmass/Source/Editor/Private/LandmassErosionBrushBase.cpp
@@ -57,9 +57,9 @@ void ALandmassErosionBrushBase::SetTargetLandscape(ALandscape* InTargetLandscape
 			return;
 		}
 
-		if (!InTargetLandscape->CanHaveLayersContent())
+		if (InTargetLandscape->IsTemplate())
 		{
-			UE_LOG(LandmassErosionBrush, Warning, TEXT("Landscape target for Landmass Erosion Brush did not have edit layers enabled. Unable to attach manager."));
+			UE_LOG(LandmassErosionBrush, Warning, TEXT("Landscape target for Landmass Erosion Brush is a template. Unable to attach manager."));
 			SetOwningLandscape(nullptr);
 		}
 	}
@@ -88,7 +88,7 @@ void ALandmassErosionBrushBase::FindAndAssignLandscape()
 
 	for (TActorIterator<ALandscape> LandscapeIterator(GetWorld()); LandscapeIterator; ++LandscapeIterator)
 	{
-		if (LandscapeIterator->CanHaveLayersContent())
+		if (!LandscapeIterator->IsTemplate())
 		{
 			SetTargetLandscape(*LandscapeIterator);
 			break;
diff --git a/Engine/Plugins/Experimental/Landmass/Source/Editor/Private/LandmassManagerBase.cpp b/Engine/Plugins/Experimental/Landmass/Source/Editor/Private/LandmassManagerBase.cpp
--- a/Engine/Plugins/Experimental/Landmass/Source/Editor/Private/LandmassManagerBase.cpp
+++ b/Engine/Plugins/Experimental/Landmass/Source/Editor/Private/LandmassManagerBase.cpp
@@ -47,9 +47,9 @@ void ALandmassManagerBase::SetTargetLandscape(ALandscape* InTargetLandscape)
 			return;
 		}
 
-		if (!InTargetLandscape->CanHaveLayersContent())
+		if (InTargetLandscape->IsTemplate())
 		{
-			UE_LOG(LandmassManager, Warning, TEXT("Landscape target for Landmass Brush manager did not have edit layers enabled. Unable to attach manager."));
+			UE_LOG(LandmassManager, Warning, TEXT("Landscape target for Landmass Brush manager is a template. Unable to attach manager."));
 			SetOwningLandscape(nullptr);
 		}
 	}
diff --git a/Engine/Plugins/Experimental/LandscapePatch/Source/LandscapePatch/Private/LandscapePatchComponent.cpp b/Engine/Plugins/Experimental/LandscapePatch/Source/LandscapePatch/Private/LandscapePatchComponent.cpp
--- a/Engine/Plugins/Experimental/LandscapePatch/Source/LandscapePatch/Private/LandscapePatchComponent.cpp
+++ b/Engine/Plugins/Experimental/LandscapePatch/Source/LandscapePatch/Private/LandscapePatchComponent.cpp
@@ -212,7 +212,7 @@ void ULandscapePatchComponent::CheckForErrors()
 	if (EditLayerGuid.IsValid())
 	{
 		const ULandscapeEditLayerBase* LocalEditLayer = nullptr;
-		if (!Landscape.IsValid() || !Landscape->CanHaveLayersContent() || (LocalEditLayer = Landscape->GetEditLayerConst(EditLayerGuid)) == nullptr)
+		if (!Landscape.IsValid() || Landscape->IsTemplate() || (LocalEditLayer = Landscape->GetEditLayerConst(EditLayerGuid)) == nullptr)
 		{
 			FMessageLog("MapCheck").Warning()
 			->AddToken(FUObjectToken::Create(this))
@@ -244,21 +244,7 @@ void ULandscapePatchComponent::CheckForErrors()
 
 	if (!PatchManager.IsValid() && !EditLayerGuid.IsValid())
 	{
-		if (Landscape.IsValid() && !Landscape->CanHaveLayersContent())
-		{
-			FMessageLog("MapCheck").Warning()
-			->AddToken(FUObjectToken::Create(this))
-			->AddToken(FTextToken::Create(FText::Format(LOCTEXT("LandscapeEditLayersNotEnabled", "The patch points "
-				"to a landscape that does not have edit layers enabled. (Package: {Package}, Actor: {Actor}). "
-				"Fix individually or run LandscapePatch.FixPatchBindings."), GetPackageAndActorArgs())))
-			->AddToken(FActionToken::Create(FixPatchBindingsText, FText(),
-				FOnActionTokenExecuted::CreateWeakLambda(this, [this]() 
-			{
-				const FScopedTransaction Transaction(FixPatchBindingsTransactionName);
-				FixBindings();
-			})));
-		}
-		else if (Landscape.IsValid())
+		if (Landscape.IsValid())
 		{
 			FMessageLog("MapCheck").Warning()
 			->AddToken(FUObjectToken::Create(this))
@@ -918,7 +904,7 @@ bool ULandscapePatchComponent::BindToLandscape(ALandscape* LandscapeIn)
 		return true;
 	}
 
-	if (!LandscapeIn->CanHaveLayersContent())
+	if (LandscapeIn->IsTemplate())
 	{
 		ResetPatchManager();
 		ResetEditLayer();
@@ -995,7 +981,7 @@ bool ULandscapePatchComponent::BindToAnyLandscape()
 
 	for (TActorIterator<ALandscape> LandscapeIterator(World); LandscapeIterator; ++LandscapeIterator)
 	{
-		if (LandscapeIterator->CanHaveLayersContent() && BindToLandscape(*LandscapeIterator))
+		if (BindToLandscape(*LandscapeIterator))
 		{
 			return true;
 		}
diff --git a/Engine/Plugins/Experimental/LandscapePatch/Source/LandscapePatch/Private/LandscapePatchManager.cpp b/Engine/Plugins/Experimental/LandscapePatch/Source/LandscapePatch/Private/LandscapePatchManager.cpp
--- a/Engine/Plugins/Experimental/LandscapePatch/Source/LandscapePatch/Private/LandscapePatchManager.cpp
+++ b/Engine/Plugins/Experimental/LandscapePatch/Source/LandscapePatch/Private/LandscapePatchManager.cpp
@@ -317,9 +317,9 @@ void ALandscapePatchManager::SetTargetLandscape(ALandscape* InTargetLandscape)
 			return;
 		}
 
-		if (!InTargetLandscape->CanHaveLayersContent())
+		if (InTargetLandscape->IsTemplate())
 		{
-			UE_LOG(LogLandscapePatch, Warning, TEXT("Landscape target for patch manager did not have edit layers enabled. Unable to attach manager."));
+			UE_LOG(LogLandscapePatch, Warning, TEXT("Landscape target for patch manager is a template. Unable to attach manager."));
 			SetOwningLandscape(nullptr);
 		}
 	}
diff --git a/Engine/Plugins/Experimental/Water/Source/Editor/Private/WaterLandscapeBrush.cpp b/Engine/Plugins/Experimental/Water/Source/Editor/Private/WaterLandscapeBrush.cpp
--- a/Engine/Plugins/Experimental/Water/Source/Editor/Private/WaterLandscapeBrush.cpp
+++ b/Engine/Plugins/Experimental/Water/Source/Editor/Private/WaterLandscapeBrush.cpp
@@ -481,7 +481,7 @@ void AWaterLandscapeBrush::SetTargetLandscape(ALandscape* InTargetLandscape)
 			OwningLandscape->RemoveBrush(this);
 		}
 
-		if (InTargetLandscape && InTargetLandscape->CanHaveLayersContent())
+		if (InTargetLandscape && !InTargetLandscape->IsTemplate())
 		{
 			static const FName WaterLayerName = FName("Water");
 			
@@ -557,7 +557,7 @@ AWaterLandscapeBrush::EWaterBrushStatus AWaterLandscapeBrush::CheckWaterBrushSta
 	if (GetWorld() && !IsTemplate())
 	{
 		ALandscape* Landscape = GetOwningLandscape();
-		if (Landscape == nullptr || !Landscape->CanHaveLayersContent())
+		if (Landscape == nullptr)
 		{
 			return EWaterBrushStatus::MissingLandscapeWithEditLayers;
 		}
@@ -580,8 +580,8 @@ void AWaterLandscapeBrush::CheckForErrors()
 	case EWaterBrushStatus::MissingLandscapeWithEditLayers:
 		FMessageLog("MapCheck").Error()
 			->AddToken(FUObjectToken::Create(this))
-			->AddToken(FTextToken::Create(LOCTEXT("MapCheck_Message_NonEditLayersLandscape", "The water brush requires a Landscape with Edit Layers enabled.")))
-			->AddToken(FMapErrorToken::Create(TEXT("WaterBrushNonEditLayersLandscape")));
+			->AddToken(FTextToken::Create(LOCTEXT("MapCheck_Message_NoLandscape", "The water brush requires an edit layer Landscape.")))
+			->AddToken(FMapErrorToken::Create(TEXT("WaterBrushNoLandscape")));
 	case EWaterBrushStatus::MissingFromLandscapeEditLayers:
 		FMessageLog("MapCheck").Error()
 			->AddToken(FUObjectToken::Create(this))
diff --git a/Engine/Source/Editor/LandscapeEditor/Private/LandscapeEdMode.cpp b/Engine/Source/Editor/LandscapeEditor/Private/LandscapeEdMode.cpp
--- a/Engine/Source/Editor/LandscapeEditor/Private/LandscapeEdMode.cpp
+++ b/Engine/Source/Editor/LandscapeEditor/Private/LandscapeEdMode.cpp
@@ -220,10 +220,7 @@ void FEdModeLandscape::UpdateToolModes()
 	ToolMode_Manage->ValidTools.Add(TEXT("ResizeLandscape"));
 	ToolMode_Manage->ValidTools.Add(TEXT("Splines"));
 	ToolMode_Manage->ValidTools.Add(TEXT("ImportExport"));
-	if (CanHaveLandscapeLayersContent())
-	{
-		ToolMode_Manage->ValidTools.Add(TEXT("BlueprintBrush"));
-	}
+	ToolMode_Manage->ValidTools.Add(TEXT("BlueprintBrush"));
 	
 	// Restore
 	FName* PreviousToolName = PreviousTools.Find(ToolMode_Manage->ToolModeName);
@@ -235,10 +232,7 @@ void FEdModeLandscape::UpdateToolModes()
 
 	FLandscapeToolMode* ToolMode_Sculpt = new(LandscapeToolModes)FLandscapeToolMode(TEXT("ToolMode_Sculpt"), ELandscapeToolTargetTypeFlags::Heightmap | ELandscapeToolTargetTypeFlags::Visibility);
 	ToolMode_Sculpt->ValidTools.Add(TEXT("Sculpt"));
-	if (CanHaveLandscapeLayersContent())
-	{
-		ToolMode_Sculpt->ValidTools.Add(TEXT("Erase"));
-	}
+	ToolMode_Sculpt->ValidTools.Add(TEXT("Erase"));
 	ToolMode_Sculpt->ValidTools.Add(TEXT("Smooth"));
 	ToolMode_Sculpt->ValidTools.Add(TEXT("Flatten"));
 	ToolMode_Sculpt->ValidTools.Add(TEXT("Ramp"));
@@ -291,12 +285,6 @@ TSharedRef<FUICommandList> FEdModeLandscape::GetUICommandList() const
 	return Toolkit->GetToolkitCommands();
 }
 
-void FEdModeLandscape::OnCanHaveLayersContentChanged()
-{
-	RefreshDetailPanel();
-	UpdateToolModes();
-}
-
 void FEdModeLandscape::PostUpdateLayerContent()
 {
 	if (bNeedsUpdateLayerUsageInformation)
@@ -348,7 +336,7 @@ void FEdModeLandscape::SetLandscapeInfo(ULandscapeInfo* InLandscapeInfo)
 
 			// Only update the selected edit layer index for the initial landscape, 
 			// this avoids setting the last saved index to 0 when switching between existing landscapes
-			if (bIsInitialLandscapeInfo && CanHaveLandscapeLayersContent())
+			if (bIsInitialLandscapeInfo)
 			{
 				SetSelectedEditLayer(0);
 			}
@@ -506,11 +494,8 @@ void FEdModeLandscape::Enter()
 		if (ALandscape* Landscape = GetLandscape())
 		{
 			Landscape->OnBlueprintBrushChangedDelegate().AddRaw(this, &FEdModeLandscape::RefreshDetailPanel);
-			if (Landscape->HasLayersContent())
-			{
-				Landscape->RequestSplineLayerUpdate();
-				Landscape->RequestLayersContentUpdateForceAll();
-			}
+			Landscape->RequestSplineLayerUpdate();
+			Landscape->RequestLayersContentUpdateForceAll();
 		}
 	}
 
@@ -2500,6 +2485,7 @@ int32 FEdModeLandscape::UpdateLandscapeList()
 			FName CurrentToolName = CurrentTool != nullptr ? CurrentTool->GetToolName() : FName();
 			SetLandscapeInfo(LandscapeList[0].Info);
 			CurrentIndex = 0;
+			SetSelectedEditLayer(0);
 
 			UpdateLayerUsageInformation();
 						
@@ -2516,10 +2502,6 @@ int32 FEdModeLandscape::UpdateLandscapeList()
 			SetCurrentTool("NewLandscape");
 		}
 	}
-	else if (!HasValidLandscapeEditLayerSelection())
-	{
-		SetSelectedEditLayer(0);
-	}
 
 	if (!CanEditCurrentTarget())
 	{
@@ -2535,11 +2517,6 @@ bool FEdModeLandscape::IsGridBased() const
 	return GetWorld()->GetSubsystem<ULandscapeSubsystem>()->IsGridBased();
 }
 
-bool FEdModeLandscape::HasValidLandscapeEditLayerSelection() const
-{
-	return !CanHaveLandscapeLayersContent() || GetCurrentEditLayerConst();
-}
-
 void FEdModeLandscape::SetTargetLandscape(const TWeakObjectPtr<ULandscapeInfo>& InLandscapeInfo)
 {
 	if ((CurrentToolTarget.LandscapeInfo == InLandscapeInfo) || !InLandscapeInfo.IsValid())
@@ -2594,20 +2571,16 @@ bool FEdModeLandscape::CanEditCurrentTarget(FText* Reason) const
 	{
 		if (!Proxy)
 		{
-            LocalReason = NSLOCTEXT("UnrealEd", "LandscapeNotFound", "No Landscape found.");
-			return false;
-		}
-		
-		if (Proxy->HasLayersContent())
-		{
-			LocalReason = NSLOCTEXT("UnrealEd", "LandscapeActorNotLoaded", "Landscape actor is not loaded. It is needed to do layer editing.");
+           		LocalReason = NSLOCTEXT("UnrealEd", "LandscapeNotFound", "No Landscape found.");
 			return false;
 		}
 
+		LocalReason = NSLOCTEXT("UnrealEd", "LandscapeActorNotLoaded", "Landscape actor is not loaded. An edit layer landscape is required for editing.");
+		return false;
 	}
 	
 	// To edit a Landscape with layers all components need to be registered
-	if(Proxy && Proxy->HasLayersContent() && !CurrentToolTarget.LandscapeInfo->AreAllComponentsRegistered())
+	if(Proxy && !CurrentToolTarget.LandscapeInfo->AreAllComponentsRegistered())
 	{
 		LocalReason = NSLOCTEXT("UnrealEd", "LandscapeMissingComponents", "Landscape has some unregistered components (hidden levels).");
 		return false;
@@ -2822,11 +2795,6 @@ void FEdModeLandscape::OnLandscapeMaterialChangedDelegate(ALandscapeProxy* InPro
 void FEdModeLandscape::RequestUpdateLayerUsageInformation()
 {
 	bNeedsUpdateLayerUsageInformation = true;
-
-	if (CurrentToolTarget.LandscapeInfo.IsValid() && !CurrentToolTarget.LandscapeInfo->CanHaveLayersContent())
-	{
-		UpdateLayerUsageInformation(); // do it synchronously when not in edit layers mode
-	}
 }
 
 void FEdModeLandscape::UpdateLayerUsageInformation(TWeakObjectPtr<ULandscapeLayerInfoObject>* LayerInfoObjectThatChanged)
@@ -3724,51 +3692,38 @@ ALandscape* FEdModeLandscape::ChangeComponentSetting(int32 NumComponentsX, int32
 					}
 				};
 
-				if (HasLandscapeLayersContent())
-				{
-					int32 HeightCount = 0;
+				int32 HeightCount = 0;
 
-					for (const ULandscapeEditLayerBase* OldEditLayer : OldLandscape->GetEditLayersConst())
-					{
-						FScopedSetLandscapeEditingLayer Scope(OldLandscape, OldEditLayer->GetGuid());
-
-						TArray<uint16> HeightData;
-						TArray<FLandscapeImportLayerInfo> ImportMaterialLayerInfos;
-
-						ExtractHeightmapWeightmapContent(HeightData, ImportMaterialLayerInfos);
+				for (const ULandscapeEditLayerBase* OldEditLayer : OldLandscape->GetEditLayersConst())
+				{
+					FScopedSetLandscapeEditingLayer Scope(OldLandscape, OldEditLayer->GetGuid());
 
-						HeightCount = FMath::Max(HeightCount, HeightData.Num());
-						HeightDataPerLayers.Add(OldEditLayer->GetGuid(), MoveTemp(HeightData));
-						ImportMaterialLayerInfosPerLayers.Add(OldEditLayer->GetGuid(), MoveTemp(ImportMaterialLayerInfos));
-					}
+					TArray<uint16> HeightData;
+					TArray<FLandscapeImportLayerInfo> ImportMaterialLayerInfos;
 
-					TArray<uint16> DefaultHeightData;
-					DefaultHeightData.AddUninitialized(HeightCount);
+					ExtractHeightmapWeightmapContent(HeightData, ImportMaterialLayerInfos);
 
-					uint16 DefaultValue = LandscapeDataAccess::GetTexHeight(0.f);
+					HeightCount = FMath::Max(HeightCount, HeightData.Num());
+					HeightDataPerLayers.Add(OldEditLayer->GetGuid(), MoveTemp(HeightData));
+					ImportMaterialLayerInfosPerLayers.Add(OldEditLayer->GetGuid(), MoveTemp(ImportMaterialLayerInfos));
+				}
 
-					// Initialize blank heightmap data
-					for (int32 i = 0; i < DefaultHeightData.Num(); i++)
-					{
-						DefaultHeightData[i] = DefaultValue;
-					}
+				TArray<uint16> DefaultHeightData;
+				DefaultHeightData.AddUninitialized(HeightCount);
 
-					HeightDataPerLayers.Add(FGuid(), DefaultHeightData);
+				uint16 DefaultValue = LandscapeDataAccess::GetTexHeight(0.f);
 
-					TArray<FLandscapeImportLayerInfo> EmptyImportLayer;
-					ImportMaterialLayerInfosPerLayers.Add(FGuid(), EmptyImportLayer);
-				}
-				else
+				// Initialize blank heightmap data
+				for (int32 i = 0; i < DefaultHeightData.Num(); i++)
 				{
-					TArray<uint16> HeightData;
-					TArray<FLandscapeImportLayerInfo> ImportMaterialLayerInfos;
-
-					ExtractHeightmapWeightmapContent(HeightData, ImportMaterialLayerInfos);
-
-					HeightDataPerLayers.Add(FGuid(), MoveTemp(HeightData));
-					ImportMaterialLayerInfosPerLayers.Add(FGuid(), MoveTemp(ImportMaterialLayerInfos));
+					DefaultHeightData[i] = DefaultValue;
 				}
 
+				HeightDataPerLayers.Add(FGuid(), DefaultHeightData);
+
+				TArray<FLandscapeImportLayerInfo> EmptyImportLayer;
+				ImportMaterialLayerInfosPerLayers.Add(FGuid(), EmptyImportLayer);
+				
 				if (!bResample)
 				{
 					NewMinX = 0;
@@ -3784,7 +3739,6 @@ ALandscape* FEdModeLandscape::ChangeComponentSetting(int32 NumComponentsX, int32
 			FActorSpawnParameters SpawnParams;
 			SpawnParams.OverrideLevel = OldLandscape->GetLevel();
 			NewLandscape = OldLandscape->GetWorld()->SpawnActor<ALandscape>(Location, OldLandscape->GetActorRotation(), SpawnParams);
-			NewLandscape->bCanHaveLayersContent = OldLandscape->bCanHaveLayersContent;
 			NewLandscape->bAreNewLandscapeActorsSpatiallyLoaded = OldLandscape->bAreNewLandscapeActorsSpatiallyLoaded;
 			NewLandscape->bIncludeGridSizeInNameForLandscapeActors = OldLandscape->bIncludeGridSizeInNameForLandscapeActors;
 			NewLandscape->bUseGeneratedLandscapeSplineMeshesActors = OldLandscape->bUseGeneratedLandscapeSplineMeshesActors;
@@ -3804,10 +3758,7 @@ ALandscape* FEdModeLandscape::ChangeComponentSetting(int32 NumComponentsX, int32
  			NewLandscape->SetLandscapeGuid(FGuid(), /* bValidateGuid= */ false);
 
 			TArrayView<const FLandscapeLayer> LandscapeLayers;
-			if (CanHaveLandscapeLayersContent())
-			{
-				LandscapeLayers = OldLandscape->GetLayersConst();
-			} 
+			LandscapeLayers = OldLandscape->GetLayersConst();
 			NewLandscape->Import(FGuid::NewGuid(), NewMinX, NewMinY, NewMaxX, NewMaxY, NumSubsections, SubsectionSizeQuads, HeightDataPerLayers, *OldLandscape->ReimportHeightmapFilePath, ImportMaterialLayerInfosPerLayers, ELandscapeImportAlphamapType::Additive, LandscapeLayers);
 
 			ULandscapeInfo* NewLandscapeInfo = NewLandscape->GetLandscapeInfo();
@@ -3989,21 +3940,9 @@ ELandscapeEditingState FEdModeLandscape::GetEditingState() const
 	return ELandscapeEditingState::Enabled;
 }
 
-bool FEdModeLandscape::CanHaveLandscapeLayersContent() const
-{
-	ALandscape* Landscape = GetLandscape();
-	return Landscape != nullptr ? Landscape->CanHaveLayersContent() : false;
-}
-
-bool FEdModeLandscape::HasLandscapeLayersContent() const
-{
-	ALandscape* Landscape = GetLandscape();	
-	return Landscape != nullptr ? Landscape->HasLayersContent() : false;
-}
-
 bool FEdModeLandscape::HasSplinesEditLayer() const
 {
-	if (ALandscape* Landscape = GetLandscape(); Landscape != nullptr && CanHaveLandscapeLayersContent())
+	if (ALandscape* Landscape = GetLandscape(); Landscape != nullptr)
 	{
 		const ULandscapeEditLayerBase* SplinesLayer = Landscape->FindEditLayerOfTypeConst(ULandscapeEditLayerSplines::StaticClass());
 		return SplinesLayer != nullptr;
@@ -4025,9 +3964,7 @@ void FEdModeLandscape::SetSelectedEditLayer(int32 InLayerIndex)
 		Landscape->SetSelectedEditLayerIndex(InLayerIndex);
 		const ULandscapeEditLayerBase* EditLayer = Landscape->GetEditLayerConst(InLayerIndex);
 
-		if (Landscape->HasLayersContent() 
-			&& ensure(EditLayer)
-			&& EditLayer->IsA<ULandscapeEditLayerSplines>())
+		if (EditLayer && EditLayer->IsA<ULandscapeEditLayerSplines>())
 		{
 			SetCurrentToolMode("ToolMode_Manage", false);
 			SetCurrentTool(FName("Splines"));
@@ -4167,7 +4104,7 @@ const ULandscapeEditLayerBase* FEdModeLandscape::GetCurrentEditLayerConst() cons
 
 void FEdModeLandscape::AutoUpdateDirtyLandscapeSplines()
 {
-	if (HasLandscapeLayersContent() && GEditor->IsTransactionActive())
+	if (GEditor->IsTransactionActive())
 	{
 		// Only auto-update if a layer is reserved for landscape splines
 		ALandscape* Landscape = GetLandscape();
@@ -4181,68 +4118,65 @@ void FEdModeLandscape::AutoUpdateDirtyLandscapeSplines()
 
 bool FEdModeLandscape::CanEditLayer(FText* Reason /*=nullptr*/, const ULandscapeEditLayerBase* InLayer /*= nullptr*/)
 {
-	if (CanHaveLandscapeLayersContent())
+	ALandscape* Landscape = GetLandscape();
+	const ULandscapeEditLayerBase* TargetLayer = InLayer ? InLayer : GetCurrentEditLayerConst();
+	if (!TargetLayer)
 	{
-		ALandscape* Landscape = GetLandscape();
-		const ULandscapeEditLayerBase* TargetLayer = InLayer ? InLayer : GetCurrentEditLayerConst();
-		if (!TargetLayer)
+		if (Reason)
 		{
-			if (Reason)
-			{
-				*Reason = NSLOCTEXT("UnrealEd", "LandscapeInvalidLayer", "No layer selected. You must first chose a layer to work on.");
-			}
-			return false;
+			*Reason = NSLOCTEXT("UnrealEd", "LandscapeInvalidLayer", "No layer selected. You must first chose a layer to work on.");
 		}
+		return false;
+	}
 
-		check(TargetLayer != nullptr);
-		if (!TargetLayer->IsVisible())
+	check(TargetLayer != nullptr);
+	if (!TargetLayer->IsVisible())
+	{
+		if (Reason)
 		{
-			if (Reason)
-			{
-				*Reason = NSLOCTEXT("UnrealEd", "LandscapeLayerHidden", "Painting or sculpting in a hidden layer is not allowed.");
-			}
-			return false;
+			*Reason = NSLOCTEXT("UnrealEd", "LandscapeLayerHidden", "Painting or sculpting in a hidden layer is not allowed.");
 		}
-		else if (TargetLayer->IsLocked())
+		return false;
+	}
+	else if (TargetLayer->IsLocked())
+	{
+		if (Reason)
 		{
-			if (Reason)
-			{
-				*Reason = NSLOCTEXT("UnrealEd", "LandscapeLayerLocked", "This layer is locked. You must unlock it before you can work on this layer.");
-			}
-			return false;
+			*Reason = NSLOCTEXT("UnrealEd", "LandscapeLayerLocked", "This layer is locked. You must unlock it before you can work on this layer.");
 		}
-		else if (CurrentTool)
+		return false;
+	}
+	else if (CurrentTool)
+	{
+		// Special case for the splines tool, which is supported on the splines layer : 
+		if (CurrentTool != (FLandscapeTool*)SplinesTool)
 		{
-			// Special case for the splines tool, which is supported on the splines layer : 
-			if (CurrentTool != (FLandscapeTool*)SplinesTool)
+			if (Landscape && (TargetLayer == Landscape->FindEditLayerOfTypeConst(ULandscapeEditLayerSplines::StaticClass())))
 			{
-				if (Landscape && (TargetLayer == Landscape->FindEditLayerOfTypeConst(ULandscapeEditLayerSplines::StaticClass())))
-				{
-					if (Reason)
-					{
-						*Reason = NSLOCTEXT("UnrealEd", "LandscapeLayerReservedForSplines", "This layer is reserved for Landscape Splines.");
-					}
-					return false;
-				}
-
-				if (!TargetLayer->SupportsEditingTools())
+				if (Reason)
 				{
-					if (Reason)
-					{
-						*Reason = FText::Format(NSLOCTEXT("UnrealEd", "LandscapeLayerEditLayerDoesntSupportEditing", "This layer's type ({0}) doesn't support direct editing."), TargetLayer->GetClass()->GetDisplayNameText());
-					}
-					return false;
+					*Reason = NSLOCTEXT("UnrealEd", "LandscapeLayerReservedForSplines", "This layer is reserved for Landscape Splines.");
 				}
+				return false;
 			}
-			if (CurrentTool->GetToolName() == FName("Retopologize"))
+
+			if (!TargetLayer->SupportsEditingTools())
 			{
 				if (Reason)
 				{
-					*Reason = FText::Format(NSLOCTEXT("UnrealEd", "LandscapeLayersNoSupportForRetopologize", "{0} Tool is not available with the Landscape Edit Layer System and will be entirely deprecated for all landscape types in UE5.6."), CurrentTool->GetDisplayName());
+					*Reason = FText::Format(NSLOCTEXT("UnrealEd", "LandscapeLayerEditLayerDoesntSupportEditing", "This layer's type ({0}) doesn't support direct editing."), TargetLayer->GetClass()->GetDisplayNameText());
 				}
 				return false;
 			}
 		}
+		if (CurrentTool->GetToolName() == FName("Retopologize"))
+		{
+			if (Reason)
+			{
+				*Reason = FText::Format(NSLOCTEXT("UnrealEd", "LandscapeLayersNoSupportForRetopologize", "{0} Tool is not available with the Landscape Edit Layer System and will be entirely deprecated for all landscape types in UE5.7."), CurrentTool->GetDisplayName());
+			}
+			return false;
+		}
 	}
 
 	return true;
@@ -4274,20 +4208,9 @@ bool FEdModeLandscape::CanEditTargetLayer(FText* Reason /*=nullptr*/, const ULan
 
 void FEdModeLandscape::UpdateLandscapeSplines(bool bUpdateOnlySelected /* = false*/)
 {
-	if (HasLandscapeLayersContent())
-	{
-		ALandscape* Landscape = GetLandscape();
-		if (Landscape)
-		{
-			Landscape->UpdateLandscapeSplines(GetCurrentLayerGuid(), bUpdateOnlySelected);
-		}
-	}
-	else
+	if (ALandscape* Landscape = GetLandscape())
 	{
-		if (CurrentToolTarget.LandscapeInfo.IsValid())
-		{
-			CurrentToolTarget.LandscapeInfo->ApplySplines(bUpdateOnlySelected);
-		}
+		Landscape->UpdateLandscapeSplines(GetCurrentLayerGuid(), bUpdateOnlySelected);
 	}
 }
 
diff --git a/Engine/Source/Editor/LandscapeEditor/Private/LandscapeEdMode.h b/Engine/Source/Editor/LandscapeEditor/Private/LandscapeEdMode.h
--- a/Engine/Source/Editor/LandscapeEditor/Private/LandscapeEdMode.h
+++ b/Engine/Source/Editor/LandscapeEditor/Private/LandscapeEdMode.h
@@ -360,7 +360,8 @@ public:
 	virtual ELandscapeToolTargetType GetLandscapeToolTargetType() const override;
 	virtual const ULandscapeEditLayerBase* GetLandscapeSelectedLayer() const override;
 	virtual ULandscapeLayerInfoObject* GetSelectedLandscapeLayerInfo() const override;
-	virtual void OnCanHaveLayersContentChanged() override;
+	// Deprecated
+	virtual void OnCanHaveLayersContentChanged() final {};
 
 	/** FGCObject interface */
 	virtual void AddReferencedObjects(FReferenceCollector& Collector) override;
@@ -546,9 +547,6 @@ public:
 	bool IsGridBased() const;
 
 	// Edit Layers
-	bool HasValidLandscapeEditLayerSelection() const;
-	bool CanHaveLandscapeLayersContent() const;
-	bool HasLandscapeLayersContent() const;
 	bool HasSplinesEditLayer() const;
 	int32 GetLayerCount() const;
 	void SetSelectedEditLayer(int32 InLayerIndex);
diff --git a/Engine/Source/Editor/LandscapeEditor/Private/LandscapeEdModeComponentTools.cpp b/Engine/Source/Editor/LandscapeEditor/Private/LandscapeEdModeComponentTools.cpp
--- a/Engine/Source/Editor/LandscapeEditor/Private/LandscapeEdModeComponentTools.cpp
+++ b/Engine/Source/Editor/LandscapeEditor/Private/LandscapeEdModeComponentTools.cpp
@@ -649,20 +649,11 @@ public:
 			HeightCache.SetCachedData(X1, Y1, X2, Y2, Data);
 			HeightCache.Flush();
 
-			ALandscape* Landscape = LandscapeInfo->LandscapeActor.Get();
-
-			bool bHasLandscapeLayersContent = Landscape && Landscape->HasLayersContent();
-
-			if (bHasLandscapeLayersContent)
+			if (ALandscape* Landscape = LandscapeInfo->LandscapeActor.Get())
 			{
-				check(Landscape != nullptr); // Landscape actor is required if layer system is enabled
-
 				Landscape->RequestLayersInitialization();
-			}
 
-			for (ULandscapeComponent* NewComponent : NewComponents)
-			{
-				if (bHasLandscapeLayersContent)
+				for (ULandscapeComponent* NewComponent : NewComponents)
 				{
 					TArray<ULandscapeComponent*> ComponentsUsingHeightmap;
 					ComponentsUsingHeightmap.Add(NewComponent);
@@ -673,74 +664,13 @@ public:
 						TMap<UTexture2D*, UTexture2D*> CreatedHeightmapTextures;
 						NewComponent->AddDefaultLayerData(EditLayer->GetGuid(), ComponentsUsingHeightmap, CreatedHeightmapTextures);
 					}
+					
+					// Update Collision
+					NewComponent->UpdateCachedBounds();
+					NewComponent->UpdateBounds();
+					NewComponent->MarkRenderStateDirty();
 				}
 
-				// Update Collision
-				NewComponent->UpdateCachedBounds();
-				NewComponent->UpdateBounds();
-				NewComponent->MarkRenderStateDirty();
-
-				if (!bHasLandscapeLayersContent)
-				{
-					ULandscapeHeightfieldCollisionComponent* CollisionComp = NewComponent->GetCollisionComponent();
-					if (CollisionComp && !bHasXYOffset)
-					{
-						CollisionComp->MarkRenderStateDirty();
-						CollisionComp->RecreateCollision();
-					}
-				}
-			}
-
-			// Add/update "add collision" around the newly added components
-			if (!bHasLandscapeLayersContent)
-			{
-				// Top row
-				int32 ComponentIndexY = ComponentIndexY1 - 1;
-				for (int32 ComponentIndexX = ComponentIndexX1 - 1; ComponentIndexX <= ComponentIndexX2 + 1; ++ComponentIndexX)
-				{
-					if (!LandscapeInfo->XYtoComponentMap.FindRef(FIntPoint(ComponentIndexX, ComponentIndexY)))
-					{
-						LandscapeInfo->UpdateAddCollision(FIntPoint(ComponentIndexX, ComponentIndexY));
-					}
-				}
-
-				// Sides
-				for (ComponentIndexY = ComponentIndexY1; ComponentIndexY <= ComponentIndexY2; ++ComponentIndexY)
-				{
-					// Left
-					int32 ComponentIndexX = ComponentIndexX1 - 1;
-					if (!LandscapeInfo->XYtoComponentMap.FindRef(FIntPoint(ComponentIndexX, ComponentIndexY)))
-					{
-						LandscapeInfo->UpdateAddCollision(FIntPoint(ComponentIndexX, ComponentIndexY));
-					}
-
-					// Right
-					ComponentIndexX = ComponentIndexX1 + 1;
-					if (!LandscapeInfo->XYtoComponentMap.FindRef(FIntPoint(ComponentIndexX, ComponentIndexY)))
-					{
-						LandscapeInfo->UpdateAddCollision(FIntPoint(ComponentIndexX, ComponentIndexY));
-					}
-				}
-
-				// Bottom row
-				ComponentIndexY = ComponentIndexY2 + 1;
-				for (int32 ComponentIndexX = ComponentIndexX1 - 1; ComponentIndexX <= ComponentIndexX2 + 1; ++ComponentIndexX)
-				{
-					if (!LandscapeInfo->XYtoComponentMap.FindRef(FIntPoint(ComponentIndexX, ComponentIndexY)))
-					{
-						LandscapeInfo->UpdateAddCollision(FIntPoint(ComponentIndexX, ComponentIndexY));
-					}
-				}
-			}
-
-			for (ALandscapeProxy* ModifiedProxy : ModifiedProxies)
-			{
-				// Make sure the proxy has an appropriate value for bHasLayersContent after this component creation : 
-				ModifiedProxy->UpdateCachedHasLayersContent();
-			}
-
-			if (Landscape)
-			{
 				GEngine->BroadcastOnActorMoved(Landscape);
 			}
 		}
diff --git a/Engine/Source/Editor/LandscapeEditor/Private/LandscapeEdModeErosionTools.cpp b/Engine/Source/Editor/LandscapeEditor/Private/LandscapeEdModeErosionTools.cpp
--- a/Engine/Source/Editor/LandscapeEditor/Private/LandscapeEdModeErosionTools.cpp
+++ b/Engine/Source/Editor/LandscapeEditor/Private/LandscapeEdModeErosionTools.cpp
@@ -83,7 +83,7 @@ public:
 		}
 
 		ALandscape* Landscape = this->LandscapeInfo->LandscapeActor.Get();
-		bool bCombinedLayerOperation = UISettings->bCombinedLayersOperation && Landscape && Landscape->HasLayersContent();
+		bool bCombinedLayerOperation = UISettings->bCombinedLayersOperation && Landscape;
 
 		// Get list of verts to update
 		FLandscapeBrushData BrushInfo = Brush->ApplyBrush(InteractorPositions);
@@ -340,7 +340,7 @@ public:
 		}
 
 		ALandscape* Landscape = this->LandscapeInfo->LandscapeActor.Get();
-		bool bCombinedLayerOperation = UISettings->bCombinedLayersOperation && Landscape && Landscape->HasLayersContent();
+		bool bCombinedLayerOperation = UISettings->bCombinedLayersOperation && Landscape;
 
 		// Get list of verts to update
 		FLandscapeBrushData BrushInfo = Brush->ApplyBrush(InteractorPositions);
diff --git a/Engine/Source/Editor/LandscapeEditor/Private/LandscapeEdModeMirrorTool.cpp b/Engine/Source/Editor/LandscapeEditor/Private/LandscapeEdModeMirrorTool.cpp
--- a/Engine/Source/Editor/LandscapeEditor/Private/LandscapeEdModeMirrorTool.cpp
+++ b/Engine/Source/Editor/LandscapeEditor/Private/LandscapeEdModeMirrorTool.cpp
@@ -644,21 +644,6 @@ public:
 		TSet<ULandscapeComponent*> Components;
 		if (LandscapeEdit.GetComponentsInRegion(DestMinX, DestMinY, DestMaxX, DestMaxY, &Components) && Components.Num() > 0)
 		{
-			if (!EdMode->HasLandscapeLayersContent())
-			{
-				ALandscapeProxy::InvalidateGeneratedComponentData(Components);
-				for (ULandscapeComponent* Component : Components)
-				{
-					// Recreate collision for modified components and update the navmesh
-					ULandscapeHeightfieldCollisionComponent* CollisionComponent = Component->GetCollisionComponent();
-					if (CollisionComponent)
-					{
-						CollisionComponent->RecreateCollision();
-						FNavigationSystem::UpdateComponentData(*CollisionComponent);
-					}
-				}
-			}
-
 			EdMode->UpdateLayerUsageInformation();
 		}
 	}
diff --git a/Engine/Source/Editor/LandscapeEditor/Private/LandscapeEdModePaintTools.cpp b/Engine/Source/Editor/LandscapeEditor/Private/LandscapeEdModePaintTools.cpp
--- a/Engine/Source/Editor/LandscapeEditor/Private/LandscapeEdModePaintTools.cpp
+++ b/Engine/Source/Editor/LandscapeEditor/Private/LandscapeEdModePaintTools.cpp
@@ -677,7 +677,7 @@ public:
 		if (!this->LandscapeInfo) return;
 
 		ALandscape* Landscape = this->LandscapeInfo->LandscapeActor.Get();
-		const bool bCombinedLayerOperation = bTargetIsHeightmap && UISettings->bCombinedLayersOperation && Landscape && Landscape->HasLayersContent();
+		const bool bCombinedLayerOperation = bTargetIsHeightmap && UISettings->bCombinedLayersOperation && Landscape;
 
 		// Get list of verts to update
 		FLandscapeBrushData BrushInfo = Brush->ApplyBrush(InteractorPositions);
@@ -838,7 +838,7 @@ public:
 		if (!this->LandscapeInfo) return;
 
 		ALandscape* Landscape = this->LandscapeInfo->LandscapeActor.Get();
-		const bool bCombinedLayerOperation = bTargetIsHeightmap && UISettings->bCombinedLayersOperation && Landscape && Landscape->HasLayersContent();
+		const bool bCombinedLayerOperation = bTargetIsHeightmap && UISettings->bCombinedLayersOperation && Landscape;
 
 		// Can't use slope if use Flatten Target because no normal is provided
 		bool bUseSlopeFlatten = UISettings->bUseSlopeFlatten && !UISettings->bUseFlattenTarget;
diff --git a/Engine/Source/Editor/LandscapeEditor/Private/LandscapeEdModeRampTool.cpp b/Engine/Source/Editor/LandscapeEditor/Private/LandscapeEdModeRampTool.cpp
--- a/Engine/Source/Editor/LandscapeEditor/Private/LandscapeEdModeRampTool.cpp
+++ b/Engine/Source/Editor/LandscapeEditor/Private/LandscapeEdModeRampTool.cpp
@@ -545,7 +545,7 @@ public:
 		FLandscapeLayerDataCache<FHeightmapToolTarget> LayerHeightDataCache(EdMode->CurrentToolTarget, HeightCache);
 		LayerHeightDataCache.SetCacheEditingLayer(EditLayerGUID);
 
-		const bool bCombinedLayerOperation = EdMode->UISettings->bCombinedLayersOperation && Landscape && Landscape->HasLayersContent();
+		const bool bCombinedLayerOperation = EdMode->UISettings->bCombinedLayersOperation && Landscape;
 		LayerHeightDataCache.Initialize(EdMode->CurrentToolTarget.LandscapeInfo.Get(), bCombinedLayerOperation);
 
 		// Heights raster
@@ -594,24 +594,6 @@ public:
 
 			// Dirty any runtime virtual textures that our landscape components write to.
 			LandscapeInfo->DirtyRuntimeVirtualTextureForLandscapeArea(MinX, MinY, MaxX, MaxY);
-
-			if (!EdMode->HasLandscapeLayersContent())
-			{
-				TSet<ULandscapeComponent*> Components;
-				if (LandscapeEdit.GetComponentsInRegion(MinX, MinY, MaxX, MaxY, &Components))
-				{
-					for (ULandscapeComponent* Component : Components)
-					{
-						// Recreate collision for modified components and update the navmesh
-						ULandscapeHeightfieldCollisionComponent* CollisionComponent = Component->GetCollisionComponent();
-						if (CollisionComponent)
-						{
-							CollisionComponent->RecreateCollision();
-							FNavigationSystem::UpdateComponentData(*CollisionComponent);
-						}
-					}
-				}
-			}
 		}
  		if (Landscape)
  		{
diff --git a/Engine/Source/Editor/LandscapeEditor/Private/LandscapeEdModeTools.h b/Engine/Source/Editor/LandscapeEditor/Private/LandscapeEdModeTools.h
--- a/Engine/Source/Editor/LandscapeEditor/Private/LandscapeEdModeTools.h
+++ b/Engine/Source/Editor/LandscapeEditor/Private/LandscapeEdModeTools.h
@@ -866,58 +866,9 @@ struct FXYOffsetmapAccessor
 					NewHeights[X - X1 + (Y - Y1) * (X2 - X1 + 1)] = LandscapeDataAccess::GetTexHeight(static_cast<float>(Data[(X - X1 + (Y - Y1) * (X2 - X1 + 1))].Z));
 				}
 			}
-						
-			// Notify foliage to move any attached instances
-			bool bUpdateFoliage = false;
-			bool bUpdateNormals = false;
-			
-			if (!LandscapeEdit->HasLandscapeLayersContent())
-			{
-				ALandscapeProxy::InvalidateGeneratedComponentData(Components);
 
-				bUpdateNormals = true;
-				for (ULandscapeComponent* Component : Components)
-				{
-					ULandscapeHeightfieldCollisionComponent* CollisionComponent = Component->GetCollisionComponent();
-					if (CollisionComponent && AInstancedFoliageActor::HasFoliageAttached(CollisionComponent))
-					{
-						bUpdateFoliage = true;
-						break;
-					}
-				}
-			}
-
-			if (bUpdateFoliage)
-			{
-				// Calculate landscape local-space bounding box of old data, to look for foliage instances.
-				TArray<ULandscapeHeightfieldCollisionComponent*> CollisionComponents;
-				CollisionComponents.Empty(Components.Num());
-				TArray<FBox> PreUpdateLocalBoxes;
-				PreUpdateLocalBoxes.Empty(Components.Num());
-
-				for (ULandscapeComponent* Component : Components)
-				{
-					CollisionComponents.Add(Component->GetCollisionComponent());
-					PreUpdateLocalBoxes.Add(FBox(FVector((float)X1, (float)Y1, Component->CachedLocalBox.Min.Z), FVector((float)X2, (float)Y2, Component->CachedLocalBox.Max.Z)));
-				}
-
-				// Update landscape.
-				LandscapeEdit->SetXYOffsetData(X1, Y1, X2, Y2, Data, 0); // XY Offset always need to be update before the height update
-				LandscapeEdit->SetHeightData(X1, Y1, X2, Y2, NewHeights.GetData(), 0, bUpdateNormals);
-
-				// Snap foliage for each component.
-				for (int32 Index = 0; Index < CollisionComponents.Num(); ++Index)
-				{
-					ULandscapeHeightfieldCollisionComponent* CollisionComponent = CollisionComponents[Index];
-					CollisionComponent->SnapFoliageInstances(PreUpdateLocalBoxes[Index].TransformBy(LandscapeInfo->GetLandscapeProxy()->LandscapeActorToWorld().ToMatrixWithScale()).ExpandBy(1.0f));
-				}
-			}
-			else
-			{
-				// No foliage, just update landscape.
-				LandscapeEdit->SetXYOffsetData(X1, Y1, X2, Y2, Data, 0); // XY Offset always need to be update before the height update
-				LandscapeEdit->SetHeightData(X1, Y1, X2, Y2, NewHeights.GetData(), 0, bUpdateNormals);
-			}
+			LandscapeEdit->SetXYOffsetData(X1, Y1, X2, Y2, Data, 0); // XY Offset always need to be update before the height update
+			LandscapeEdit->SetHeightData(X1, Y1, X2, Y2, NewHeights.GetData(), 0, /*InCalcNormals=*/ false);
 		}
 	}
 
@@ -1020,7 +971,7 @@ public:
 		{
 			LandscapeInfo = InLandscapeInfo;
 			Landscape = LandscapeInfo ? LandscapeInfo->LandscapeActor.Get() : nullptr;
-			bCombinedLayerOperation = Landscape && Landscape->HasLayersContent() && InCombinedLayerOperation && bTargetIsHeightmap;
+			bCombinedLayerOperation = Landscape && InCombinedLayerOperation && bTargetIsHeightmap;
 			if (bCombinedLayerOperation)
 			{
 				int32 I = 0;
@@ -1228,28 +1179,6 @@ struct FFullWeightmapAccessor
 	{
 	}
 
-	~FFullWeightmapAccessor()
-	{
-		if (!LandscapeEdit.HasLandscapeLayersContent())
-		{
-			// Recreate collision for modified components to update the physical materials
-			for (ULandscapeComponent* Component : ModifiedComponents)
-			{
-				ULandscapeHeightfieldCollisionComponent* CollisionComponent = Component->GetCollisionComponent();
-				if (CollisionComponent)
-				{
-					CollisionComponent->RecreateCollision();
-
-					// We need to trigger navigation mesh build, in case user have painted holes on a landscape
-					if (LandscapeInfo->GetLayerInfoIndex(ALandscapeProxy::VisibilityLayer) != INDEX_NONE)
-					{
-						FNavigationSystem::UpdateComponentData(*CollisionComponent);
-					}
-				}
-			}
-		}
-	}
-
 	void SetEditLayer(const FGuid& InEditLayerGUID)
 	{
 		LandscapeEdit.SetEditLayer(InEditLayerGUID);
@@ -1277,11 +1206,6 @@ struct FFullWeightmapAccessor
 		TSet<ULandscapeComponent*> Components;
 		if (LandscapeEdit.GetComponentsInRegion(X1, Y1, X2, Y2, &Components))
 		{
-			if (!LandscapeEdit.HasLandscapeLayersContent())
-			{
-				ALandscapeProxy::InvalidateGeneratedComponentData(Components);
-				ModifiedComponents.Append(Components);
-			}
 			LandscapeEdit.SetAlphaData(DirtyLayerInfos, X1, Y1, X2, Y2, Data, 0, PaintingRestriction);
 		}
 		DirtyLayerInfos.Empty();
@@ -1544,7 +1468,7 @@ public:
 
 	virtual bool ShouldUpdateEditingLayer() const 
 	{ 
-		return AffectsEditLayers() && EdMode->CanHaveLandscapeLayersContent();
+		return AffectsEditLayers();
 	}
 
 	virtual ELandscapeLayerUpdateMode GetBeginToolContentUpdateFlag() const
diff --git a/Engine/Source/Editor/LandscapeEditor/Private/LandscapeEditorDetailCustomization_Blueprint.cpp b/Engine/Source/Editor/LandscapeEditor/Private/LandscapeEditorDetailCustomization_Blueprint.cpp
--- a/Engine/Source/Editor/LandscapeEditor/Private/LandscapeEditorDetailCustomization_Blueprint.cpp
+++ b/Engine/Source/Editor/LandscapeEditor/Private/LandscapeEditorDetailCustomization_Blueprint.cpp
@@ -25,22 +25,18 @@ void FLandscapeEditorDetailCustomization_Blueprint::CustomizeDetails(IDetailLayo
 
 	if (const FEdModeLandscape* LandscapeEdMode = GetEditorMode(); LandscapeEdMode && LandscapeEdMode->GetLandscape() && LandscapeEdMode->CurrentToolMode)
 	{
-		// Tool is disabled for landscapes with edit layers disabled
-		if (LandscapeEdMode->CanHaveLandscapeLayersContent())
-		{
-			const ULandscapeEditLayerBase* EditLayer = LandscapeEdMode->GetCurrentEditLayerConst();
-			check(EditLayer != nullptr);
-			const bool bSupportsBPBrushes = EditLayer->SupportsBlueprintBrushes();
-
-			const FText EnabledToolTip = LOCTEXT("LandscapeBlueprintBrush_EnabledToolTip", "Selects the blueprint brush to apply to the current edit layer. Click on the landscape to apply it.");
-			const FText DisabledToolTip = FText::Format(LOCTEXT("LandscapeBlueprintBrush_DisabledToolTip", "Cannot add blueprint brush : the type of layer {0} ({1}) doesn't support blueprint brushes."),
-				FText::FromName(EditLayer->GetName()), EditLayer->GetClass()->GetDisplayNameText());
-
-			TSharedRef<IPropertyHandle> PropertyHandle_Blueprint = DetailBuilder.GetProperty(GET_MEMBER_NAME_CHECKED(ULandscapeEditorObject, BlueprintBrush));
-			ToolsCategory.AddProperty(PropertyHandle_Blueprint)
-				.IsEnabled(bSupportsBPBrushes)
-				.ToolTip(bSupportsBPBrushes ? EnabledToolTip : DisabledToolTip);
-		}
+		const ULandscapeEditLayerBase* EditLayer = LandscapeEdMode->GetCurrentEditLayerConst();
+		check(EditLayer != nullptr);
+		const bool bSupportsBPBrushes = EditLayer->SupportsBlueprintBrushes();
+
+		const FText EnabledToolTip = LOCTEXT("LandscapeBlueprintBrush_EnabledToolTip", "Selects the blueprint brush to apply to the current edit layer. Click on the landscape to apply it.");
+		const FText DisabledToolTip = FText::Format(LOCTEXT("LandscapeBlueprintBrush_DisabledToolTip", "Cannot add blueprint brush : the type of layer {0} ({1}) doesn't support blueprint brushes."),
+			FText::FromName(EditLayer->GetName()), EditLayer->GetClass()->GetDisplayNameText());
+
+		TSharedRef<IPropertyHandle> PropertyHandle_Blueprint = DetailBuilder.GetProperty(GET_MEMBER_NAME_CHECKED(ULandscapeEditorObject, BlueprintBrush));
+		ToolsCategory.AddProperty(PropertyHandle_Blueprint)
+			.IsEnabled(bSupportsBPBrushes)
+			.ToolTip(bSupportsBPBrushes ? EnabledToolTip : DisabledToolTip);
 	}
 }
 END_SLATE_FUNCTION_BUILD_OPTIMIZATION
diff --git a/Engine/Source/Editor/LandscapeEditor/Private/LandscapeEditorDetailCustomization_ImportExport.cpp b/Engine/Source/Editor/LandscapeEditor/Private/LandscapeEditorDetailCustomization_ImportExport.cpp
--- a/Engine/Source/Editor/LandscapeEditor/Private/LandscapeEditorDetailCustomization_ImportExport.cpp
+++ b/Engine/Source/Editor/LandscapeEditor/Private/LandscapeEditorDetailCustomization_ImportExport.cpp
@@ -272,14 +272,6 @@ void FLandscapeEditorDetailCustomization_ImportExport::CustomizeDetails(IDetailL
 
 	ImportExportCategory.AddProperty(PropertyHandle_ExportEditLayer).Visibility(MakeAttributeLambda([this]()
 	{
-		if (ULandscapeInfo* LandscapeInfo = GetEditorMode()->CurrentToolTarget.LandscapeInfo.Get())
-		{
-			if (!LandscapeInfo->CanHaveLayersContent())
-			{
-				return EVisibility::Collapsed;
-			}
-		}
-
 		return FLandscapeEditorDetailCustomization_ImportExport::GetImportExportVisibility(false);
 	}));
 
diff --git a/Engine/Source/Editor/LandscapeEditor/Private/LandscapeEditorDetailCustomization_NewLandscape.cpp b/Engine/Source/Editor/LandscapeEditor/Private/LandscapeEditorDetailCustomization_NewLandscape.cpp
--- a/Engine/Source/Editor/LandscapeEditor/Private/LandscapeEditorDetailCustomization_NewLandscape.cpp
+++ b/Engine/Source/Editor/LandscapeEditor/Private/LandscapeEditorDetailCustomization_NewLandscape.cpp
@@ -235,13 +235,6 @@ void FLandscapeEditorDetailCustomization_NewLandscape::CustomizeDetails(IDetailL
 		
 	];
 
-	TSharedRef<IPropertyHandle> PropertyHandle_CanHaveLayersContent = DetailBuilder.GetProperty(GET_MEMBER_NAME_CHECKED(ULandscapeEditorObject, bCanHaveLayersContent));
-	NewLandscapeCategory.AddProperty(PropertyHandle_CanHaveLayersContent).Visibility(MakeAttributeLambda([]()
-	{
-		const ULandscapeSettings* Settings = GetDefault<ULandscapeSettings>();
-		return Settings->InRestrictiveMode() ? EVisibility::Hidden : EVisibility::Visible;
-	}));
-
 	TSharedRef<IPropertyHandle> PropertyHandle_FlipYAxis = DetailBuilder.GetProperty(GET_MEMBER_NAME_CHECKED(ULandscapeEditorObject, bFlipYAxis));
 	NewLandscapeCategory.AddProperty(PropertyHandle_FlipYAxis).Visibility(MakeAttributeLambda([]()
 	{ 
@@ -1102,11 +1095,9 @@ void FLandscapeEditorDetailCustomization_NewLandscape::AddComponents(ULandscapeI
 	const bool bHasXYOffset = false;
 	ALandscape* Landscape = InLandscapeInfo->LandscapeActor.Get();
 
-	bool bHasLandscapeLayersContent = Landscape && Landscape->HasLayersContent();
-
 	for (ULandscapeComponent* NewComponent : NewComponents)
 	{
-		if (bHasLandscapeLayersContent)
+		if (Landscape)
 		{
 			TArray<ULandscapeComponent*> ComponentsUsingHeightmap;
 			ComponentsUsingHeightmap.Add(NewComponent);
@@ -1123,19 +1114,8 @@ void FLandscapeEditorDetailCustomization_NewLandscape::AddComponents(ULandscapeI
 		NewComponent->UpdateCachedBounds();
 		NewComponent->UpdateBounds();
 		NewComponent->MarkRenderStateDirty();
-
-		if (!bHasLandscapeLayersContent)
-		{
-			ULandscapeHeightfieldCollisionComponent* CollisionComp = NewComponent->GetCollisionComponent();
-			if (CollisionComp && !bHasXYOffset)
-			{
-				CollisionComp->MarkRenderStateDirty();
-				CollisionComp->RecreateCollision();
-			}
-		}
 	}
 
-
 	if (Landscape)
 	{
 		GEngine->BroadcastOnActorMoved(Landscape);
@@ -1226,7 +1206,6 @@ FReply FLandscapeEditorDetailCustomization_NewLandscape::OnCreateButtonClicked()
 	const FVector Offset = FTransform(UISettings->NewLandscape_Rotation, FVector::ZeroVector, UISettings->NewLandscape_Scale).TransformVector(FVector(-UISettings->NewLandscape_ComponentCount.X * QuadsPerComponent / 2.0, -UISettings->NewLandscape_ComponentCount.Y * QuadsPerComponent / 2.0, 0.0));
 	
 	ALandscape* Landscape = LandscapeEdMode->GetWorld()->SpawnActor<ALandscape>(UISettings->NewLandscape_Location + Offset, UISettings->NewLandscape_Rotation);
-	Landscape->bCanHaveLayersContent = UISettings->bCanHaveLayersContent;
 	Landscape->LandscapeMaterial = UISettings->NewLandscape_Material.Get();
 	Landscape->SetActorRelativeScale3D(UISettings->NewLandscape_Scale);
 
diff --git a/Engine/Source/Editor/LandscapeEditor/Private/LandscapeEditorDetailCustomization_TargetLayers.cpp b/Engine/Source/Editor/LandscapeEditor/Private/LandscapeEditorDetailCustomization_TargetLayers.cpp
--- a/Engine/Source/Editor/LandscapeEditor/Private/LandscapeEditorDetailCustomization_TargetLayers.cpp
+++ b/Engine/Source/Editor/LandscapeEditor/Private/LandscapeEditorDetailCustomization_TargetLayers.cpp
@@ -210,8 +210,7 @@ void FLandscapeEditorDetailCustomization_TargetLayers::CustomizeDetails(IDetailL
 	check(LandscapeEdMode);
 
 	TargetsCategory.AddProperty(PropertyHandle_PaintingRestriction)
-		.Visibility(TAttribute<EVisibility>::Create(TAttribute<EVisibility>::FGetter::CreateStatic(&FLandscapeEditorDetailCustomization_TargetLayers::GetPaintingRestrictionVisibility)))
-		.IsEnabled(TAttribute<bool>::Create(TAttribute<bool>::FGetter::CreateLambda([LandscapeEdMode]() { return LandscapeEdMode->HasValidLandscapeEditLayerSelection(); })));
+		.Visibility(TAttribute<EVisibility>::Create(TAttribute<EVisibility>::FGetter::CreateStatic(&FLandscapeEditorDetailCustomization_TargetLayers::GetPaintingRestrictionVisibility)));
 
 	TargetsCategory.AddProperty(PropertyHandle_TargetLayerAssetFilePath);
 
@@ -224,7 +223,6 @@ void FLandscapeEditorDetailCustomization_TargetLayers::CustomizeDetails(IDetailL
 				.BackgroundColor(FAppStyle::GetColor("ErrorReporting.WarningBackgroundColor"))
 				.Text(LOCTEXT("Visibility_Tip", "Note: There are some areas where visibility painting is disabled because Component/Proxy don't have a \"Landscape Visibility Mask\" node in their material."))
 				.AutoWrapText(true)
-				.IsEnabled(TAttribute<bool>::Create(TAttribute<bool>::FGetter::CreateLambda([LandscapeEdMode]() { return LandscapeEdMode->HasValidLandscapeEditLayerSelection(); })))
 		];
 
 	TargetsCategory.AddCustomBuilder(MakeShareable(new FLandscapeEditorCustomNodeBuilder_TargetLayers(DetailBuilder.GetThumbnailPool().ToSharedRef(), PropertyHandle_TargetDisplayOrder, PropertyHandle_TargetShowUnusedLayers)));
@@ -550,7 +548,6 @@ void FLandscapeEditorCustomNodeBuilder_TargetLayers::GenerateHeaderRowContent(FD
 			]
 		];
 	}
-	NodeRow.IsEnabled(TAttribute<bool>::Create(TAttribute<bool>::FGetter::CreateLambda([LandscapeEdMode]() { return LandscapeEdMode->HasValidLandscapeEditLayerSelection(); })));
 }
 
 const TArray<const TSharedRef<FLandscapeTargetListInfo>> FLandscapeEditorCustomNodeBuilder_TargetLayers::GetUnassignedTargetLayersFromMaterial() const
@@ -1025,8 +1022,7 @@ void FLandscapeEditorCustomNodeBuilder_TargetLayers::GenerateChildContent(IDetai
 		TSharedPtr<SDragAndDropVerticalBox> TargetLayerList = SNew(SDragAndDropVerticalBox)
 			.OnCanAcceptDropAdvanced(this, &FLandscapeEditorCustomNodeBuilder_TargetLayers::HandleCanAcceptDrop)
 			.OnAcceptDrop(this, &FLandscapeEditorCustomNodeBuilder_TargetLayers::HandleAcceptDrop)
-			.OnDragDetected(this, &FLandscapeEditorCustomNodeBuilder_TargetLayers::HandleDragDetected)
-			.IsEnabled(TAttribute<bool>::Create(TAttribute<bool>::FGetter::CreateLambda([LandscapeEdMode]() { return LandscapeEdMode->HasValidLandscapeEditLayerSelection(); })));
+			.OnDragDetected(this, &FLandscapeEditorCustomNodeBuilder_TargetLayers::HandleDragDetected);
 
 		TargetLayerList->SetDropIndicator_Above(*FAppStyle::GetBrush("LandscapeEditor.TargetList.DropZone.Above"));
 		TargetLayerList->SetDropIndicator_Below(*FAppStyle::GetBrush("LandscapeEditor.TargetList.DropZone.Below"));
@@ -2066,8 +2062,7 @@ EVisibility FLandscapeEditorCustomNodeBuilder_TargetLayers::GetLayersSubstractiv
 	FEdModeLandscape* LandscapeEdMode = GetEditorMode();
 	if (LandscapeEdMode != nullptr && Target->TargetType == ELandscapeToolTargetType::Weightmap && Target->LayerInfoObj.IsValid())
 	{
-		ULandscapeEditLayerBase* EditLayer = LandscapeEdMode->GetEditLayer(LandscapeEdMode->GetSelectedEditLayerIndex());
-		if (EditLayer && LandscapeEdMode->CanHaveLandscapeLayersContent())
+		if (ULandscapeEditLayerBase* EditLayer = LandscapeEdMode->GetEditLayer(LandscapeEdMode->GetSelectedEditLayerIndex()))
 		{
 			return EVisibility::Visible;
 		}
diff --git a/Engine/Source/Editor/LandscapeEditor/Private/LandscapeEditorDetails.cpp b/Engine/Source/Editor/LandscapeEditor/Private/LandscapeEditorDetails.cpp
--- a/Engine/Source/Editor/LandscapeEditor/Private/LandscapeEditorDetails.cpp
+++ b/Engine/Source/Editor/LandscapeEditor/Private/LandscapeEditorDetails.cpp
@@ -202,8 +202,8 @@ void FLandscapeEditorDetails::CustomizeDetails(IDetailLayoutBuilder& DetailBuild
 	Customization_AlphaBrush = MakeShareable(new FLandscapeEditorDetailCustomization_AlphaBrush);
 	Customization_AlphaBrush->CustomizeDetails(DetailBuilder);
 
-	// Hide Layers/LayersBrushStack when New Landscape tab is active
-	if (LandscapeEdMode->CanHaveLandscapeLayersContent() && LandscapeEdMode->CurrentTool->GetToolName() != FName("NewLandscape"))
+	// Hide BlueprintTool/Layers/LayersBrushStack when New Landscape tab is active
+	if (LandscapeEdMode->CurrentTool->GetToolName() != FName("NewLandscape"))
 	{
 		// Blueprint tool
 		Customization_Blueprint = MakeShareable(new FLandscapeEditorDetailCustomization_Blueprint);
diff --git a/Engine/Source/Editor/LandscapeEditor/Private/LandscapeEditorModule.cpp b/Engine/Source/Editor/LandscapeEditor/Private/LandscapeEditorModule.cpp
--- a/Engine/Source/Editor/LandscapeEditor/Private/LandscapeEditorModule.cpp
+++ b/Engine/Source/Editor/LandscapeEditor/Private/LandscapeEditorModule.cpp
@@ -39,7 +39,6 @@
 #include "LandscapeSplineDetails.h"
 #include "LandscapeGrassTypeDetails.h"
 #include "LandscapeProxyUIDetails.h"
-#include "LandscapeUIDetails.h"
 #include "LandscapeModule.h"
 
 #include "LevelEditor.h"
@@ -105,7 +104,6 @@ public:
 		PropertyModule.RegisterCustomClassLayout("LandscapeSplineControlPoint", FOnGetDetailCustomizationInstance::CreateStatic(&FLandscapeSplineDetails::MakeInstance));
 		PropertyModule.RegisterCustomClassLayout("LandscapeSplineSegment", FOnGetDetailCustomizationInstance::CreateStatic(&FLandscapeSplineDetails::MakeInstance));
 
-		PropertyModule.RegisterCustomClassLayout("Landscape", FOnGetDetailCustomizationInstance::CreateStatic(&FLandscapeUIDetails::MakeInstance));
 		PropertyModule.RegisterCustomClassLayout("LandscapeProxy", FOnGetDetailCustomizationInstance::CreateStatic(&FLandscapeProxyUIDetails::MakeInstance));
 		PropertyModule.RegisterCustomClassLayout("LandscapeGrassType", FOnGetDetailCustomizationInstance::CreateStatic(&FLandscapeGrassTypeDetails::MakeInstance));
 
@@ -199,7 +197,6 @@ public:
 		PropertyModule.UnregisterCustomClassLayout("LandscapeSplineControlPoint");
 		PropertyModule.UnregisterCustomClassLayout("LandscapeSplineSegment");
 
-		PropertyModule.UnregisterCustomClassLayout("Landscape");
 		PropertyModule.UnregisterCustomClassLayout("LandscapeProxy");
 		PropertyModule.UnregisterCustomClassLayout("LandscapeGrassType");
 
@@ -235,10 +232,7 @@ public:
 		{
 			MenuSection.AddMenuEntryWithCommandList(LandscapeActions.ViewModeLayerDebug, GlobalUICommandList, LOCTEXT("LandscapeViewModeLayerDebug", "Layer Debug"));
 			FEdModeLandscape* LandscapeMode = (FEdModeLandscape*)GLevelEditorModeTools().GetActiveMode(FBuiltinEditorModes::EM_Landscape);
-			if (LandscapeMode->CanHaveLandscapeLayersContent())
-			{
-				MenuSection.AddMenuEntryWithCommandList(LandscapeActions.ViewModeLayerContribution, GlobalUICommandList, LOCTEXT("LandscapeViewModeLayerContribution", "Layer Contribution"));
-			}
+			MenuSection.AddMenuEntryWithCommandList(LandscapeActions.ViewModeLayerContribution, GlobalUICommandList, LOCTEXT("LandscapeViewModeLayerContribution", "Layer Contribution"));
 		}
 		MenuSection.AddMenuEntryWithCommandList(LandscapeActions.ViewModeWireframeOnTop, GlobalUICommandList, LOCTEXT("LandscapeViewModeWireframeOnTop", "Wireframe on Top"));
 	}
diff --git a/Engine/Source/Editor/LandscapeEditor/Private/LandscapeUIDetails.cpp b/Engine/Source/Editor/LandscapeEditor/Private/LandscapeUIDetails.cpp
deleted file mode 100644
diff --git a/Engine/Source/Editor/LandscapeEditor/Private/LandscapeUIDetails.h b/Engine/Source/Editor/LandscapeEditor/Private/LandscapeUIDetails.h
deleted file mode 100644
diff --git a/Engine/Source/Editor/LandscapeEditor/Private/SLandscapeEditor.cpp b/Engine/Source/Editor/LandscapeEditor/Private/SLandscapeEditor.cpp
--- a/Engine/Source/Editor/LandscapeEditor/Private/SLandscapeEditor.cpp
+++ b/Engine/Source/Editor/LandscapeEditor/Private/SLandscapeEditor.cpp
@@ -136,7 +136,7 @@ void FLandscapeToolKit::Init(const TSharedPtr<IToolkitHost>& InitToolkitHost)
 	MAP_MODE("ToolMode_Paint");
 #undef MAP_MODE
 
-#define MAP_TOOL(ToolName) CommandList->MapAction(NameToCommandMap.FindChecked("Tool_" ToolName), FUIAction(FExecuteAction::CreateSP(this, &FLandscapeToolKit::OnChangeTool, FName(ToolName)), FCanExecuteAction::CreateSP(this, &FLandscapeToolKit::IsToolEnabled, FName(ToolName)), FIsActionChecked::CreateSP(this, &FLandscapeToolKit::IsToolActive, FName(ToolName)), FIsActionButtonVisible::CreateSP(this, &FLandscapeToolKit::IsToolAvailable, FName(ToolName))));
+#define MAP_TOOL(ToolName) CommandList->MapAction(NameToCommandMap.FindChecked("Tool_" ToolName), FUIAction(FExecuteAction::CreateSP(this, &FLandscapeToolKit::OnChangeTool, FName(ToolName)), FCanExecuteAction::CreateSP(this, &FLandscapeToolKit::IsToolEnabled, FName(ToolName)), FIsActionChecked::CreateSP(this, &FLandscapeToolKit::IsToolActive, FName(ToolName))));
 	MAP_TOOL("NewLandscape");
 	MAP_TOOL("ResizeLandscape");
 	MAP_TOOL("ImportExport");
@@ -431,7 +431,8 @@ bool FLandscapeToolKit::IsToolEnabled(FName ToolName) const
 
 		if (ToolName == "Retopologize")
 		{
-			return UE::Landscape::Private::bEnableRetopoTool && !LandscapeEdMode->CanHaveLandscapeLayersContent();
+			// TODO [jared.ritchie] this temporarily hides tool but need to fully remove landscape's retopologize tool / XY offset texture
+			return false;
 		}
 
 		// Other tools are available if there is an existing landscape.
@@ -444,21 +445,6 @@ bool FLandscapeToolKit::IsToolEnabled(FName ToolName) const
 	return false;
 }
 
-bool FLandscapeToolKit::IsToolAvailable(FName ToolName) const
-{
-	// Hide Tools that are available in Edit Layers only
-	FEdModeLandscape* LandscapeEdMode = GetEditorMode();
-	if (LandscapeEdMode != nullptr)
-	{
-		if ((ToolName == "BlueprintBrush" || ToolName == "Erase") && !LandscapeEdMode->CanHaveLandscapeLayersContent())
-		{
-			return false;
-		}
-	}
-
-	return true;
-}
-
 bool FLandscapeToolKit::IsToolActive(FName ToolName) const
 {
 	FEdModeLandscape* LandscapeEdMode = GetEditorMode();
@@ -782,14 +768,6 @@ bool FLandscapeToolKit::GetIsPropertyVisibleFromProperty(const FProperty& Proper
 				return false;
 			}
 		}
-		if (Property.HasMetaData("ShowForLandscapeLayerSystem"))
-		{
-			if (!LandscapeEdMode->HasLandscapeLayersContent())
-			{
-				return false;
-			}
-		}
-
 		return true;
 	}
 
diff --git a/Engine/Source/Editor/LandscapeEditor/Private/SLandscapeEditor.h b/Engine/Source/Editor/LandscapeEditor/Private/SLandscapeEditor.h
--- a/Engine/Source/Editor/LandscapeEditor/Private/SLandscapeEditor.h
+++ b/Engine/Source/Editor/LandscapeEditor/Private/SLandscapeEditor.h
@@ -91,7 +91,6 @@ protected:
 	void OnChangeTool(FName ToolName);
 	bool IsToolEnabled(FName ToolName) const;
 	bool IsToolActive(FName ToolName) const;
-	bool IsToolAvailable(FName ToolName) const;
 
 	void OnChangeBrushSet(FName BrushSetName);
 	bool IsBrushSetEnabled(FName BrushSetName) const;
diff --git a/Engine/Source/Editor/LandscapeEditor/Public/LandscapeEditorObject.h b/Engine/Source/Editor/LandscapeEditor/Public/LandscapeEditorObject.h
--- a/Engine/Source/Editor/LandscapeEditor/Public/LandscapeEditorObject.h
+++ b/Engine/Source/Editor/LandscapeEditor/Public/LandscapeEditorObject.h
@@ -320,7 +320,7 @@ class ULandscapeEditorObject : public UObject
 	float MaximumValueRadius;
 
 	// Use the combined result of the underlying layers as input to the operation.  When not checked, it will use only the data in the currently selected layer as input.
-	UPROPERTY(Category="Tool Settings", EditAnywhere, NonTransactional, meta=(ShowForTools="Flatten,Smooth,Erosion,HydraErosion,Ramp", ShowForTargetTypes="Heightmap", ShowForLandscapeLayerSystem))
+	UPROPERTY(Category="Tool Settings", EditAnywhere, NonTransactional, meta=(ShowForTools="Flatten,Smooth,Erosion,HydraErosion,Ramp", ShowForTargetTypes="Heightmap"))
 	bool bCombinedLayersOperation;
 
 	// Flatten Tool:
@@ -620,8 +620,9 @@ private:
 	UPROPERTY(NonTransactional)
 	TArray<uint16> ImportLandscape_Data;
 public:
-	UPROPERTY(Category = "New Landscape", EditAnywhere, NonTransactional, meta = (DisplayName= "Enable Edit Layers", ToolTip="Enable support for landscape edit layers.", ShowForTools= "NewLandscape"))
-	bool bCanHaveLayersContent = true;
+	UE_DEPRECATED(5.7, "Non-edit layer landscapes are deprecated, all landscapes use the edit layer system now.")
+	UPROPERTY(NonTransactional)
+	bool bCanHaveLayersContent_DEPRECATED = true;
 
 	UPROPERTY(Category = "New Landscape", EditAnywhere, NonTransactional, meta = (DisplayName = "Flip Y Axis", ToolTip = "Whether to flip Y coordinate of imported files.", ShowForTools = "NewLandscape,ImportExport"))
 	bool bFlipYAxis = false;
diff --git a/Engine/Source/Editor/WorldBrowser/Private/LevelCollectionModel.h b/Engine/Source/Editor/WorldBrowser/Private/LevelCollectionModel.h
--- a/Engine/Source/Editor/WorldBrowser/Private/LevelCollectionModel.h
+++ b/Engine/Source/Editor/WorldBrowser/Private/LevelCollectionModel.h
@@ -691,7 +691,6 @@ struct FTiledLandscapeImportSettings
 		, TilesCoordinatesOffset(0,0)
 		, SizeX(1009)
 		, bFlipYAxis(true)
-		, bEditLayersEnabled(false)
 	{}
 	
 	FVector				Scale3D;
@@ -704,7 +703,6 @@ struct FTiledLandscapeImportSettings
 	FIntPoint			TilesCoordinatesOffset;
 	int32				SizeX;
 	bool				bFlipYAxis;
-	bool				bEditLayersEnabled;
 
 
 	TWeakObjectPtr<UMaterialInterface>	LandscapeMaterial;
diff --git a/Engine/Source/Editor/WorldBrowser/Private/Tiles/STiledLandscapeImportDlg.cpp b/Engine/Source/Editor/WorldBrowser/Private/Tiles/STiledLandscapeImportDlg.cpp
--- a/Engine/Source/Editor/WorldBrowser/Private/Tiles/STiledLandscapeImportDlg.cpp
+++ b/Engine/Source/Editor/WorldBrowser/Private/Tiles/STiledLandscapeImportDlg.cpp
@@ -98,21 +98,6 @@ void STiledLandscapeImportDlg::Construct(const FArguments& InArgs, TSharedPtr<SW
 					.ToolTipText(LOCTEXT("TiledLandscapeImport_FlipYAxisToolTip", "Whether tile Y coordinate should be flipped (Make sure 'Flip Y-Axis Orientation' option is switched off in World Machine) "))
 				]
 
-				// Enable edit layers
-				+SUniformGridPanel::Slot(0, 2)
-				.VAlign(VAlign_Center)
-				[
-					SNew(STextBlock)
-					.Text(LOCTEXT("TiledLandscapeImport_EnableEditLayers", "Enable Edit Layers"))
-				]
-				+SUniformGridPanel::Slot(1, 2)
-				.VAlign(VAlign_Center)
-				[
-					SNew(SCheckBox)
-					.IsChecked(this, &STiledLandscapeImportDlg::GetEditLayersState)
-					.OnCheckStateChanged(this, &STiledLandscapeImportDlg::OnEditLayersStateChanged)
-					.ToolTipText(LOCTEXT("TiledLandscapeImport_EditLayersToolTip", "Whether landscape edit layers should be enabled"))
-				]
 				// Tiles origin offset
 				+SUniformGridPanel::Slot(0, 3)
 				.VAlign(VAlign_Center)
@@ -431,16 +416,6 @@ void STiledLandscapeImportDlg::OnFlipYAxisStateChanged(ECheckBoxState NewState)
 	ImportSettings.bFlipYAxis = (NewState == ECheckBoxState::Checked);
 }
 
-ECheckBoxState STiledLandscapeImportDlg::GetEditLayersState() const
-{
-	return ImportSettings.bEditLayersEnabled ? ECheckBoxState::Checked : ECheckBoxState::Unchecked;
-}
-
-void STiledLandscapeImportDlg::OnEditLayersStateChanged(ECheckBoxState NewState)
-{
-	ImportSettings.bEditLayersEnabled = (NewState == ECheckBoxState::Checked);
-}
-
 void STiledLandscapeImportDlg::OnSetImportConfiguration(TSharedPtr<FTileImportConfiguration> InTileConfig, ESelectInfo::Type SelectInfo)
 {
 	if (InTileConfig.IsValid())
diff --git a/Engine/Source/Editor/WorldBrowser/Private/Tiles/STiledLandscapeImportDlg.h b/Engine/Source/Editor/WorldBrowser/Private/Tiles/STiledLandscapeImportDlg.h
--- a/Engine/Source/Editor/WorldBrowser/Private/Tiles/STiledLandscapeImportDlg.h
+++ b/Engine/Source/Editor/WorldBrowser/Private/Tiles/STiledLandscapeImportDlg.h
@@ -63,10 +63,6 @@ private:
 	ECheckBoxState GetFlipYAxisState() const;
 	void OnFlipYAxisStateChanged(ECheckBoxState NewState);
 
-	/** */
-	ECheckBoxState GetEditLayersState() const;
-	void OnEditLayersStateChanged(ECheckBoxState NewState);
-
 	/** */
 	void OnSetImportConfiguration(TSharedPtr<FTileImportConfiguration> InTileConfig, ESelectInfo::Type SelectInfo);
 
diff --git a/Engine/Source/Editor/WorldBrowser/Private/Tiles/WorldTileCollectionModel.cpp b/Engine/Source/Editor/WorldBrowser/Private/Tiles/WorldTileCollectionModel.cpp
--- a/Engine/Source/Editor/WorldBrowser/Private/Tiles/WorldTileCollectionModel.cpp
+++ b/Engine/Source/Editor/WorldBrowser/Private/Tiles/WorldTileCollectionModel.cpp
@@ -1641,7 +1641,6 @@ void FWorldTileCollectionModel::ImportTiledLandscape_Executed()
 			SetupLandscapeImportLayers(ImportSettings, GetWorld()->GetOutermost()->GetName(), INDEX_NONE, ImportLayers);
 
 			// Set landscape configuration
-			Landscape->bCanHaveLayersContent = ImportSettings.bEditLayersEnabled;
 			Landscape->LandscapeMaterial	= ImportSettings.LandscapeMaterial.Get();
 			Landscape->ComponentSizeQuads	= ImportSettings.QuadsPerSection*ImportSettings.SectionsPerComponent;
 			Landscape->NumSubsections		= ImportSettings.SectionsPerComponent;
diff --git a/Engine/Source/Runtime/Core/Public/UObject/FortniteMainBranchObjectVersions.inl b/Engine/Source/Runtime/Core/Public/UObject/FortniteMainBranchObjectVersions.inl
--- a/Engine/Source/Runtime/Core/Public/UObject/FortniteMainBranchObjectVersions.inl
+++ b/Engine/Source/Runtime/Core/Public/UObject/FortniteMainBranchObjectVersions.inl
@@ -679,5 +679,8 @@ FORTNITE_MAIN_VERSION(SpecializeBloomIntensity, 221)
 // Add support for world partition actor component descriptors
 FORTNITE_MAIN_VERSION(WorldPartitionActorComponentDesc, 222)
 
+// Migrate Non-Edit layer landscapes to use the edit layer (ULandscapeEditLayer) system
+FORTNITE_MAIN_VERSION(MigrateLandscapeNonEditLayerToEditLayer, 223)
+
 // -----<new versions can be added above this line>-------------------------------------------------
 #undef FORTNITE_MAIN_VERSION
\ No newline at end of file
diff --git a/Engine/Source/Runtime/Landscape/Classes/Landscape.h b/Engine/Source/Runtime/Landscape/Classes/Landscape.h
--- a/Engine/Source/Runtime/Landscape/Classes/Landscape.h
+++ b/Engine/Source/Runtime/Landscape/Classes/Landscape.h
@@ -396,7 +396,9 @@ public:
 	void RegisterLandscapeEdMode(ILandscapeEdModeInterface* InLandscapeEdMode) { LandscapeEdMode = InLandscapeEdMode; }
 	void UnregisterLandscapeEdMode() { LandscapeEdMode = nullptr; }
 	bool HasLandscapeEdMode() const { return LandscapeEdMode != nullptr; }
+	UE_DEPRECATED(5.7, "Non-edit layer landscapes are deprecated, all landscapes use the edit layer system now.")
 	LANDSCAPE_API virtual bool HasLayersContent() const override;
+	UE_DEPRECATED(5.7, "Non-edit layer landscapes are deprecated, all landscapes use the edit layer system now.")
 	LANDSCAPE_API virtual void UpdateCachedHasLayersContent(bool bInCheckComponentDataIntegrity) override;
 	LANDSCAPE_API void RequestSplineLayerUpdate();
 	LANDSCAPE_API void RequestLayersInitialization(bool bInRequestContentUpdate = true, bool bInForceLayerResourceReset = true);
@@ -496,7 +498,9 @@ PRAGMA_ENABLE_DEPRECATION_WARNINGS
 	LANDSCAPE_API bool DeleteLayer(int32 InLayerIndex);
 	LANDSCAPE_API void CollapseLayer(int32 InLayerIndex);
 	void CollapseLayerIntermediateRender(int32 InLayerIndex);
+	UE_DEPRECATED(5.7, "Non-edit layer landscapes are deprecated. Use CollapseEditLayers to replace all existing layers with a new default layer.")
 	LANDSCAPE_API void DeleteLayers();
+	LANDSCAPE_API void CollapseEditLayers();
 	LANDSCAPE_API void SetEditingLayer(const FGuid& InLayerGuid = FGuid());
 	LANDSCAPE_API void SetGrassUpdateEnabled(bool bInGrassUpdateEnabled);
 	LANDSCAPE_API const FGuid& GetEditingLayer() const;
@@ -529,6 +533,8 @@ PRAGMA_ENABLE_DEPRECATION_WARNINGS
 	void ReleaseLayersRenderingResource();
 	void ClearDirtyData(ULandscapeComponent* InLandscapeComponent);
 	
+	LANDSCAPE_API void ConvertNonEditLayerLandscape();
+	UE_DEPRECATED(5.7, "Non-edit layer landscapes are deprecated. Use ConvertNonEditLayerLandscape to convert to non-edit layer landscapes to edit layer based landscapes.")
 	LANDSCAPE_API void ToggleCanHaveLayersContent();
 	LANDSCAPE_API void ForceUpdateLayersContent(bool bIntermediateRender = false);
 
@@ -566,6 +572,9 @@ PRAGMA_ENABLE_DEPRECATION_WARNINGS
 
 	UE_INTERNAL LANDSCAPE_API bool AllowSelectiveRenderEditLayers() const;
 
+	bool GetIsPendingEditLayerConversion() const;
+	void SetIsPendingEditLayerConversion(bool bInPendingEditLayerConversion);
+
 private:
 	FLandscapeLayer* GetLayerInternal(int32 InLayerIndex);
 	void OnLayerCreatedInternal(ULandscapeEditLayerBase* EditLayer);
@@ -687,8 +696,9 @@ public:
 	UPROPERTY()
 	bool bIncludeGridSizeInNameForLandscapeActors = false;
 
-	UPROPERTY(EditAnywhere, Category=Landscape)
-	bool bCanHaveLayersContent = false;
+	UE_DEPRECATED(5.7, "Non-edit layer landscapes are deprecated, all landscapes use the edit layer system now.")
+	UPROPERTY()
+	bool bCanHaveLayersContent_DEPRECATED = false;
 
 	/*
 	 * If true, WorldPartitionLandscapeSplineMeshesBuilder is responsible of generating partitioned actors of type ALandscapeSplineMeshesActor that will contain all landscape spline/controlpoints static meshes. 
@@ -831,8 +841,11 @@ private:
 	bool bHasWarnedAboutInvalidShadingModel = false;
 
 	// Counter to detect re-entrance
-	uint32 InLayerUpdateCount = 0;
+	uint32 LayerUpdateCount = 0;
 #endif
+
+	/** If true, edit layers were not supported on load. All proxies must copy data to the default edit layer during registration */
+	bool bIsPendingEditLayerConversion = false;
 };
 
 #if WITH_EDITOR
diff --git a/Engine/Source/Runtime/Landscape/Classes/LandscapeHeightfieldCollisionComponent.h b/Engine/Source/Runtime/Landscape/Classes/LandscapeHeightfieldCollisionComponent.h
--- a/Engine/Source/Runtime/Landscape/Classes/LandscapeHeightfieldCollisionComponent.h
+++ b/Engine/Source/Runtime/Landscape/Classes/LandscapeHeightfieldCollisionComponent.h
@@ -234,8 +234,6 @@ public:
 	virtual bool NeedsLoadForServer() const override;
 	virtual void ExportCustomProperties(FOutputDevice& Out, uint32 Indent) override;
 	virtual void ImportCustomProperties(const TCHAR* SourceText, FFeedbackContext* Warn) override;
-	virtual void PostEditImport() override;
-	virtual void PostEditUndo() override;
 	//~ End UObject Interface.
 		
 	/**  We speculatively async load collision object from DDC to remove hitch when streaming */
diff --git a/Engine/Source/Runtime/Landscape/Classes/LandscapeInfo.h b/Engine/Source/Runtime/Landscape/Classes/LandscapeInfo.h
--- a/Engine/Source/Runtime/Landscape/Classes/LandscapeInfo.h
+++ b/Engine/Source/Runtime/Landscape/Classes/LandscapeInfo.h
@@ -206,6 +206,7 @@ public:
 	LANDSCAPE_API void ExportHeightmap(const FString& Filename, const FIntRect& ExportRegion);
 	LANDSCAPE_API void ExportLayer(ULandscapeLayerInfoObject* LayerInfo, const FString& Filename);
 	LANDSCAPE_API void ExportLayer(ULandscapeLayerInfoObject* LayerInfo, const FString& Filename, const FIntRect& ExportRegion);
+	// TODO [jared.ritchie] - deprecate for 5.7 after non-edit layer landscape deprecation
 	LANDSCAPE_API bool ApplySplines(bool bOnlySelected, TSet<TObjectPtr<ULandscapeComponent>>* OutModifiedComponents = nullptr, bool bMarkPackageDirty = true);
 
 	LANDSCAPE_API bool GetSelectedExtent(int32& MinX, int32& MinY, int32& MaxX, int32& MaxY) const;
@@ -331,6 +332,7 @@ public:
 	LANDSCAPE_API void RemoveXYOffsets();
 
 	/** Will tell if the landscape actor can have some content related to the layer system */
+	UE_DEPRECATED(5.7, "Non-edit layer landscapes are deprecated, all landscapes use the edit layer system now.")
 	LANDSCAPE_API bool CanHaveLayersContent() const;
 
 	/** Will clear all component dirty data */
diff --git a/Engine/Source/Runtime/Landscape/Classes/LandscapeProxy.h b/Engine/Source/Runtime/Landscape/Classes/LandscapeProxy.h
--- a/Engine/Source/Runtime/Landscape/Classes/LandscapeProxy.h
+++ b/Engine/Source/Runtime/Landscape/Classes/LandscapeProxy.h
@@ -85,6 +85,14 @@ DECLARE_MULTICAST_DELEGATE_TwoParams(FOnLandscapeProxyFixupSharedDataDelegate, A
 
 // ----------------------------------------------------------------------------------
 
+enum class ELandscapeEditLayerContentState : uint8
+{
+	Undefined = 0,
+	PostLoad_HadEditLayers,
+	PostLoad_NoEditLayers,
+	AutoConvertedToEditLayers
+};
+
 USTRUCT()
 struct FLandscapeTargetLayerSettings
 {
@@ -950,15 +958,16 @@ PRAGMA_ENABLE_DEPRECATION_WARNINGS
 	int32 HLODMeshSourceLOD;
 #endif
 
-	/** Flag that tell if we have some layers content **/
+	UE_DEPRECATED(5.7, "Non-edit layer landscapes are deprecated, all landscapes use the edit layer system now.")
 	UPROPERTY()
-	bool bHasLayersContent;
+	bool bHasLayersContent_DEPRECATED = false;
 
-	/** Because bHasLayersContent can be stomped by the FixupSharedData copying layers over from the parent LandscapeActor,
-	  * we preserve the initial state of bHasLayersContent until we get to registration and can apply upgrade steps
-	  **/
+	UE_DEPRECATED(5.7, "Non-edit layer landscapes are deprecated, all landscapes use the edit layer system now.")
 	TOptional<bool> HadLayersContentAtPostLoadTime;
 
+	/** Tracks the has edit layer state of a proxy through load and registration */
+	ELandscapeEditLayerContentState HadLayersContentAtPostLoadState = ELandscapeEditLayerContentState::Undefined;
+
 	UE_DEPRECATED(5.6, "This has been replaced by the per-platform CVar landscape.HeightmapCompressionMode")
 	UPROPERTY()
 	bool bUseCompressedHeightmapStorage_DEPRECATED = false;
@@ -1544,21 +1553,17 @@ PRAGMA_ENABLE_DEPRECATION_WARNINGS
 	* Only works in the editor
 	*/
 	UFUNCTION(BlueprintCallable, meta = (DisplayName = "Landscape Import Heightmap from RenderTarget", Keywords = "Push RenderTarget to Landscape Heightmap", UnsafeDuringActorConstruction = "true"), Category = Rendering)
-	bool LandscapeImportHeightmapFromRenderTarget(UTextureRenderTarget2D* InRenderTarget, bool InImportHeightFromRGChannel = false);
+	bool LandscapeImportHeightmapFromRenderTarget(UTextureRenderTarget2D* InRenderTarget, bool InImportHeightFromRGChannel = false, int32 InEditLayerIndex = 0);
 
 	/**
 	* Overwrites a landscape weightmap with render target data
 	* Only works in the editor
 	*/
 	UFUNCTION(BlueprintCallable, meta = (DisplayName = "Landscape Import Weightmap from RenderTarget", Keywords = "Push RenderTarget to Landscape Weightmap", UnsafeDuringActorConstruction = "true"), Category = Rendering)
-	bool LandscapeImportWeightmapFromRenderTarget(UTextureRenderTarget2D* InRenderTarget, FName InLayerName);
+	bool LandscapeImportWeightmapFromRenderTarget(UTextureRenderTarget2D* InRenderTarget, FName InLayerName, int32 InEditLayerIndex = 0);
 
-	
-	/**
-	* Output a landscape weightmap to a render target
-	* Only works in the editor
-	*/
-	UFUNCTION(BlueprintCallable, meta = (DisplayName = "Landscape Export Weightmap to RenderTarget", Keywords = "Push Landscape Weightmap to RenderTarget", UnsafeDuringActorConstruction = "true"), Category = Rendering)
+	UE_DEPRECATED(5.7, "Deprecating placeholder function. Not needed after non-edit layer landscape deprecation. ")
+	UFUNCTION()
 	bool LandscapeExportWeightmapToRenderTarget(UTextureRenderTarget2D* InRenderTarget, FName InLayerName);
 
 	DECLARE_EVENT(ALandscape, FLandscapeMaterialChangedDelegate);
@@ -1570,12 +1575,13 @@ PRAGMA_DISABLE_DEPRECATION_WARNINGS
 PRAGMA_ENABLE_DEPRECATION_WARNINGS
 	}
 
-	/** Will tell if the landscape proxy as some content related to the layer system */
+	UE_DEPRECATED(5.7, "Non-edit layer landscapes are deprecated, all landscapes use the edit layer system now.")
 	LANDSCAPE_API virtual bool HasLayersContent() const;
-	
-	/** Will tell if the landscape proxy can have some content related to the layer system */
+
+	UE_DEPRECATED(5.7, "Non-edit layer landscapes are deprecated, all landscapes use the edit layer system now.")
 	LANDSCAPE_API bool CanHaveLayersContent() const;
 
+	UE_DEPRECATED(5.7, "Non-edit layer landscapes are deprecated, all landscapes use the edit layer system now.")
 	LANDSCAPE_API virtual void UpdateCachedHasLayersContent(bool InCheckComponentDataIntegrity = false);
 
 	FAsyncWorkMonitor& GetAsyncWorkMonitor() { return AsyncWorkMonitor; }
diff --git a/Engine/Source/Runtime/Landscape/Private/Landscape.cpp b/Engine/Source/Runtime/Landscape/Private/Landscape.cpp
--- a/Engine/Source/Runtime/Landscape/Private/Landscape.cpp
+++ b/Engine/Source/Runtime/Landscape/Private/Landscape.cpp
@@ -1659,7 +1659,6 @@ ALandscapeProxy::ALandscapeProxy(const FObjectInitializer& ObjectInitializer)
 	bGenerateOverlapEvents = false;
 #if WITH_EDITORONLY_DATA
 	MaxPaintedLayersPerComponent = 0;
-	bHasLayersContent = false;
 	HLODTextureSizePolicy = ELandscapeHLODTextureSizePolicy::SpecificSize;
 	HLODTextureSize = 256;
 	HLODMeshSourceLODPolicy = ELandscapeHLODMeshSourceLODPolicy::LowestDetailLOD;
@@ -3111,15 +3110,24 @@ PRAGMA_ENABLE_DEPRECATION_WARNINGS
 		}
 	}
 
-	// In case we're a landscape with edit layers but we actually lack a layer (e.g. it was removed by the test above, because its edit layer class is unknown), let's create one all the same 
-	//  because we're always supposed to have at least 1 : 
-	if (CanHaveLayersContent())
+	bool bIsEditLayerLandscape = true;
+
+	// Non-edit layer landscapes are marked for auto-conversion in ALandscapeProxy::PostLoad. Run validations on serialized edit layer landscapes only
+	if (GetLinkerCustomVersion(FFortniteMainBranchObjectVersion::GUID) < FFortniteMainBranchObjectVersion::MigrateLandscapeNonEditLayerToEditLayer)
+	{
+PRAGMA_DISABLE_DEPRECATION_WARNINGS
+		bIsEditLayerLandscape = bCanHaveLayersContent_DEPRECATED;
+PRAGMA_ENABLE_DEPRECATION_WARNINGS
+	}
+
+	// In case we're a landscape with edit layers but we actually lack a layer (e.g. it was removed by the test above, because its edit layer class is unknown),
+	// Create a default layer because we're always supposed to have at least 1, then copy over relevant data if needed :
+	if (bIsEditLayerLandscape)
 	{
 		if (LandscapeEditLayers.IsEmpty())
 		{
 			// Create a default layer
 			CreateDefaultLayer();
-			ensure(!LandscapeEditLayers.IsEmpty());
 		}
 
 		// Ensure a valid edit layer is selected post load
@@ -3127,6 +3135,8 @@ PRAGMA_ENABLE_DEPRECATION_WARNINGS
 		{
 			SetSelectedEditLayerIndex(0);
 		}
+
+		ensure(!LandscapeEditLayers.IsEmpty());
 	}
 #endif // WITH_EDITOR
 
@@ -4252,7 +4262,7 @@ bool ULandscapeInfo::UpdateLayerInfoMap(ALandscapeProxy* Proxy /*= nullptr*/, bo
 	if (GIsEditor)
 	{
 		ALandscape* Landscape = LandscapeActor.Get();
-		if (Landscape && Landscape->HasLayersContent())
+		if (Landscape && !Landscape->IsTemplate())
 		{
 			Landscape->RequestLayersInitialization(/*bInRequestContentUpdate*/false, /*bInForceLayerResourceReset*/false);
 		}
@@ -4342,19 +4352,24 @@ void ALandscapeProxy::PostLoad()
 		}
 	}
 
-	// Fix up bHasLayersContent if needed : there was a point where there was a missing call (when adding new components) that was leading to  it not being 
-	//  properly updated even though there was some layer content. We need to fix it now, so that we don't attempt to migrate data from non-edit layers to 
-	//  edit layers later on, since that would stomp the edit layer data : 
-	if (!bHasLayersContent && !LandscapeComponents.IsEmpty() && (LandscapeComponents[0] != nullptr) && LandscapeComponents[0]->HasLayersData())
+	if (GetLinkerCustomVersion(FFortniteMainBranchObjectVersion::GUID) < FFortniteMainBranchObjectVersion::MigrateLandscapeNonEditLayerToEditLayer)
 	{
-		bHasLayersContent = true;
+PRAGMA_DISABLE_DEPRECATION_WARNINGS
+		// Fix up bHasLayersContent_DEPRECATED if needed : there was a point where there was a missing call (when adding new components) that was leading to  it not being 
+		//  properly updated even though there was some layer content. We need to fix it now, so that we don't attempt to migrate data from non-edit layers to 
+		//  edit layers later on, since that would stomp the edit layer data : 
+		if (!bHasLayersContent_DEPRECATED && !LandscapeComponents.IsEmpty() && (LandscapeComponents[0] != nullptr) && LandscapeComponents[0]->HasLayersData())
+		{
+			bHasLayersContent_DEPRECATED = true;
+		}
+
+		// save the load time state of layers content before doing anything else (specifically FixupSharedData can stomp it)
+		check(HadLayersContentAtPostLoadState == ELandscapeEditLayerContentState::Undefined);
+		HadLayersContentAtPostLoadState = bHasLayersContent_DEPRECATED ? ELandscapeEditLayerContentState::PostLoad_HadEditLayers : ELandscapeEditLayerContentState::PostLoad_NoEditLayers;
+PRAGMA_ENABLE_DEPRECATION_WARNINGS
 	}
 #endif // WITH_EDITOR
 
-	// save the load time state of layers content before doing anything else (specifically FixupSharedData can stomp it)
-	check(!HadLayersContentAtPostLoadTime.IsSet());
-	HadLayersContentAtPostLoadTime = bHasLayersContent;
-
 	Super::PostLoad();
 
 	PostLoadFixupLandscapeGuidsIfInstanced();
@@ -5865,7 +5880,6 @@ bool ULandscapeInfo::IsRegistered(const ALandscapeProxy* Proxy) const
 	return bResult;
 }
 
-
 // this function contains all of the registration code that requires the ALandscape actor to be present
 void ULandscapeInfo::RegisterLandscapeActorWithProxyInternal(ALandscapeProxy* Proxy, bool bMapCheck)
 {
@@ -5914,24 +5928,69 @@ void ULandscapeInfo::RegisterLandscapeActorWithProxyInternal(ALandscapeProxy* Pr
 		{
 			Proxy->FixupWeightmaps();
 		}
+		
+		// All proxy data fixup is complete - CopyOldDataToDefaultLayer will only migrate data once this flag is set
+		Proxy->bIsLandscapeActorRegisteredWithLandscapeInfo = true;
 
-		Proxy->UpdateCachedHasLayersContent(true);
-
-		if (Proxy->HadLayersContentAtPostLoadTime.IsSet())
 		{
-			bool bHasLayersContentBefore = Proxy->HadLayersContentAtPostLoadTime.GetValue();
-			check(Proxy->WeightmapFixupVersion == Proxy->CurrentVersion);
-			const bool bNeedOldDataMigration = !bHasLayersContentBefore && CanHaveLayersContent();
-			if (bNeedOldDataMigration && LandscapeActor->HasLayersContent())
+			/*
+			Automatically enable edit layers for non-edit layer landscapes here (since FixupSharedData can stomp updates in PostLoad)
+
+			There are two possibilities for the ALandscape on registration
+			1. ALandscape has layers content -> No conversion needed
+			2. ALandscape does not have layers content -> Trigger full auto conversion to an edit layer based landscape
+
+			There are two possibilities for any ALandscapeProxy on registration
+			1. Proxy does not have layers content -> Proxy needs to copy data to default edit layer
+				* Conditions - the entire landscape had edit layers disabled, or edit layers were enabled and serialized only on the ALandscape in a WP setup
+			2. Proxy has layers content -> If ALandscape does not support layers content on load, copy data to the default edit layer
+				* Conditions - the entire landscape had edit layers enabled, or edit layers are enabled and serialized on just the ALandscapeProxy in a WP setup
+				* During an auto conversion the proxy edit layer data is flattened and moved to the default edit layer
+			*/
+
+			const bool bProxyHasNoLayersContentPostLoad = Proxy->HadLayersContentAtPostLoadState == ELandscapeEditLayerContentState::PostLoad_NoEditLayers;
+			const bool bProxyHasBeenAutoConverted = Proxy->HadLayersContentAtPostLoadState == ELandscapeEditLayerContentState::AutoConvertedToEditLayers;
+			const bool bIsLandscapeActor = Proxy == LandscapeActor;
+
+			// If the LandscapeActor did not support edit layers we need to run a full auto conversion
+			if (bIsLandscapeActor && bProxyHasNoLayersContentPostLoad && !bProxyHasBeenAutoConverted)
 			{
-				LandscapeActor->CopyOldDataToDefaultLayer(Proxy);
-				Proxy->HadLayersContentAtPostLoadTime = false;
+				LandscapeActor->SetIsPendingEditLayerConversion(bProxyHasNoLayersContentPostLoad);
+			}
+
+			// Proxies are registered many times (on Undo), but conversion can only happen once on initial load
+			if (!LandscapeActor->IsTemplate() && !bProxyHasBeenAutoConverted && (LandscapeActor->GetIsPendingEditLayerConversion() || bProxyHasNoLayersContentPostLoad))
+			{
+				if (bIsLandscapeActor)
+				{
+					UE_LOG(LogLandscape, Log, TEXT("RegisterLandscapeActorWithProxyInternal: Automatically enabling edit layers on ALandscape: %s"), *Proxy->GetFullName());
+
+					// Create a default edit layer and copy the ALandscape data
+					// Streaming proxy data gets copied as proxies complete registration
+					LandscapeActor->ConvertNonEditLayerLandscape();
+				}
+				else
+				{
+					UE_LOG(LogLandscape, Log, TEXT("RegisterLandscapeActorWithProxyInternal: Automatically enabling edit layers on Proxy: %s"), *Proxy->GetFullName());
+
+					// As proxies are registered, copy the data to edit layer based landscape. Proxies with edit layer data will be flattened to the default layer
+					// Proxy data should not be copied before this point in registration to ensure data is fixed up
+					LandscapeActor->CopyOldDataToDefaultLayer(Proxy);
+				}
+
+				Proxy->RemoveXYOffsets();
+				Proxy->ValidateProxyLayersWeightmapUsage();
+
+				// Mark the proxy dirty to display the 'Requires Save' warning to the user
+				MarkObjectDirty(/*InObject = */ Proxy, /*bInForceResave =*/ true, /*InLandscapeOverride =*/ Landscape);
+				Proxy->HadLayersContentAtPostLoadState = ELandscapeEditLayerContentState::AutoConvertedToEditLayers;
 			}
 		}
+
+		// Landscape needs to have edit layers by end of registration
+		check((Proxy != LandscapeActor) || !LandscapeActor->GetEditLayers().IsEmpty());
 	}
 #endif // WITH_EDITOR
-
-	Proxy->bIsLandscapeActorRegisteredWithLandscapeInfo = true;
 }
 
 void ULandscapeInfo::RegisterActor(ALandscapeProxy* Proxy, bool bMapCheck, bool bUpdateAllAddCollisions)
diff --git a/Engine/Source/Runtime/Landscape/Private/LandscapeBlueprintSupport.cpp b/Engine/Source/Runtime/Landscape/Private/LandscapeBlueprintSupport.cpp
--- a/Engine/Source/Runtime/Landscape/Private/LandscapeBlueprintSupport.cpp
+++ b/Engine/Source/Runtime/Landscape/Private/LandscapeBlueprintSupport.cpp
@@ -25,7 +25,7 @@ void ALandscapeProxy::EditorApplySpline(USplineComponent* InSplineComponent, flo
 		if (ALandscape* Landscape = GetLandscapeInfo()->LandscapeActor.Get())
 		{
 			const ULandscapeEditLayerBase* EditLayer = Landscape->GetEditLayerConst(EditLayerName);
-			if (Landscape->HasLayersContent() && (EditLayer == nullptr))
+			if (EditLayer == nullptr)
 			{
 				UE_LOG(LogLandscape, Error, TEXT("Invalid landscape edit layer name (\"%s\") for Edit Layers-enabled landscape. Cannot apply spline. "), *EditLayerName.ToString());
 				return;
diff --git a/Engine/Source/Runtime/Landscape/Private/LandscapeCollision.cpp b/Engine/Source/Runtime/Landscape/Private/LandscapeCollision.cpp
--- a/Engine/Source/Runtime/Landscape/Private/LandscapeCollision.cpp
+++ b/Engine/Source/Runtime/Landscape/Private/LandscapeCollision.cpp
@@ -2486,39 +2486,6 @@ void ULandscapeMeshCollisionComponent::Serialize(FArchive& Ar)
 	}
 }
 
-#if WITH_EDITOR
-void ULandscapeHeightfieldCollisionComponent::PostEditImport()
-{
-	Super::PostEditImport();
-
-	if (!GetLandscapeProxy()->HasLayersContent())
-	{
-		// Reinitialize physics after paste
-		if (CollisionSizeQuads > 0)
-		{
-			RecreateCollision();
-		}
-	}
-}
-
-void ULandscapeHeightfieldCollisionComponent::PostEditUndo()
-{
-	Super::PostEditUndo();
-
-    // Landscape Layers are updates are delayed and done in  ALandscape::TickLayers
-	if (!GetLandscapeProxy()->HasLayersContent())
-	{
-		// Reinitialize physics after undo
-		if (CollisionSizeQuads > 0)
-		{
-			RecreateCollision();
-		}
-
-		FNavigationSystem::UpdateComponentData(*this);
-	}
-}
-#endif // WITH_EDITOR
-
 bool ULandscapeHeightfieldCollisionComponent::IsShown(const FEngineShowFlags& ShowFlags) const
 {
 	return ShowFlags.Landscape;
diff --git a/Engine/Source/Runtime/Landscape/Private/LandscapeConfigHelper.cpp b/Engine/Source/Runtime/Landscape/Private/LandscapeConfigHelper.cpp
--- a/Engine/Source/Runtime/Landscape/Private/LandscapeConfigHelper.cpp
+++ b/Engine/Source/Runtime/Landscape/Private/LandscapeConfigHelper.cpp
@@ -93,7 +93,7 @@ ALandscapeProxy* FLandscapeConfigHelper::FindOrAddLandscapeStreamingProxy(UActor
 bool FLandscapeConfigHelper::ChangeGridSize(ULandscapeInfo* InLandscapeInfo, uint32 InNewGridSizeInComponents, TSet<AActor*>& OutActorsToDelete)
 {
 	check(InLandscapeInfo);
-
+	check(!InLandscapeInfo->IsTemplate());
 	const uint32 GridSize = InLandscapeInfo->GetGridSize(InNewGridSizeInComponents);
 	
 	InLandscapeInfo->LandscapeActor->Modify();
@@ -234,10 +234,7 @@ bool FLandscapeConfigHelper::ChangeGridSize(ULandscapeInfo* InLandscapeInfo, uin
 		OutActorsToDelete.Add(ProxyToDelete);
 	}
 
-	if (InLandscapeInfo->CanHaveLayersContent())
-	{
-		InLandscapeInfo->ForceLayersFullUpdate();
-	}
+	InLandscapeInfo->ForceLayersFullUpdate();
 
 	return true;
 }
@@ -346,75 +343,50 @@ void FLandscapeConfigHelper::CopyRegionToComponent(ULandscapeInfo* InLandscapeIn
 	TArray<uint16> NewHeightData;
 	TArray<uint8> NewWeightData;
 
-	if (!NewLandscapeInfo->LandscapeActor->bCanHaveLayersContent)
+	TArray<FColor> HeightInitData;
+	HeightInitData.SetNum(FMath::Square(HeightmapSize));
+	const uint16 DefaultHeight = LandscapeDataAccess::GetTexHeight(0.f);
+	FColor DefaultPackedHeight = LandscapeDataAccess::PackHeight(DefaultHeight);
+	for (int32 Index = 0; Index < HeightInitData.Num(); ++Index)
+	{
+		HeightInitData[Index] = DefaultPackedHeight;
+	}
+	InComponent->InitHeightmapData(HeightInitData, false);
+
+	for (const ULandscapeEditLayerBase* LandscapeEditLayer : InLandscapeInfo->LandscapeActor->GetEditLayersConst())
 	{
-		InComponent->HeightmapScaleBias = FVector4(1.0f / (float)HeightmapSize, 1.0f / (float)HeightmapSize, 0.0f, 0.0f);
-		InComponent->SetHeightmap(InComponent->GetLandscapeProxy()->CreateLandscapeTexture(HeightmapSize, HeightmapSize, TEXTUREGROUP_Terrain_Heightmap, TSF_BGRA8));
+		check(LandscapeEditLayer != nullptr);
+		if (!LandscapeEditLayer->NeedsPersistentTextures())
+		{
+			continue;
+		}
 
-		ExtractLandscapeData(InLandscapeInfo, InRegion, FGuid(), SrcHeightData, SrcWeightData);
+		TMap<UTexture2D*, UTexture2D*> CreatedTextures;
+		InComponent->AddDefaultLayerData(LandscapeEditLayer->GetGuid(), { InComponent }, CreatedTextures);
+
+		// No need to copy the persistent textures if they are not manually edited (e.g. fully procedural layers with persistent textures : e.g. the splines layer) :
+		if (!LandscapeEditLayer->SupportsEditingTools())
+		{
+			continue;
+		}
 
+		SrcWeightData.Empty();
+		ExtractLandscapeData(InLandscapeInfo, InRegion, LandscapeEditLayer->GetGuid(), SrcHeightData, SrcWeightData);
 		CopyData(SrcHeightData, NewHeightData, InRegion, NewRegion, bInResample);
 
+		FScopedSetLandscapeEditingLayer LayerScope(NewLandscapeInfo->LandscapeActor.Get(), LandscapeEditLayer->GetGuid());
 		{
 			FLandscapeEditDataInterface LandscapeEdit(NewLandscapeInfo);
 			LandscapeEdit.SetHeightData(NewRegion.Min.X, NewRegion.Min.Y, NewRegion.Max.X, NewRegion.Max.Y, NewHeightData.GetData(), 0, false);
 		}
 
-		
 		for (const FLandscapeImportLayerInfo& Layer : SrcWeightData)
 		{
-			FLandscapeEditDataInterface LandscapeEdit(NewLandscapeInfo);
 			CopyData(Layer.LayerData, NewWeightData, InRegion, NewRegion, bInResample);
+			FLandscapeEditDataInterface LandscapeEdit(NewLandscapeInfo);
 			LandscapeEdit.SetAlphaData(Layer.LayerInfo, NewRegion.Min.X, NewRegion.Min.Y, NewRegion.Max.X, NewRegion.Max.Y, NewWeightData.GetData(), 0);
 		}
 	}
-	else
-	{
-		TArray<FColor> HeightInitData;
-		HeightInitData.SetNum(FMath::Square(HeightmapSize));
-		const uint16 DefaultHeight = LandscapeDataAccess::GetTexHeight(0.f);
-		FColor DefaultPackedHeight = LandscapeDataAccess::PackHeight(DefaultHeight);
-		for (int32 Index = 0; Index < HeightInitData.Num(); ++Index)
-		{
-			HeightInitData[Index] = DefaultPackedHeight;
-		}
-		InComponent->InitHeightmapData(HeightInitData, false);
-
-		for (const ULandscapeEditLayerBase* LandscapeEditLayer : InLandscapeInfo->LandscapeActor->GetEditLayersConst())
-		{		
-			check(LandscapeEditLayer != nullptr);
-			if (!LandscapeEditLayer->NeedsPersistentTextures())
-			{
-				continue;
-			}
-
-			TMap<UTexture2D*, UTexture2D*> CreatedTextures;
-			InComponent->AddDefaultLayerData(LandscapeEditLayer->GetGuid(), { InComponent }, CreatedTextures);
-
-			// No need to copy the persistent textures if they are not manually edited (e.g. fully procedural layers with persistent textures : e.g. the splines layer) :
-			if (!LandscapeEditLayer->SupportsEditingTools())
-			{
-				continue;
-			}
-
-			SrcWeightData.Empty();
-			ExtractLandscapeData(InLandscapeInfo, InRegion, LandscapeEditLayer->GetGuid(), SrcHeightData, SrcWeightData);
-			CopyData(SrcHeightData, NewHeightData, InRegion, NewRegion, bInResample);
-
-			FScopedSetLandscapeEditingLayer LayerScope(NewLandscapeInfo->LandscapeActor.Get(), LandscapeEditLayer->GetGuid());
-			{
-				FLandscapeEditDataInterface LandscapeEdit(NewLandscapeInfo);
-				LandscapeEdit.SetHeightData(NewRegion.Min.X, NewRegion.Min.Y, NewRegion.Max.X, NewRegion.Max.Y, NewHeightData.GetData(), 0, false);
-			}
-
-			for (const FLandscapeImportLayerInfo& Layer : SrcWeightData)
-			{
-				CopyData(Layer.LayerData, NewWeightData, InRegion, NewRegion, bInResample);
-				FLandscapeEditDataInterface LandscapeEdit(NewLandscapeInfo);
-				LandscapeEdit.SetAlphaData(Layer.LayerInfo, NewRegion.Min.X, NewRegion.Min.Y, NewRegion.Max.X, NewRegion.Max.Y, NewWeightData.GetData(), 0);
-			}
-		}
-	}
 }
 
 void FLandscapeConfigHelper::MoveSplinesToLandscape(ULandscapeInfo* InLandscapeInfo, ALandscapeProxy* InLandscape, float InScaleFactor)
@@ -496,7 +468,7 @@ ULandscapeInfo* FLandscapeConfigHelper::ChangeConfiguration(ULandscapeInfo* InLa
 	check(InLandscapeInfo);
 	ALandscape* OldLandscape = InLandscapeInfo->LandscapeActor.Get();
 	check(OldLandscape);
-
+	check(!OldLandscape->IsTemplate());
 	UWorld* World = OldLandscape->GetWorld();
 	UActorPartitionSubsystem* ActorPartitionSubsystem = World->GetSubsystem<UActorPartitionSubsystem>();
 	FLandscapeConfig OldConfig(InLandscapeInfo);
@@ -568,16 +540,12 @@ ULandscapeInfo* FLandscapeConfigHelper::ChangeConfiguration(ULandscapeInfo* InLa
 	NewLandscape->NumSubsections = InNewConfig.ComponentNumSubsections;
 	NewLandscape->SubsectionSizeQuads = InNewConfig.SubsectionSizeQuads;	
 	NewLandscape->SetActorRelativeScale3D(NewScale);
-	
+
 	// Copy Layer Data
-	if (OldLandscape->HasLayersContent())
+	for (const FLandscapeLayer& OldLayer : OldLandscape->GetLayersConst())
 	{
-		NewLandscape->bCanHaveLayersContent = true;
-		for (const FLandscapeLayer& OldLayer : OldLandscape->GetLayersConst())
-		{
-			const FLandscapeLayer* NewLayer = NewLandscape->DuplicateLayerAndMoveBrushes(OldLayer);
-			check((NewLayer == nullptr) || (NewLayer->EditLayer != nullptr)); // it's possible DuplicateLayerAndMoveBrushes fails (e.g. max number of layers reached), but if not, we should always have an EditLayer
-		}
+		const FLandscapeLayer* NewLayer = NewLandscape->DuplicateLayerAndMoveBrushes(OldLayer);
+		check((NewLayer == nullptr) || (NewLayer->EditLayer != nullptr)); // it's possible DuplicateLayerAndMoveBrushes fails (e.g. max number of layers reached), but if not, we should always have an EditLayer
 	}
 		
 	ULandscapeInfo* NewLandscapeInfo = NewLandscape->CreateLandscapeInfo();
@@ -634,16 +602,8 @@ ULandscapeInfo* FLandscapeConfigHelper::ChangeConfiguration(ULandscapeInfo* InLa
 		}
 	}
 
-	if (NewLandscapeInfo->CanHaveLayersContent())
-	{
-		NewLandscapeInfo->ForceLayersFullUpdate();
-	}
-	else
-	{
-		FLandscapeEditDataInterface LandscapeEdit(NewLandscapeInfo);
-		LandscapeEdit.RecalculateNormals();
-	}
-
+	NewLandscapeInfo->ForceLayersFullUpdate();
+	
 	// Do as last step so that components have a collision component
 	MoveFoliageToLandscape(InLandscapeInfo, NewLandscapeInfo);
 
diff --git a/Engine/Source/Runtime/Landscape/Private/LandscapeEdit.cpp b/Engine/Source/Runtime/Landscape/Private/LandscapeEdit.cpp
--- a/Engine/Source/Runtime/Landscape/Private/LandscapeEdit.cpp
+++ b/Engine/Source/Runtime/Landscape/Private/LandscapeEdit.cpp
@@ -332,7 +332,7 @@ bool ULandscapeComponent::UpdateCachedBounds(bool bInApproximateBounds)
 
 	// Update local-space bounding box
 	NewLocalBox.Init();
-	if (bInApproximateBounds && GetLandscapeProxy()->HasLayersContent())
+	if (bInApproximateBounds)
 	{
 		FVector MinBox(0, 0, LandscapeDataAccess::GetLocalHeight(0));
 		FVector MaxBox(ComponentSizeQuads + 1, ComponentSizeQuads + 1, LandscapeDataAccess::GetLocalHeight(UINT16_MAX));
@@ -454,11 +454,6 @@ bool ULandscapeComponent::UpdateCachedBounds(bool bInApproximateBounds)
 	ULandscapeHeightfieldCollisionComponent* HFCollisionComponent = GetCollisionComponent();
 	if (HFCollisionComponent)
 	{
-        // In Landscape Layers the Collision Component gets regenerated after the heightmap changes are undone and doesn't need to be transacted
-		if (!GetLandscapeProxy()->HasLayersContent())
-		{
-			HFCollisionComponent->Modify();
-		}
 		HFCollisionComponent->CachedLocalBox = CachedLocalBox;
 		HFCollisionComponent->UpdateComponentToWorld();
 	}
@@ -975,34 +970,20 @@ void ULandscapeComponent::PostEditUndo()
 	if (IsValidChecked(this))
 	{
 		EditToolRenderData.UpdateSelectionMaterial(EditToolRenderData.SelectedType, this);
-		if (!Proxy->HasLayersContent())
-		{
-			EditToolRenderData.UpdateDebugColorMaterial(this);
-            UpdateEditToolRenderData();
-		}	
-	}
-		
-	if (Proxy->HasLayersContent())
-	{
-		// Updating the bounds after undo will re-sync CachedLocalBox with the heightmap data. This is necessary because ULandscapeComponent might be part of 
-		//  the transaction buffer, but the final heightmaps aren't, so undoing might desynchronize them. If the heightmap update that is requested right after 
-		//  ends up modifying heightmap data, then CachedLocalBox will get updated properly again so this is a 'safe' way to fix the issue
-		if (IsValidChecked(this))
-		{
-			UpdateCachedBounds();
-		}
-
-		const bool bUpdateAll = true;
-		RequestHeightmapUpdate(bUpdateAll);
-		RequestWeightmapUpdate(bUpdateAll);
 	}
-	else
+	
+	// Updating the bounds after undo will re-sync CachedLocalBox with the heightmap data. This is necessary because ULandscapeComponent might be part of 
+	//  the transaction buffer, but the final heightmaps aren't, so undoing might desynchronize them. If the heightmap update that is requested right after 
+	//  ends up modifying heightmap data, then CachedLocalBox will get updated properly again so this is a 'safe' way to fix the issue
+	if (IsValidChecked(this))
 	{
-		TSet<ULandscapeComponent*> Components;
-		Components.Add(this);
-		Proxy->FlushGrassComponents(&Components);
+		UpdateCachedBounds();
 	}
 
+	const bool bUpdateAll = true;
+	RequestHeightmapUpdate(bUpdateAll);
+	RequestWeightmapUpdate(bUpdateAll);
+	
 	// Nanite data is non-transactional so we simply want to regenerate it on undo (depending on whether auto-update is enabled, it will either trigger its 
 	//  rebuild or will simply invalidate it) : 
 	Proxy->InvalidateOrUpdateNaniteRepresentation(/* bInCheckContentId = */true, /*InTargetPlatform = */nullptr);
@@ -1539,26 +1520,15 @@ void ULandscapeComponent::UpdateCollisionHeightData(const FColor* const Heightma
 	bool CreatedNew = false;
 	bool ChangeType = false;
 
-	// In Landscape Layers the Collision Component gets regenerated after the heightmap changes are undone and doesn't need to be transacted
-	if (!Proxy->HasLayersContent())
+	// In Edit Layers the Collision Component gets regenerated after the heightmap changes are undone and doesn't need to be transacted, only update dirtied collision height data
+	if (bInUpdateHeightfieldRegion && ComponentX1 == 0 && ComponentY1 == 0 && ComponentX2 == MAX_int32 && ComponentY2 == MAX_int32 && !LayerDirtyCollisionHeightData.IsEmpty())
 	{
-		if (CollisionComp)
-		{
-			CollisionComp->Modify();
-		}
-	}
-	else
-	{
-		// In Landscape Layers, only update dirtied collision height data
-		if (bInUpdateHeightfieldRegion && ComponentX1 == 0 && ComponentY1 == 0 && ComponentX2 == MAX_int32 && ComponentY2 == MAX_int32 && !LayerDirtyCollisionHeightData.IsEmpty())
-		{
-			ComponentX1 = LayerDirtyCollisionHeightData.Min.X;
-			ComponentY1 = LayerDirtyCollisionHeightData.Min.Y;
-			ComponentX2 = LayerDirtyCollisionHeightData.Max.X;
-			ComponentY2 = LayerDirtyCollisionHeightData.Max.Y;
-		}
-		ClearDirtyCollisionHeightData();
+		ComponentX1 = LayerDirtyCollisionHeightData.Min.X;
+		ComponentY1 = LayerDirtyCollisionHeightData.Min.Y;
+		ComponentX2 = LayerDirtyCollisionHeightData.Max.X;
+		ComponentY2 = LayerDirtyCollisionHeightData.Max.Y;
 	}
+	ClearDirtyCollisionHeightData();
 
 	// Existing collision component is same type with collision
 	if (CollisionComp && ((XYOffsetmapTexture == nullptr) == (MeshCollisionComponent == nullptr)))
@@ -1995,11 +1965,6 @@ void ULandscapeComponent::UpdateCollisionLayerData(const FColor* const* const We
 
 	if (CollisionComp)
 	{
-		if (!Proxy->HasLayersContent())
-		{
-			CollisionComp->Modify();
-		}
-
 		// Simple collision is not currently supported with mesh collision components
 		const bool bUsingSimpleCollision = (SimpleCollisionMipLevel > CollisionMipLevel && SimpleCollisionWeightmapTextureMipData && !XYOffsetmapTexture);
 
@@ -3268,7 +3233,7 @@ void ALandscapeProxy::Import(const FGuid& InGuid, int32 InMinX, int32 InMinY, in
 	check(InGuid.IsValid());
 	check(InImportHeightData.Num() == InImportMaterialLayerInfos.Num());
 
-	check(CanHaveLayersContent() || InImportLayers.IsEmpty());
+	check(!IsTemplate());
 
 	FScopedSlowTask SlowTask(2, LOCTEXT("BeingImportingLandscapeTask", "Importing Landscape"));
 	SlowTask.MakeDialog();
@@ -3886,143 +3851,129 @@ void ALandscapeProxy::Import(const FGuid& InGuid, int32 InMinX, int32 InMinY, in
 	// Must be after the FMaterialUpdateContext is destroyed
 	RecreateRenderStateContexts.Reset();
 
-	// Create and initialize landscape info object
-	ULandscapeInfo* LandscapeInfo = CreateLandscapeInfo();
-
-	if (CanHaveLayersContent())
+	// Create the default layer first
+	ALandscape* LandscapeActor = GetLandscapeActor();
+	check(LandscapeActor != nullptr);
+	if (LandscapeActor->GetEditLayersConst().IsEmpty() && InImportLayers.IsEmpty())
 	{
-		// Create the default layer first
-		ALandscape* LandscapeActor = GetLandscapeActor();
-		check(LandscapeActor != nullptr);
-		if (LandscapeActor->GetEditLayersConst().IsEmpty() && InImportLayers.IsEmpty())
-		{
-			LandscapeActor->CreateDefaultLayer();
-		}
+		LandscapeActor->CreateDefaultLayer();
+	}
 
-		// Components need to be registered to be able to import the layer content and we will remove them if they should have not been visible
-		bool ShouldComponentBeRegistered = GetLevel()->bIsVisible;
-		RegisterAllComponents();
+	// Create and initialize landscape info object after default layer exists
+	ULandscapeInfo* LandscapeInfo = CreateLandscapeInfo();
 
-		TSet<ULandscapeComponent*> ComponentsToProcess;
+	// Components need to be registered to be able to import the layer content and we will remove them if they should have not been visible
+	bool ShouldComponentBeRegistered = GetLevel()->bIsVisible;
+	RegisterAllComponents();
 
-		struct FLayerImportSettings
-		{
-			FGuid SourceLayerGuid;
-			FGuid DestinationLayerGuid;
-		};
+	TSet<ULandscapeComponent*> ComponentsToProcess;
 
-		TArray<FLayerImportSettings> LayerImportSettings;		
+	struct FLayerImportSettings
+	{
+		FGuid SourceLayerGuid;
+		FGuid DestinationLayerGuid;
+	};
 
-		// Only create Layers on main Landscape
-		if (LandscapeActor == this && !InImportLayers.IsEmpty())
+	TArray<FLayerImportSettings> LayerImportSettings;		
+
+	// Only create Layers on main Landscape
+	if (LandscapeActor == this && !InImportLayers.IsEmpty())
+	{
+		for (const FLandscapeLayer& OldLayer : InImportLayers)
 		{
-			for (const FLandscapeLayer& OldLayer : InImportLayers)
+			const FLandscapeLayer* NewLayer = LandscapeActor->DuplicateLayerAndMoveBrushes(OldLayer);
+			if (NewLayer != nullptr)
 			{
-				const FLandscapeLayer* NewLayer = LandscapeActor->DuplicateLayerAndMoveBrushes(OldLayer);
-				if (NewLayer != nullptr)
-				{
-					check(NewLayer->EditLayer != nullptr); // it's possible DuplicateLayerAndMoveBrushes fails (e.g. max number of layers reached), but if not, we should always have an EditLayer
+				check(NewLayer->EditLayer != nullptr); // it's possible DuplicateLayerAndMoveBrushes fails (e.g. max number of layers reached), but if not, we should always have an EditLayer
 
-					FLayerImportSettings ImportSettings;
-					ImportSettings.SourceLayerGuid = OldLayer.EditLayer->GetGuid();
-					ImportSettings.DestinationLayerGuid = NewLayer->EditLayer->GetGuid();
-					LayerImportSettings.Add(ImportSettings);
-				}
+				FLayerImportSettings ImportSettings;
+				ImportSettings.SourceLayerGuid = OldLayer.EditLayer->GetGuid();
+				ImportSettings.DestinationLayerGuid = NewLayer->EditLayer->GetGuid();
+				LayerImportSettings.Add(ImportSettings);
 			}
-
-			LandscapeInfo->GetComponentsInRegion(InMinX, InMinY, InMaxX, InMaxY, ComponentsToProcess);
 		}
-		else
+
+		LandscapeInfo->GetComponentsInRegion(InMinX, InMinY, InMaxX, InMaxY, ComponentsToProcess);
+	}
+	else
+	{
+		// In the case of a streaming proxy, we will generate the layer data for each components that the proxy hold so no need of the grid min/max to calculate the components to update
+		if (LandscapeActor != this)
 		{
-			// In the case of a streaming proxy, we will generate the layer data for each components that the proxy hold so no need of the grid min/max to calculate the components to update
-			if (LandscapeActor != this)
-			{
-				LandscapeActor->AddLayersToProxy(this);
-			}
+			LandscapeActor->AddLayersToProxy(this);
+		}
 
-			// And we will fill all the landscape components with the provided final layer content put into the default layer (aka layer index 0)
-			const ULandscapeEditLayerBase* DefaultEditLayer = LandscapeActor->GetEditLayerConst(0);
-			check(DefaultEditLayer != nullptr);
+		// And we will fill all the landscape components with the provided final layer content put into the default layer (aka layer index 0)
+		const ULandscapeEditLayerBase* DefaultEditLayer = LandscapeActor->GetEditLayerConst(0);
+		check(DefaultEditLayer != nullptr);
 
-			FLayerImportSettings ImportSettings;
-			ImportSettings.SourceLayerGuid = FinalLayerGuid;
-			ImportSettings.DestinationLayerGuid = DefaultEditLayer->GetGuid();
-			LayerImportSettings.Add(ImportSettings);
+		FLayerImportSettings ImportSettings;
+		ImportSettings.SourceLayerGuid = FinalLayerGuid;
+		ImportSettings.DestinationLayerGuid = DefaultEditLayer->GetGuid();
+		LayerImportSettings.Add(ImportSettings);
 
-			ComponentsToProcess.Append(ToRawPtrTArrayUnsafe(LandscapeComponents));
-		}
-		// CreateDefaultLayer or DuplicateLayerAndMoveBrushes must set the SelectedEditLayerIndex during import
-		check(LandscapeActor->GetSelectedEditLayerIndex() != INDEX_NONE);
-		check(LayerImportSettings.Num() != 0);
-		// Currently only supports reimporting heightmap data into a single edit layer, which will always be the default layer
-		ReimportDestinationLayerGuid = LayerImportSettings[0].DestinationLayerGuid;
+		ComponentsToProcess.Append(ToRawPtrTArrayUnsafe(LandscapeComponents));
+	}
+	// CreateDefaultLayer or DuplicateLayerAndMoveBrushes must set the SelectedEditLayerIndex during import
+	check(LandscapeActor->GetSelectedEditLayerIndex() != INDEX_NONE);
+	check(LayerImportSettings.Num() != 0);
+	// Currently only supports reimporting heightmap data into a single edit layer, which will always be the default layer
+	ReimportDestinationLayerGuid = LayerImportSettings[0].DestinationLayerGuid;
 
-		TSet<UTexture2D*> LayersTextures;
+	TSet<UTexture2D*> LayersTextures;
 
-		for (const FLayerImportSettings& ImportSettings : LayerImportSettings)
-		{
-			FLandscapeEditDataInterface LandscapeEdit(LandscapeInfo, false);
-			FScopedSetLandscapeEditingLayer Scope(LandscapeActor, ImportSettings.DestinationLayerGuid);
+	for (const FLayerImportSettings& ImportSettings : LayerImportSettings)
+	{
+		FLandscapeEditDataInterface LandscapeEdit(LandscapeInfo, false);
+		FScopedSetLandscapeEditingLayer Scope(LandscapeActor, ImportSettings.DestinationLayerGuid);
 
-			const TArray<uint16>* ImportHeightData = InImportHeightData.Find(ImportSettings.SourceLayerGuid);
+		const TArray<uint16>* ImportHeightData = InImportHeightData.Find(ImportSettings.SourceLayerGuid);
 
-			if (ImportHeightData != nullptr && ImportHeightData->Num() != 0)
-			{
-				LandscapeEdit.SetHeightData(InMinX, InMinY, InMaxX, InMaxY, (uint16*)ImportHeightData->GetData(), 0, false, nullptr);
-			}
+		if (ImportHeightData != nullptr && ImportHeightData->Num() != 0)
+		{
+			LandscapeEdit.SetHeightData(InMinX, InMinY, InMaxX, InMaxY, (uint16*)ImportHeightData->GetData(), 0, false, nullptr);
+		}
 
-			const TArray<FLandscapeImportLayerInfo>* ImportWeightData = InImportMaterialLayerInfos.Find(ImportSettings.SourceLayerGuid);
+		const TArray<FLandscapeImportLayerInfo>* ImportWeightData = InImportMaterialLayerInfos.Find(ImportSettings.SourceLayerGuid);
 
-			if (ImportWeightData != nullptr)
+		if (ImportWeightData != nullptr)
+		{
+			for (const FLandscapeImportLayerInfo& MaterialLayerInfo : *ImportWeightData)
 			{
-				for (const FLandscapeImportLayerInfo& MaterialLayerInfo : *ImportWeightData)
+				if (MaterialLayerInfo.LayerInfo != nullptr && MaterialLayerInfo.LayerData.Num() > 0)
 				{
-					if (MaterialLayerInfo.LayerInfo != nullptr && MaterialLayerInfo.LayerData.Num() > 0)
-					{
-						LandscapeEdit.SetAlphaData(MaterialLayerInfo.LayerInfo, InMinX, InMinY, InMaxX, InMaxY, MaterialLayerInfo.LayerData.GetData(), 0, ELandscapeLayerPaintingRestriction::None, true, false);
+					LandscapeEdit.SetAlphaData(MaterialLayerInfo.LayerInfo, InMinX, InMinY, InMaxX, InMaxY, MaterialLayerInfo.LayerData.GetData(), 0, ELandscapeLayerPaintingRestriction::None, true, false);
 
-						// Add to TargetLayers if needed.
-						if (!HasTargetLayer(MaterialLayerInfo.LayerInfo->LayerName))
-						{
-							AddTargetLayer(MaterialLayerInfo.LayerInfo->LayerName, FLandscapeTargetLayerSettings(MaterialLayerInfo.LayerInfo));
-						}
+					// Add to TargetLayers if needed.
+					if (!HasTargetLayer(MaterialLayerInfo.LayerInfo->LayerName))
+					{
+						AddTargetLayer(MaterialLayerInfo.LayerInfo->LayerName, FLandscapeTargetLayerSettings(MaterialLayerInfo.LayerInfo));
 					}
 				}
 			}
-
-			for (ULandscapeComponent* Component : ComponentsToProcess)
-			{
-				FLandscapeLayerComponentData* ComponentLayerData = Component->GetLayerData(ImportSettings.DestinationLayerGuid);
-				check(ComponentLayerData != nullptr);
-
-				LayersTextures.Add(ComponentLayerData->HeightmapData.Texture);
-				LayersTextures.Append(ToRawPtrTArrayUnsafe(ComponentLayerData->WeightmapData.Textures));
-			}
 		}
 
-		// Retrigger a caching of the platform data as we wrote again in the textures
-		for (UTexture2D* Texture : LayersTextures)
+		for (ULandscapeComponent* Component : ComponentsToProcess)
 		{
-			Texture->UpdateResource();
-		}
-
-		LandscapeActor->RequestLayersContentUpdateForceAll();
+			FLandscapeLayerComponentData* ComponentLayerData = Component->GetLayerData(ImportSettings.DestinationLayerGuid);
+			check(ComponentLayerData != nullptr);
 
-		if (!ShouldComponentBeRegistered)
-		{
-			UnregisterAllComponents();
+			LayersTextures.Add(ComponentLayerData->HeightmapData.Texture);
+			LayersTextures.Append(ToRawPtrTArrayUnsafe(ComponentLayerData->WeightmapData.Textures));
 		}
-	}	
-	else
+	}
+
+	// Retrigger a caching of the platform data as we wrote again in the textures
+	for (UTexture2D* Texture : LayersTextures)
 	{
-		if (GetLevel()->bIsVisible)
-		{
-			ReregisterAllComponents();
-		}
+		Texture->UpdateResource();
+	}
+
+	LandscapeActor->RequestLayersContentUpdateForceAll();
 
-		ReimportDestinationLayerGuid = FGuid();
-		LandscapeInfo->RecreateCollisionComponents();
-		LandscapeInfo->UpdateAllAddCollisions();
+	if (!ShouldComponentBeRegistered)
+	{
+		UnregisterAllComponents();
 	}
 
 	ReimportHeightmapFilePath = InHeightmapFileName;
@@ -5472,11 +5423,8 @@ void ALandscape::PostDuplicate(bool bDuplicateForPIE)
 		RequestLayersContentUpdate(ELandscapeLayerUpdateMode::Update_All);
 	}
 
-	// If edit layers are supported, ensure the first edit layer is selected after duplicating
-	if (CanHaveLayersContent())
-	{
-		SetSelectedEditLayerIndex(0);
-	}
+	// Ensure the first edit layer is selected after duplicating
+	SetSelectedEditLayerIndex(0);
 }
 #endif	//WITH_EDITOR
 
@@ -5688,13 +5636,10 @@ void ULandscapeInfo::RemoveXYOffsets()
 	});
 }
 
+// Deprecated
 bool ULandscapeInfo::CanHaveLayersContent() const
 {
-	if (ALandscape* Landscape = LandscapeActor.Get())
-	{
-		return Landscape->CanHaveLayersContent();
-	}
-	return false;
+	return true;
 }
 
 void ULandscapeInfo::ClearDirtyData()
@@ -5954,7 +5899,6 @@ ALandscapeProxy* ULandscapeInfo::MoveComponentsToProxy(const TArray<ULandscapeCo
 
 		Component->UpdateMaterialInstances();
 	}
-	LandscapeProxy->UpdateCachedHasLayersContent();
 
 	for (ULandscapeHeightfieldCollisionComponent* Component : TargetSelectedCollisionComponents)
 	{
@@ -6001,8 +5945,6 @@ void ALandscape::SplitHeightmap(ULandscapeComponent* Comp, ALandscapeProxy* Targ
 
 	ALandscapeProxy* SrcProxy = Comp->GetLandscapeProxy();
 	ALandscapeProxy* DstProxy = TargetProxy ? TargetProxy : SrcProxy;
-	check(SrcProxy->HasLayersContent() == DstProxy->CanHaveLayersContent());
-	const bool bWithLayersContent = SrcProxy->HasLayersContent() && DstProxy->CanHaveLayersContent();
 	SrcProxy->Modify();
 	DstProxy->Modify();
 
@@ -6036,18 +5978,8 @@ void ALandscape::SplitHeightmap(ULandscapeComponent* Comp, ALandscapeProxy* Targ
 		Comp->HeightmapScaleBias = NewHeightmapScaleBias;
 		Comp->SetHeightmap(NewHeightmapTexture);
 
+		// Don't bother copying the data for edit layer based landscapes because that will be done in the next UpdateLayersContent.
 		check(Comp->GetHeightmap(false) == Comp->GetHeightmap(true));
-		// Don't bother copying the data if we're using edit layers, because that will be done in the next UpdateLayersContent. 
-		//  It would also have a weird, pervert effect, because when edit layers are in use, the "InGenerateMips = true" argument passed to the following call to SetHeightData (default value) is actually
-		//  ignored. This leads to the heightmap's FTextureSource having non-zero values on mip0 and zeroes on all the other mip levels, which is wrong because then FTextureSource's higher mips won't match what 
-		//  the GPU texture's higher mips will contain after completion of UpdateLayersContent. You would expect that it would actually get fixed after the GPU readback, but this is not actually happening, because the 
-		//  computed hash (see ULandscapeTextureHash::CalculateTextureHash64 in the following lines) only takes into account mip0! So when the readback occurs, it will consider that the FTextureSource and 
-		//  the GPU texture are in sync (even though it's not the case since their mips 1 and above are different) and will *not* update FTextureSource's higher mips accordingly!
-		if (!bWithLayersContent)
-		{
-			LandscapeEdit.SetHeightData(Comp->GetSectionBase().X, Comp->GetSectionBase().Y, Comp->GetSectionBase().X + Comp->ComponentSizeQuads, Comp->GetSectionBase().Y + Comp->ComponentSizeQuads, 
-				(uint16*)HeightData.GetData(), 0, false, (uint16*)NormalData.GetData());
-		}
 	}
 
 	OldHeightmapTexture->PostEditChange();
@@ -6103,89 +6035,88 @@ void ALandscape::SplitHeightmap(ULandscapeComponent* Comp, ALandscapeProxy* Targ
 	GDisableAutomaticTextureMaterialUpdateDependencies = false;
 
 #if WITH_EDITORONLY_DATA
-	if (bWithLayersContent)
-	{
-		FLandscapeEditLayerReadback* NewCPUReadback = new FLandscapeEditLayerReadback();
-		bool bHashIsValid = true;
-		const uint64 Hash = ULandscapeTextureHash::CalculateTextureHash64(NewHeightmapTexture, ELandscapeTextureType::Heightmap, bHashIsValid);
-		NewCPUReadback->SetHash(Hash);
-		DstProxy->HeightmapsCPUReadback.Add(NewHeightmapTexture, NewCPUReadback);
 
-		// Free OldHeightmapTexture's CPUReadBackResource if not used by any component
-		bool FreeCPUReadBack = true;
-		for (ULandscapeComponent* Component : SrcProxy->LandscapeComponents)
+	FLandscapeEditLayerReadback* NewCPUReadback = new FLandscapeEditLayerReadback();
+	bool bHashIsValid = true;
+	const uint64 Hash = ULandscapeTextureHash::CalculateTextureHash64(NewHeightmapTexture, ELandscapeTextureType::Heightmap, bHashIsValid);
+	NewCPUReadback->SetHash(Hash);
+	DstProxy->HeightmapsCPUReadback.Add(NewHeightmapTexture, NewCPUReadback);
+
+	// Free OldHeightmapTexture's CPUReadBackResource if not used by any component
+	bool FreeCPUReadBack = true;
+	for (ULandscapeComponent* Component : SrcProxy->LandscapeComponents)
+	{
+		if (Component != Comp && Component->GetHeightmap(false) == OldHeightmapTexture)
 		{
-			if (Component != Comp && Component->GetHeightmap(false) == OldHeightmapTexture)
-			{
-				FreeCPUReadBack = false;
-				break;
-			}
+			FreeCPUReadBack = false;
+			break;
 		}
-		if (FreeCPUReadBack)
+	}
+	if (FreeCPUReadBack)
+	{
+		FLandscapeEditLayerReadback** OldCPUReadback = SrcProxy->HeightmapsCPUReadback.Find(OldHeightmapTexture);
+		if (OldCPUReadback)
 		{
-			FLandscapeEditLayerReadback** OldCPUReadback = SrcProxy->HeightmapsCPUReadback.Find(OldHeightmapTexture);
-			if (OldCPUReadback)
+			if (FLandscapeEditLayerReadback* ResourceToDelete = *OldCPUReadback)
 			{
-				if (FLandscapeEditLayerReadback* ResourceToDelete = *OldCPUReadback)
-				{
-					delete ResourceToDelete;
-					SrcProxy->HeightmapsCPUReadback.Remove(OldHeightmapTexture);
-				}
+				delete ResourceToDelete;
+				SrcProxy->HeightmapsCPUReadback.Remove(OldHeightmapTexture);
 			}
 		}
+	}
 
 
-		// TODO-LS: Although we are only reading from the existing heightmap texture, the current code in FLandscapeEditDataInterface
-		// is performing a call to LockMip(), which results in the UnlockMip() modifying the TextureSource CompressionFormat.
-		// Modification to textures MUST be wrapped in PreEditChange/PostEditChange() in order to prevent texture compilation issues.
-		// But really, the ideal fix here would be to call LockMipReadOnly() instead of LockMip(), and completely avoid modifying the texture.
-		TArray<UTexture2D*> OldLayerHeightmaps;
+	// TODO-LS: Although we are only reading from the existing heightmap texture, the current code in FLandscapeEditDataInterface
+	// is performing a call to LockMip(), which results in the UnlockMip() modifying the TextureSource CompressionFormat.
+	// Modification to textures MUST be wrapped in PreEditChange/PostEditChange() in order to prevent texture compilation issues.
+	// But really, the ideal fix here would be to call LockMipReadOnly() instead of LockMip(), and completely avoid modifying the texture.
+	TArray<UTexture2D*> OldLayerHeightmaps;
+
+	{
+		// Move layer content to new layer heightmap
+		FLandscapeEditDataInterface LandscapeEdit(Info);
+		ALandscape* Landscape = Info->LandscapeActor.Get();
 
+		Comp->ForEachLayer([&](const FGuid& LayerGuid, FLandscapeLayerComponentData& LayerData)
 		{
-			// Move layer content to new layer heightmap
-			FLandscapeEditDataInterface LandscapeEdit(Info);
-			ALandscape* Landscape = Info->LandscapeActor.Get();
-
-			Comp->ForEachLayer([&](const FGuid& LayerGuid, FLandscapeLayerComponentData& LayerData)
+			UTexture2D* OldLayerHeightmap = LayerData.HeightmapData.Texture;
+			if (OldLayerHeightmap != nullptr)
 			{
-				UTexture2D* OldLayerHeightmap = LayerData.HeightmapData.Texture;
-				if (OldLayerHeightmap != nullptr)
+				FScopedSetLandscapeEditingLayer Scope(Landscape, LayerGuid);
+
+				OldLayerHeightmaps.Add(OldLayerHeightmap);
+				OldLayerHeightmap->PreEditChange(nullptr);
+
+				// Read old data and split
+				TArray<uint8> LayerHeightData;
+				LayerHeightData.AddZeroed((1 + Comp->ComponentSizeQuads) * (1 + Comp->ComponentSizeQuads) * sizeof(uint16));
+				// Because of edge problem, normal would be just copy from old component data
+				TArray<uint8> LayerNormalData;
+				LayerNormalData.AddZeroed((1 + Comp->ComponentSizeQuads) * (1 + Comp->ComponentSizeQuads) * sizeof(uint16));
+
+				// Read using old heightmap scale/bias
+				Comp->HeightmapScaleBias = OldHeightmapScaleBias;
+				LandscapeEdit.GetHeightDataFast(Comp->GetSectionBase().X, Comp->GetSectionBase().Y, Comp->GetSectionBase().X + Comp->ComponentSizeQuads, Comp->GetSectionBase().Y + Comp->ComponentSizeQuads, (uint16*)LayerHeightData.GetData(), 0, (uint16*)LayerNormalData.GetData());
+				// Restore new heightmap scale/bias
+				Comp->HeightmapScaleBias = NewHeightmapScaleBias;
 				{
-					FScopedSetLandscapeEditingLayer Scope(Landscape, LayerGuid);
-
-					OldLayerHeightmaps.Add(OldLayerHeightmap);
-					OldLayerHeightmap->PreEditChange(nullptr);
-
-					// Read old data and split
-					TArray<uint8> LayerHeightData;
-					LayerHeightData.AddZeroed((1 + Comp->ComponentSizeQuads) * (1 + Comp->ComponentSizeQuads) * sizeof(uint16));
-					// Because of edge problem, normal would be just copy from old component data
-					TArray<uint8> LayerNormalData;
-					LayerNormalData.AddZeroed((1 + Comp->ComponentSizeQuads) * (1 + Comp->ComponentSizeQuads) * sizeof(uint16));
-
-					// Read using old heightmap scale/bias
-					Comp->HeightmapScaleBias = OldHeightmapScaleBias;
-					LandscapeEdit.GetHeightDataFast(Comp->GetSectionBase().X, Comp->GetSectionBase().Y, Comp->GetSectionBase().X + Comp->ComponentSizeQuads, Comp->GetSectionBase().Y + Comp->ComponentSizeQuads, (uint16*)LayerHeightData.GetData(), 0, (uint16*)LayerNormalData.GetData());
-					// Restore new heightmap scale/bias
-					Comp->HeightmapScaleBias = NewHeightmapScaleBias;
-					{
-						// no mipchain required as these layer weight maps are used in layer compositing to generate a final set of weight maps to be used for rendering
-						UTexture2D* LayerHeightmapTexture = DstProxy->CreateLandscapeTexture(HeightmapSizeU, HeightmapSizeV, TEXTUREGROUP_Terrain_Heightmap, TSF_BGRA8, /* OptionalOverrideOuter = */ nullptr, /* bCompress = */ false, /* bMipChain = */ false);
-						const bool bClear = true;
-						ULandscapeComponent::CreateEmptyTextureMips(LayerHeightmapTexture, ELandscapeTextureUsage::FinalData, ELandscapeTextureType::Heightmap, bClear);
-						LayerHeightmapTexture->PostEditChange();
-						// Set Layer heightmap texture
-						LayerData.HeightmapData.Texture = LayerHeightmapTexture;
-						LandscapeEdit.SetHeightData(Comp->GetSectionBase().X, Comp->GetSectionBase().Y, Comp->GetSectionBase().X + Comp->ComponentSizeQuads, Comp->GetSectionBase().Y + Comp->ComponentSizeQuads, (uint16*)LayerHeightData.GetData(), 0, false, (uint16*)LayerNormalData.GetData());
-					}
+					// no mipchain required as these layer weight maps are used in layer compositing to generate a final set of weight maps to be used for rendering
+					UTexture2D* LayerHeightmapTexture = DstProxy->CreateLandscapeTexture(HeightmapSizeU, HeightmapSizeV, TEXTUREGROUP_Terrain_Heightmap, TSF_BGRA8, /* OptionalOverrideOuter = */ nullptr, /* bCompress = */ false, /* bMipChain = */ false);
+					const bool bClear = true;
+					ULandscapeComponent::CreateEmptyTextureMips(LayerHeightmapTexture, ELandscapeTextureUsage::FinalData, ELandscapeTextureType::Heightmap, bClear);
+					LayerHeightmapTexture->PostEditChange();
+					// Set Layer heightmap texture
+					LayerData.HeightmapData.Texture = LayerHeightmapTexture;
+					LandscapeEdit.SetHeightData(Comp->GetSectionBase().X, Comp->GetSectionBase().Y, Comp->GetSectionBase().X + Comp->ComponentSizeQuads, Comp->GetSectionBase().Y + Comp->ComponentSizeQuads, (uint16*)LayerHeightData.GetData(), 0, false, (uint16*)LayerNormalData.GetData());
 				}
-			});
-
-			Landscape->RequestLayersContentUpdateForceAll();
-		}
+			}
+		});
 
-		Algo::ForEach(OldLayerHeightmaps, [](UTexture2D* Tex) { Tex->PostEditChange(); });
+		Landscape->RequestLayersContentUpdateForceAll();
 	}
+
+	Algo::ForEach(OldLayerHeightmaps, [](UTexture2D* Tex) { Tex->PostEditChange(); });
+
 #endif
 
 	// Reregister
@@ -7024,16 +6955,6 @@ void ALandscape::PostEditChangeProperty(FPropertyChangedEvent& PropertyChangedEv
 			MarkAllLandscapeRenderStateDirty();
 		}
 
-		// Update normals if DrawScale3D is changed
-		if (MemberPropertyName == FName(TEXT("RelativeScale3D")))
-		{
-			if (!Info->CanHaveLayersContent())
-			{
-				FLandscapeEditDataInterface LandscapeEdit(Info);
-				LandscapeEdit.RecalculateNormals();
-			}
-		}
-
 		if (bNeedsRecalcBoundingBox || bMaterialChanged || bChangedLighting)
 		{
 			// We cannot iterate the XYtoComponentMap directly because reregistering components modifies the array.
@@ -7301,11 +7222,6 @@ TArray<UTexture*> ULandscapeComponent::ReallocateWeightmaps(FLandscapeEditDataIn
 	
 	const bool bIsFinalWeightmap = !InEditLayerGuid.IsValid();
 
-	if (InEditLayerGuid.IsValid())
-	{
-		check(TargetProxy->HasLayersContent());
-	}
-
 	TArray<FWeightmapLayerAllocationInfo>& ComponentWeightmapLayerAllocations = GetWeightmapLayerAllocations(InEditLayerGuid);
 	TArray<TObjectPtr<UTexture2D>>& ComponentWeightmapTextures = GetWeightmapTextures(InEditLayerGuid);
 	TArray<TObjectPtr<ULandscapeWeightmapUsage>>& ComponentWeightmapTexturesUsage = GetWeightmapTexturesUsage(InEditLayerGuid);
@@ -8358,7 +8274,7 @@ TArray<FLinearColor> ALandscapeProxy::SampleRTData(UTextureRenderTarget2D* InRen
 	}
 }
 
-bool ALandscapeProxy::LandscapeImportHeightmapFromRenderTarget(UTextureRenderTarget2D* InRenderTarget, bool InImportHeightFromRGChannel)
+bool ALandscapeProxy::LandscapeImportHeightmapFromRenderTarget(UTextureRenderTarget2D* InRenderTarget, bool InImportHeightFromRGChannel, int32 InEditLayerIndex)
 {
 	uint64 StartCycle = FPlatformTime::Cycles64();
 
@@ -8369,15 +8285,17 @@ bool ALandscapeProxy::LandscapeImportHeightmapFromRenderTarget(UTextureRenderTar
 		return false;
 	}
 
-	if (Landscape->HasLayersContent())
+	ULandscapeInfo* LandscapeInfo = Landscape->GetLandscapeInfo();
+
+	const ULandscapeEditLayerBase* EditLayer = LandscapeInfo->LandscapeActor->GetEditLayer(InEditLayerIndex);
+	if (EditLayer == nullptr)
 	{
-		//todo: Support an edit layer name input parameter to support import to edit layers.
-		FMessageLog("Blueprint").Error(LOCTEXT("LandscapeImportHeightmapFromRenderTarget_LandscapeLayersNotSupported", "LandscapeImportHeightmapFromRenderTarget: Cannot import to landscape with Edit Layers enabled."));
+		FMessageLog("Blueprint").Error(LOCTEXT("LandscapeImportHeightmapFromRenderTarget_InvalidEditLayer", "LandscapeImportHeightmapFromRenderTarget: Invalid edit layer index for import."));
 		return false;
 	}
+	const FGuid& EditLayerGuid = LandscapeInfo->LandscapeActor->GetEditLayer(InEditLayerIndex)->GetGuid();
 
 	int32 MinX, MinY, MaxX, MaxY;
-	ULandscapeInfo* LandscapeInfo = Landscape->GetLandscapeInfo();
 
 	if (!LandscapeInfo->GetLandscapeExtent(MinX, MinY, MaxX, MaxY))
 	{
@@ -8449,6 +8367,8 @@ bool ALandscapeProxy::LandscapeImportHeightmapFromRenderTarget(UTextureRenderTar
 	FScopedTransaction Transaction(LOCTEXT("Undo_ImportHeightmap", "Importing Landscape Heightmap"));
 
 	FHeightmapAccessor<false> HeightmapAccessor(LandscapeInfo);
+	
+	HeightmapAccessor.SetEditLayer(EditLayerGuid);
 	HeightmapAccessor.SetData(MinX, MinY, SampleRect.Width() - 1, SampleRect.Height() - 1, HeightData.GetData());
 
 	double SecondsTaken = FPlatformTime::ToSeconds64(FPlatformTime::Cycles64() - StartCycle);
@@ -8460,6 +8380,7 @@ bool ALandscapeProxy::LandscapeImportHeightmapFromRenderTarget(UTextureRenderTar
 
 bool ALandscapeProxy::LandscapeExportHeightmapToRenderTarget(UTextureRenderTarget2D* InRenderTarget, bool bInExportHeightIntoRGChannel, bool InExportLandscapeProxies)
 {
+	// TODO [jonathan.bard] : This is redundant with RenderWeightmap, which is available in engine as well, so it'd probably best to deprecate this
 #if WITH_EDITOR
 	uint64 StartCycle = FPlatformTime::Cycles64();
 
@@ -8626,19 +8547,19 @@ TArray<FName> ALandscape::GetTargetLayerNames(bool bInIncludeVisibilityLayer) co
 
 #if WITH_EDITOR
 
-bool ALandscapeProxy::LandscapeImportWeightmapFromRenderTarget(UTextureRenderTarget2D* InRenderTarget, FName InLayerName)
+bool ALandscapeProxy::LandscapeImportWeightmapFromRenderTarget(UTextureRenderTarget2D* InRenderTarget, FName InLayerName, int32 InEditLayerIndex)
 {
-	ALandscape* Landscape = GetLandscapeActor();
-	if (Landscape != nullptr)
+	if (ALandscape* Landscape = GetLandscapeActor())
 	{
-		if (Landscape->HasLayersContent())
+		ULandscapeInfo* LandscapeInfo = Landscape->GetLandscapeInfo();
+
+		const ULandscapeEditLayerBase* EditLayer = LandscapeInfo->LandscapeActor->GetEditLayer(InEditLayerIndex);
+		if (EditLayer == nullptr)
 		{
-			//todo: Support an edit layer name input parameter to support import to edit layers.
-			FMessageLog("Blueprint").Error(LOCTEXT("LandscapeImportWeightmapFromRenderTarget_LandscapeLayersNotSupported", "LandscapeImportWeightmapFromRenderTarget: Cannot import to landscape with Edit Layers enabled."));
+			FMessageLog("Blueprint").Error(LOCTEXT("LandscapeImportRenderTarget_InvalidEditLayer", "LandscapeImportWeightmapFromRenderTarget: Invalid edit layer index for import."));
 			return false;
 		}
-
-		ULandscapeInfo* LandscapeInfo = Landscape->GetLandscapeInfo();
+		const FGuid& EditLayerGuid = LandscapeInfo->LandscapeActor->GetEditLayer(InEditLayerIndex)->GetGuid();
 
 		int32 MinX, MinY, MaxX, MaxY;
 		if (LandscapeInfo && LandscapeInfo->GetLandscapeExtent(MinX, MinY, MaxX, MaxY))
@@ -8681,7 +8602,8 @@ bool ALandscapeProxy::LandscapeImportWeightmapFromRenderTarget(UTextureRenderTar
 
 				FScopedTransaction Transaction(LOCTEXT("Undo_ImportWeightmap", "Importing Landscape Layer"));
 
-				FAlphamapAccessor<false, false> AlphamapAccessor(LandscapeInfo, CurWeightmapInfo.LayerInfoObj);
+				FAlphamapAccessor<false, false> AlphamapAccessor(LandscapeInfo, CurWeightmapInfo.LayerInfoObj);\
+				AlphamapAccessor.SetEditLayer(EditLayerGuid);
 				AlphamapAccessor.SetData(MinX, MinY, MaxX, MaxY, LayerData.GetData(), ELandscapeLayerPaintingRestriction::None);
 
 				uint64 CycleEnd = FPlatformTime::Cycles64();
@@ -8704,9 +8626,9 @@ bool ALandscapeProxy::LandscapeImportWeightmapFromRenderTarget(UTextureRenderTar
 	return false;
 }
 
+// Deprecated
 bool ALandscapeProxy::LandscapeExportWeightmapToRenderTarget(UTextureRenderTarget2D* InRenderTarget, FName InLayerName)
 {
-	// TODO [jonathan.bard] : Not implemented... Is redundant with RenderWeightmap, which is available in engine as well, so it'd probably best to deprecate LandscapeExportWeightmapToRenderTarget, as well as LandscapeExportHeightmapToRenderTarget 
 	return false;
 }
 
diff --git a/Engine/Source/Runtime/Landscape/Private/LandscapeEditInterface.cpp b/Engine/Source/Runtime/Landscape/Private/LandscapeEditInterface.cpp
--- a/Engine/Source/Runtime/Landscape/Private/LandscapeEditInterface.cpp
+++ b/Engine/Source/Runtime/Landscape/Private/LandscapeEditInterface.cpp
@@ -451,66 +451,9 @@ void FLandscapeEditDataInterface::SetHeightData(int32 X1, int32 Y1, int32 X2, in
 			}
 												
 			Component->RequestHeightmapUpdate();
-			if (!Component->GetLandscapeProxy()->HasLayersContent())
-			{
-				// See if we need to adjust the bounds. Note we never shrink the bounding box at this point
-				bool bUpdateBoxSphereBounds = false;
-
-				if (InUpdateBounds)
-				{
-					float MinLocalZ = LandscapeDataAccess::GetLocalHeight(MinHeight);
-					float MaxLocalZ = LandscapeDataAccess::GetLocalHeight(MaxHeight);
-
-					if (MinLocalZ < Component->CachedLocalBox.Min.Z)
-					{
-						Component->CachedLocalBox.Min.Z = MinLocalZ;
-						bUpdateBoxSphereBounds = true;
-					}
-					if (MaxLocalZ > Component->CachedLocalBox.Max.Z)
-					{
-						Component->CachedLocalBox.Max.Z = MaxLocalZ;
-						bUpdateBoxSphereBounds = true;
-					}
-
-					if (bUpdateBoxSphereBounds)
-					{
-						Component->UpdateComponentToWorld();
-					}
-				}
-
-				// Update mipmaps
-
-				// Work out how many mips should be calculated directly from one component's data.
-				// The remaining mips are calculated on a per texture basis.
-				// eg if subsection is 7x7 quads, we need one 3 mips total: (8x8, 4x4, 2x2 verts)
-				if (InGenerateMips)
-				{
-					int32 BaseNumMips = FMath::CeilLogTwo(SubsectionSizeQuads + 1);
-					TArray<FColor*> MipData;
-					MipData.AddUninitialized(BaseNumMips);
-					MipData[0] = HeightmapTextureData;
-					for (int32 MipIdx = 1; MipIdx < BaseNumMips; MipIdx++)
-					{
-						MipData[MipIdx] = (FColor*)TexDataInfo->GetMipData(MipIdx);
-					}
-					Component->GenerateHeightmapMips(MipData, ComponentX1, ComponentY1, ComponentX2, ComponentY2, TexDataInfo);
-
-					if (InUpdateCollision)
-					{
-						// Update collision
-						Component->UpdateCollisionHeightData(
-							MipData[Component->CollisionMipLevel],
-							Component->SimpleCollisionMipLevel > Component->CollisionMipLevel ? MipData[Component->SimpleCollisionMipLevel] : nullptr,
-							ComponentX1, ComponentY1, ComponentX2, ComponentY2, bUpdateBoxSphereBounds,
-							XYOffsetMipData);
-					}
-				}
-			}
-			else
-			{
-				// In Layer, cumulate dirty collision region (will be used next time UpdateCollisionHeightData is called)
-				Component->UpdateDirtyCollisionHeightData(FIntRect(ComponentX1, ComponentY1, ComponentX2, ComponentY2));
-			}
+			// In Edit Layer based landscapes, cumulate dirty collision region (will be used next time UpdateCollisionHeightData is called)
+			Component->UpdateDirtyCollisionHeightData(FIntRect(ComponentX1, ComponentY1, ComponentX2, ComponentY2));
+			
 		}
 	}
 
@@ -1802,41 +1745,7 @@ void ULandscapeComponent::DeleteLayerInternal(ULandscapeLayerInfoObject* LayerIn
 		}
 	}
 
-	// Update the shaders for this component
-	if (!Component->GetLandscapeProxy()->HasLayersContent())
-	{
-		Component->UpdateMaterialInstances();
-		Component->EditToolRenderData.UpdateDebugColorMaterial(Component);
-		Component->UpdateEditToolRenderData();
-	}
 	Component->RequestWeightmapUpdate();
-	
-	// Update dominant layer info stored in collision component
-	if (!Component->GetLandscapeProxy()->HasLayersContent())
-	{
-		TArray<FColor*> CollisionWeightmapMipData;
-		for (int32 WeightmapIdx = 0; WeightmapIdx < ComponentWeightmapTextures.Num(); WeightmapIdx++)
-		{
-			CollisionWeightmapMipData.Add((FColor*)LandscapeEdit.GetTextureDataInfo(ComponentWeightmapTextures[WeightmapIdx])->GetMipData(Component->CollisionMipLevel));
-		}
-		TArray<FColor*> SimpleCollisionWeightmapMipData;
-		if (Component->SimpleCollisionMipLevel > Component->CollisionMipLevel)
-		{
-			for (int32 WeightmapIdx = 0; WeightmapIdx < ComponentWeightmapTextures.Num(); WeightmapIdx++)
-			{
-				SimpleCollisionWeightmapMipData.Add((FColor*)LandscapeEdit.GetTextureDataInfo(ComponentWeightmapTextures[WeightmapIdx])->GetMipData(Component->SimpleCollisionMipLevel));
-			}
-		}
-		Component->UpdateCollisionLayerData(
-			CollisionWeightmapMipData.GetData(),
-			Component->SimpleCollisionMipLevel > Component->CollisionMipLevel ? SimpleCollisionWeightmapMipData.GetData() : nullptr);
-
-		ULandscapeHeightfieldCollisionComponent* LocalCollisionComponent = Component->GetCollisionComponent();
-		if (LocalCollisionComponent)
-		{
-			LocalCollisionComponent->RecreateCollision();
-		}
-	}
 }
 
 void FLandscapeEditDataInterface::DeleteLayer(ULandscapeLayerInfoObject* LayerInfo)
@@ -1852,7 +1761,7 @@ void FLandscapeEditDataInterface::DeleteLayer(ULandscapeLayerInfoObject* LayerIn
 		Component->DeleteLayerInternal(LayerInfo, *this, GetEditLayer());
 	}
 
-	if (LandscapeInfo->LandscapeActor.IsValid() && LandscapeInfo->LandscapeActor->HasLayersContent())
+	if (LandscapeInfo->LandscapeActor.IsValid())
 	{
 		LandscapeInfo->LandscapeActor->RequestLayersContentUpdate(ELandscapeLayerUpdateMode::Update_All);
 	}
@@ -2019,43 +1928,8 @@ void ULandscapeComponent::FillLayer(ULandscapeLayerInfoObject* LayerInfo, FLands
 		}
 	}
 
-	// Update the shaders for this component
-	if (!Component->GetLandscapeProxy()->HasLayersContent())
-	{
-		Component->InvalidateLightingCache();
-		Component->UpdateMaterialInstances();
-		Component->EditToolRenderData.UpdateDebugColorMaterial(Component);
-		Component->UpdateEditToolRenderData();
-	}
 	Component->RequestWeightmapUpdate();
 
-	// Update dominant layer info stored in collision component
-	if (!Component->GetLandscapeProxy()->HasLayersContent())
-	{
-		TArray<FColor*> CollisionWeightmapMipData;
-		for (int32 WeightmapIdx = 0; WeightmapIdx < ComponentWeightmapTextures.Num(); WeightmapIdx++)
-		{
-			CollisionWeightmapMipData.Add((FColor*)LandscapeEdit.GetTextureDataInfo(ComponentWeightmapTextures[WeightmapIdx])->GetMipData(Component->CollisionMipLevel));
-		}
-		TArray<FColor*> SimpleCollisionWeightmapMipData;
-		if (Component->SimpleCollisionMipLevel > Component->CollisionMipLevel)
-		{
-			for (int32 WeightmapIdx = 0; WeightmapIdx < ComponentWeightmapTextures.Num(); WeightmapIdx++)
-			{
-				SimpleCollisionWeightmapMipData.Add((FColor*)LandscapeEdit.GetTextureDataInfo(ComponentWeightmapTextures[WeightmapIdx])->GetMipData(Component->SimpleCollisionMipLevel));
-			}
-		}
-		Component->UpdateCollisionLayerData(
-			CollisionWeightmapMipData.GetData(),
-			Component->SimpleCollisionMipLevel > Component->CollisionMipLevel ? SimpleCollisionWeightmapMipData.GetData() : nullptr);
-
-		ULandscapeHeightfieldCollisionComponent* LocalCollisionComponent = Component->GetCollisionComponent();
-		if (LocalCollisionComponent)
-		{
-			LocalCollisionComponent->RecreateCollision();
-		}
-	}
-
 	Proxy->ValidateProxyLayersWeightmapUsage();
 }
 
@@ -2105,7 +1979,7 @@ void FLandscapeEditDataInterface::FillLayer(ULandscapeLayerInfoObject* LayerInfo
 		}
 	}
 
-	if (LandscapeInfo->LandscapeActor.IsValid() && LandscapeInfo->LandscapeActor->HasLayersContent())
+	if (LandscapeInfo->LandscapeActor.IsValid())
 	{
 		LandscapeInfo->LandscapeActor->RequestLayersContentUpdateForceAll();
 	}
@@ -2302,46 +2176,9 @@ void ULandscapeComponent::ReplaceLayerInternal(ULandscapeLayerInfoObject* FromLa
 
 		// Remove the layer
 		ComponentWeightmapLayerAllocations.RemoveAt(FromLayerIdx);
-
-		if (!GetLandscapeProxy()->HasLayersContent())
-		{
-			// Update the shaders for this component
-			UpdateMaterialInstances();
-
-			EditToolRenderData.UpdateDebugColorMaterial(this);
-			UpdateEditToolRenderData();
-		}
 	}
 
 	RequestWeightmapUpdate();
-
-	if (!LandscapeEdit.HasLandscapeLayersContent())
-	{
-		// Update dominant layer info stored in collision component
-		TArray<FColor*> CollisionWeightmapMipData;
-		for (int32 WeightmapIdx = 0; WeightmapIdx < ComponentWeightmapTextures.Num(); WeightmapIdx++)
-		{
-			CollisionWeightmapMipData.Add((FColor*)LandscapeEdit.GetTextureDataInfo(ComponentWeightmapTextures[WeightmapIdx])->GetMipData(CollisionMipLevel));
-		}
-		TArray<FColor*> SimpleCollisionWeightmapMipData;
-		if (SimpleCollisionMipLevel > CollisionMipLevel)
-		{
-			for (int32 WeightmapIdx = 0; WeightmapIdx < ComponentWeightmapTextures.Num(); WeightmapIdx++)
-			{
-				SimpleCollisionWeightmapMipData.Add((FColor*)LandscapeEdit.GetTextureDataInfo(ComponentWeightmapTextures[WeightmapIdx])->GetMipData(SimpleCollisionMipLevel));
-			}
-		}
-		UpdateCollisionLayerData(
-			CollisionWeightmapMipData.GetData(),
-			SimpleCollisionMipLevel > CollisionMipLevel ? SimpleCollisionWeightmapMipData.GetData() : nullptr);
-
-
-		if(ULandscapeHeightfieldCollisionComponent* CollisionComp = GetCollisionComponent())
-		{
-			CollisionComp->RecreateCollision();
-		}
-	}
-
 	GetLandscapeProxy()->ValidateProxyLayersWeightmapUsage();
 }
 
@@ -2358,7 +2195,7 @@ void FLandscapeEditDataInterface::ReplaceLayer(ULandscapeLayerInfoObject* FromLa
 		}
 	};
 
-	if (LandscapeInfo->LandscapeActor.IsValid() && LandscapeInfo->LandscapeActor->HasLayersContent())
+	if (LandscapeInfo->LandscapeActor.IsValid())
 	{
 		LandscapeInfo->LandscapeActor->Modify(GetShouldDirtyPackage());
 		LandscapeInfo->LandscapeActor->ReplaceLayerSubstractiveBlendStatus(FromLayerInfo, ToLayerInfo, GetShouldDirtyPackage());
@@ -2878,13 +2715,6 @@ void FLandscapeEditDataInterface::SetAlphaData(ULandscapeLayerInfoObject* const
 				UpdateLayerIdx = ComponentWeightmapLayerAllocations.Num();
 				new (ComponentWeightmapLayerAllocations) FWeightmapLayerAllocationInfo(LayerInfo);
 				Component->ReallocateWeightmaps(/*DataInterface =*/this, GetEditLayer(), /*bInSaveToTransactionBuffer = */true, /*bool bInForceReallocate = */false, /*InTargetProxy = */nullptr, /*InRestrictSharingToComponents = */nullptr);
-
-				if (!Component->GetLandscapeProxy()->HasLayersContent())
-				{
-					Component->UpdateMaterialInstances();
-					Component->EditToolRenderData.UpdateDebugColorMaterial(Component);
-					Component->UpdateEditToolRenderData();
-				}
 				Component->RequestWeightmapUpdate();
 			}
 
@@ -3160,46 +2990,6 @@ void FLandscapeEditDataInterface::SetAlphaData(ULandscapeLayerInfoObject* const
 				}
 			}
 
-			if (!Component->GetLandscapeProxy()->HasLayersContent())
-			{
-				// Update mipmaps
-				CollisionWeightmapMipData.Reset();
-				CollisionWeightmapMipData.AddUninitialized(ComponentWeightmapTextures.Num());
-				for (int32 WeightmapIdx = 0; WeightmapIdx < ComponentWeightmapTextures.Num(); WeightmapIdx++)
-				{
-					UTexture2D* const WeightmapTexture = ComponentWeightmapTextures[WeightmapIdx];
-
-					const int32 NumMips = WeightmapTexture->Source.GetNumMips();
-					WeightmapTextureMipData.Reset();
-					WeightmapTextureMipData.AddUninitialized(NumMips);
-					for (int32 MipIdx = 0; MipIdx < NumMips; MipIdx++)
-					{
-						FColor* const MipData = (FColor*)TexDataInfos[WeightmapIdx]->GetMipData(MipIdx);
-						WeightmapTextureMipData[MipIdx] = MipData;
-					}
-					CollisionWeightmapMipData[WeightmapIdx] = WeightmapTextureMipData[Component->CollisionMipLevel];
-
-					ULandscapeComponent::UpdateWeightmapMips(ComponentNumSubsections, SubsectionSizeQuads, WeightmapTexture, WeightmapTextureMipData, ComponentX1, ComponentY1, ComponentX2, ComponentY2, TexDataInfos[WeightmapIdx]);
-					WeightmapTextureMipData.Reset();
-				}
-
-				if (Component->SimpleCollisionMipLevel > Component->CollisionMipLevel)
-				{
-					for (int32 WeightmapIdx = 0; WeightmapIdx < ComponentWeightmapTextures.Num(); WeightmapIdx++)
-					{
-						SimpleCollisionWeightmapMipData.Add((FColor*)TexDataInfos[WeightmapIdx]->GetMipData(Component->SimpleCollisionMipLevel));
-					}
-				}
-
-				// Update dominant layer info stored in collision component
-				Component->UpdateCollisionLayerData(
-					CollisionWeightmapMipData.GetData(),
-					Component->SimpleCollisionMipLevel > Component->CollisionMipLevel ? SimpleCollisionWeightmapMipData.GetData() : nullptr,
-					ComponentX1, ComponentY1, ComponentX2, ComponentY2);
-				CollisionWeightmapMipData.Reset();
-				SimpleCollisionWeightmapMipData.Reset();
-			}
-
 			// Check if we need to remove weightmap allocations for layers that were completely painted away
 			bool bRemovedLayer = false;
 			for (int32 LayerIdx = 0; LayerIdx < ComponentWeightmapLayerAllocations.Num(); LayerIdx++)
@@ -3221,12 +3011,6 @@ void FLandscapeEditDataInterface::SetAlphaData(ULandscapeLayerInfoObject* const
 
 			if (bRemovedLayer)
 			{
-				if (!Component->GetLandscapeProxy()->HasLayersContent())
-				{
-					Component->UpdateMaterialInstances();
-					Component->EditToolRenderData.UpdateDebugColorMaterial(Component);
-					Component->UpdateEditToolRenderData();
-				}
 				Component->RequestWeightmapUpdate();
 			}
 
@@ -3332,13 +3116,6 @@ void FLandscapeEditDataInterface::SetAlphaData(const TSet<ULandscapeLayerInfoObj
 						ComponentWeightmapLayerAllocations.Emplace(LayerInfoNeedingAllocation);
 					}
 					Component->ReallocateWeightmaps(/*DataInterface =*/this, GetEditLayer(), /*bInSaveToTransactionBuffer = */true, /*bool bInForceReallocate = */false, /*InTargetProxy = */nullptr, /*InRestrictSharingToComponents = */nullptr);
-					
-					if (!Component->GetLandscapeProxy()->HasLayersContent())
-					{
-						Component->UpdateMaterialInstances();
-						Component->EditToolRenderData.UpdateDebugColorMaterial(Component);
-						Component->UpdateEditToolRenderData();
-					}
 					Component->RequestWeightmapUpdate();
 				}
 			}
@@ -3449,46 +3226,6 @@ void FLandscapeEditDataInterface::SetAlphaData(const TSet<ULandscapeLayerInfoObj
 				}
 			}
 
-			if (!Component->GetLandscapeProxy()->HasLayersContent())
-			{
-				// Update mipmaps
-				CollisionWeightmapMipData.Reset();
-				CollisionWeightmapMipData.AddUninitialized(ComponentWeightmapTextures.Num());
-				for (int32 WeightmapIdx = 0; WeightmapIdx < ComponentWeightmapTextures.Num(); WeightmapIdx++)
-				{
-					UTexture2D* const WeightmapTexture = ComponentWeightmapTextures[WeightmapIdx];
-
-					const int32 NumMips = WeightmapTexture->Source.GetNumMips();
-					WeightmapTextureMipData.Reset();
-					WeightmapTextureMipData.AddUninitialized(NumMips);
-					for (int32 MipIdx = 0; MipIdx < NumMips; MipIdx++)
-					{
-						FColor* const MipData = (FColor*)TexDataInfos[WeightmapIdx]->GetMipData(MipIdx);
-						WeightmapTextureMipData[MipIdx] = MipData;
-					}
-					CollisionWeightmapMipData[WeightmapIdx] = WeightmapTextureMipData[Component->CollisionMipLevel];
-
-					ULandscapeComponent::UpdateWeightmapMips(ComponentNumSubsections, SubsectionSizeQuads, WeightmapTexture, WeightmapTextureMipData, ComponentX1, ComponentY1, ComponentX2, ComponentY2, TexDataInfos[WeightmapIdx]);
-					WeightmapTextureMipData.Reset();
-				}
-
-				if (Component->SimpleCollisionMipLevel > Component->CollisionMipLevel)
-				{
-					for (int32 WeightmapIdx = 0; WeightmapIdx < ComponentWeightmapTextures.Num(); WeightmapIdx++)
-					{
-						SimpleCollisionWeightmapMipData.Add((FColor*)TexDataInfos[WeightmapIdx]->GetMipData(Component->SimpleCollisionMipLevel));
-					}
-				}
-
-				// Update dominant layer info stored in collision component
-				Component->UpdateCollisionLayerData(
-					CollisionWeightmapMipData.GetData(),
-					Component->SimpleCollisionMipLevel > Component->CollisionMipLevel ? SimpleCollisionWeightmapMipData.GetData() : nullptr,
-					ComponentX1, ComponentY1, ComponentX2, ComponentY2);
-				CollisionWeightmapMipData.Reset();
-				SimpleCollisionWeightmapMipData.Reset();
-			}
-
 			// Check if we need to remove weightmap allocations for layers that were completely painted away
 			bool bRemovedLayer = false;
 			for (int32 LayerIdx = 0; LayerIdx < ComponentWeightmapLayerAllocations.Num(); LayerIdx++)
@@ -3510,12 +3247,6 @@ void FLandscapeEditDataInterface::SetAlphaData(const TSet<ULandscapeLayerInfoObj
 
 			if (bRemovedLayer)
 			{
-				if (!Component->GetLandscapeProxy()->HasLayersContent())
-				{
-					Component->UpdateMaterialInstances();
-					Component->EditToolRenderData.UpdateDebugColorMaterial(Component);
-					Component->UpdateEditToolRenderData();
-				}
 				Component->RequestWeightmapUpdate();
 			}
 
@@ -5650,15 +5381,16 @@ const ALandscape* FLandscapeEditDataInterface::GetTargetLandscape() const
 	return LandscapeInfo ? LandscapeInfo->LandscapeActor.Get() : nullptr;
 }
 
+// Deprecated 
 bool FLandscapeEditDataInterface::CanHaveLandscapeLayersContent() const
 {
-	return LandscapeInfo ? LandscapeInfo->CanHaveLayersContent() : false;
+	return true;
 }
 
+// Deprecated 
 bool FLandscapeEditDataInterface::HasLandscapeLayersContent() const
 {
-	const ALandscape* Landscape = GetTargetLandscape();
-	return Landscape ? Landscape->HasLayersContent() : false;
+	return true;
 }
 
 void FLandscapeEditDataInterface::GetXYOffsetDataFast(const int32 X1, const int32 Y1, const int32 X2, const int32 Y2, FVector2D* Data, int32 Stride)
diff --git a/Engine/Source/Runtime/Landscape/Private/LandscapeEditLayers.cpp b/Engine/Source/Runtime/Landscape/Private/LandscapeEditLayers.cpp
--- a/Engine/Source/Runtime/Landscape/Private/LandscapeEditLayers.cpp
+++ b/Engine/Source/Runtime/Landscape/Private/LandscapeEditLayers.cpp
@@ -2536,6 +2536,16 @@ bool ALandscape::AllowSelectiveRenderEditLayers() const
 		   CVarLandscapeSelectiveRenderEditLayers.GetValueOnGameThread();
 }
 
+bool ALandscape::GetIsPendingEditLayerConversion() const
+{
+	return bIsPendingEditLayerConversion;
+}
+
+void ALandscape::SetIsPendingEditLayerConversion(bool bInIsPendingEditLayerConversion)
+{
+	bIsPendingEditLayerConversion = bInIsPendingEditLayerConversion;
+}
+
 bool ALandscape::CreateLayersRenderingResource(bool bInForceLayerResourceReset)
 {
 	ULandscapeInfo* Info = GetLandscapeInfo();
@@ -2694,38 +2704,33 @@ bool ALandscape::CreateLayersRenderingResource(bool bInForceLayerResourceReset)
 	return true;
 }
 
-void ALandscape::ToggleCanHaveLayersContent()
+void ALandscape::ConvertNonEditLayerLandscape()
 {
 	// Detect any attempt to call this in the middle of UpdateLayersContent.  If called from blueprint, log an error and return early instead of asserting.
-	if (InLayerUpdateCount > 0 && UE::Landscape::Private::InBPCallstack())
+	if (LayerUpdateCount > 0 && UE::Landscape::Private::InBPCallstack())
 	{
-		UE_LOG(LogLandscapeBP, Error, TEXT("Attempting to make illegal call to ToggleCanHaveLayersContent during UpdateLayersContent."));
+		UE_LOG(LogLandscape, Error, TEXT("Attempting to make illegal call to ConvertNonEditLayerLandscape during UpdateLayersContent."));
 		return;
 	}
-	check(InLayerUpdateCount == 0);
-
-	// TODO: All the toggle functions (handling edit layers) should ideally run before updating bCanHaveLayersContent
-	bCanHaveLayersContent = !bCanHaveLayersContent;
+	check(LayerUpdateCount == 0);
 
-	if (!bCanHaveLayersContent)
+	// If this landscape already has edit-layers it cannot be converted again
+	if (!LandscapeEditLayers.IsEmpty())
 	{
-		check(!LandscapeEditLayers.IsEmpty());
-		ReleaseLayersRenderingResource();
-		DeleteLayers();
-		check(SelectedEditLayerIndex == INDEX_NONE);
-	}
-	else
-	{
-		check(LandscapeEditLayers.IsEmpty());
-		CreateDefaultLayer();
-		CopyOldDataToDefaultLayer();
-		check(SelectedEditLayerIndex == 0);
+		UE_LOG(LogLandscape, Warning, TEXT("Attempting to ConvertNonEditLayerLandscape on a landscape with edit-layer data."));
+		return;
 	}
 
-	if (LandscapeEdMode)
-	{
-		LandscapeEdMode->OnCanHaveLayersContentChanged();
-	}
+	CreateDefaultLayer();
+	check(SelectedEditLayerIndex == 0);
+
+	// This will only copy data from proxies that have completed the registration process
+	CopyOldDataToDefaultLayer();
+}
+ 
+// Deprecated
+void ALandscape::ToggleCanHaveLayersContent()
+{
 }
 
 void ALandscape::ReleaseLayersRenderingResource()
@@ -2834,6 +2839,7 @@ void ALandscape::CopyOldDataToDefaultLayer()
 		return;
 	}
 
+	// This will only copy data from proxies that have completed the registration process
 	Info->ForEachLandscapeProxy([&](ALandscapeProxy* Proxy)
 	{
 		CopyOldDataToDefaultLayer(Proxy);
@@ -2850,9 +2856,9 @@ void ALandscape::CopyOldDataToDefaultLayer(ALandscapeProxy* InProxy)
 		return;
 	}
 
-	if (InProxy->LandscapeComponents.IsEmpty())
+	// No data to migrate or proxy has not finished registration
+	if (InProxy->LandscapeComponents.IsEmpty() || !InProxy->bIsLandscapeActorRegisteredWithLandscapeInfo)
 	{
-		// No data to migrate, we can early-out to avoid modifying the proxy:
 		return;
 	}
 
@@ -2877,6 +2883,9 @@ void ALandscape::CopyOldDataToDefaultLayer(ALandscapeProxy* InProxy)
 
 		if (ensure(LayerData != nullptr && LayerData->IsInitialized()))
 		{
+			// If this component has data on the default edit layer it should not be copied again
+			check(LayerData->WeightmapData.LayerAllocations.Num() == 0);
+
 			// Heightmap
 			UTexture2D* ComponentHeightmap = Component->GetHeightmap();
 
@@ -2993,6 +3002,10 @@ void ALandscape::CopyOldDataToDefaultLayer(ALandscapeProxy* InProxy)
 			Allocation.Free();
 		}
 	}
+
+	// Data is only copied for proxies with valid components
+	// After copying, ensure the InProxy's components are valid and have an edit layer in their LayersData map (Default edit layer)
+	check(!InProxy->LandscapeComponents.IsEmpty() && (InProxy->LandscapeComponents[0] != nullptr) && InProxy->LandscapeComponents[0]->HasLayersData());
 }
 
 void ALandscape::UpdateProxyLayersWeightmapUsage()
@@ -11786,24 +11799,20 @@ void ALandscape::ResolveLayersWeightmapTexture(
 	InvalidateGeneratedComponentData(ChangedComponents, bInvalidateLightingCache);
 }
 
+// Deprecated
 bool ALandscape::HasLayersContent() const
 {
-	return LandscapeEditLayers.Num() > 0;
+	return true;
 }
 
+// Deprecated
 void ALandscape::UpdateCachedHasLayersContent(bool bInCheckComponentDataIntegrity)
 {
-	Super::UpdateCachedHasLayersContent(bInCheckComponentDataIntegrity);
-
-	// For consistency with the ALandscape::HasLayersContent() override above, make sure the cached bHasLayersContent boolean is also valid when we have at least one edit layer : 
-	//  Otherwise, as ALandscapeProxy::UpdateCachedHasLayersContent relies on the presence of landscape components and in distributed landscape setups (one ALandscape + multiple ALandscapeStreamingProxy),  
-	//  the "parent" ALandscape actor doesn't have any landscape component, hence it would have bHasLayersContent erroneously set to false (while ALandscape::HasLayersContent() would actually return true!)
-	bHasLayersContent |= ALandscape::HasLayersContent();
 }
 
 void ALandscape::RequestLayersInitialization(bool bInRequestContentUpdate, bool bInForceLayerResourceReset)
 {
-	if (!CanHaveLayersContent())
+	if (IsTemplate())
 	{
 		return;
 	}
@@ -11820,7 +11829,7 @@ void ALandscape::RequestLayersInitialization(bool bInRequestContentUpdate, bool
 
 void ALandscape::RequestSplineLayerUpdate()
 {
-	if (HasLayersContent() && FindLayerOfTypeConst(ULandscapeEditLayerSplines::StaticClass()) != nullptr)
+	if (!IsTemplate() && FindLayerOfTypeConst(ULandscapeEditLayerSplines::StaticClass()) != nullptr)
 	{
 		bSplineLayerUpdateRequested = true;
 	}
@@ -11839,7 +11848,7 @@ void ALandscape::RequestLayersContentUpdateForceAll(ELandscapeLayerUpdateMode In
 		return;
 	}
 
-	if (!CanHaveLayersContent())
+	if (IsTemplate())
 	{
 		return;
 	}
@@ -12155,7 +12164,7 @@ bool ALandscape::CanUpdateLayersContent() const
 	UWorld* World = GetWorld();
 	ULandscapeSubsystem* Subsystem = World ? GetWorld()->GetSubsystem<ULandscapeSubsystem>() : nullptr;
 
-	return FApp::CanEverRender() && (LandscapeInfo != nullptr) && CanHaveLayersContent() &&
+	return FApp::CanEverRender() && (LandscapeInfo != nullptr) && !IsTemplate() &&
 		LandscapeInfo->AreAllComponentsRegistered() && LandscapeInfo->SupportsLandscapeEditing() &&
 		Subsystem && Subsystem->GetTextureStreamingManager();
 }
@@ -12290,14 +12299,14 @@ void ALandscape::UpdateLayersContent(bool bInWaitForStreaming, bool bInSkipMonit
 	TRACE_CPUPROFILER_EVENT_SCOPE(LandscapeLayers_UpdateLayersContent);
 
 	// Detect any attempt to re-enter.  If called from blueprint, log an error and return early instead of asserting.
-	if (InLayerUpdateCount > 0 && UE::Landscape::Private::InBPCallstack())
+	if (LayerUpdateCount > 0 && UE::Landscape::Private::InBPCallstack())
 	{
 		UE_LOG(LogLandscapeBP, Error, TEXT("Attempting to make illegal re-entrant call to UpdateLayersContent."));
 		return;
 	}
 
-	check(InLayerUpdateCount == 0);
-	InLayerUpdateCount++;
+	check(LayerUpdateCount == 0);
+	LayerUpdateCount++;
 
 	bool bHideNotifications = true;
 	ON_SCOPE_EXIT
@@ -12317,7 +12326,7 @@ void ALandscape::UpdateLayersContent(bool bInWaitForStreaming, bool bInSkipMonit
 			FLandscapeEditLayerReadback::GarbageCollectTasks();
 		}
 
-		InLayerUpdateCount--;
+		LayerUpdateCount--;
 	};
 
 	// Note : no early-out allowed before this : even if not actually updating edit layers, we need to poll our resources in order to make sure we register to streaming events when needed: 
@@ -12701,8 +12710,6 @@ uint32 ALandscape::UpdateAfterReadbackResolves(const TArrayView<FLandscapeEditLa
 
 void ALandscape::InitializeLayers()
 {
-	check(HasLayersContent());
-
 	if (CreateLayersRenderingResource(bLandscapeLayersForceResourceReset))
 	{
 		InitializeLandscapeLayersWeightmapUsage();
@@ -12828,40 +12835,21 @@ void ALandscapeProxy::FinishDestroy()
 }
 
 #if WITH_EDITOR
+// Deprecated
 bool ALandscapeProxy::CanHaveLayersContent() const
 {
-	if (HasAnyFlags(RF_ClassDefaultObject | RF_ArchetypeObject))
-	{
-		return false;
-	}
-
-	if (const ALandscape* LandscapeActor = GetLandscapeActor())
-	{
-		return LandscapeActor->bCanHaveLayersContent;
-	}
-
-	return false;
+	return true;
 }
 
+// Deprecated
 bool ALandscapeProxy::HasLayersContent() const
 {
-	return bHasLayersContent || (GetLandscapeActor() != nullptr && GetLandscapeActor()->HasLayersContent());
+	return true;
 }
 
+// Deprecated
 void ALandscapeProxy::UpdateCachedHasLayersContent(bool InCheckComponentDataIntegrity)
 {
-	// TODO: do we need to update the texture types here?
-
-	// In the case of InCheckComponentDataIntegrity we will loop through all components to make sure they all have the same state and in the other case we will assume that the 1st component represent the state of all the others.
-	bHasLayersContent = (!LandscapeComponents.IsEmpty() && (LandscapeComponents[0] != nullptr)) ? LandscapeComponents[0]->HasLayersData() : false;
-
-	if (InCheckComponentDataIntegrity)
-	{
-		for (const ULandscapeComponent* Component : LandscapeComponents)
-		{
-			check((Component == nullptr) || (bHasLayersContent == Component->HasLayersData()));
-		}
-	}
 }
 
 namespace
@@ -13011,8 +12999,6 @@ bool ALandscapeProxy::AddLayer(const FGuid& InLayerGuid)
 		}
 	}
 
-	UpdateCachedHasLayersContent();
-
 	if (bModified)
 	{
 		InitializeLayerWithEmptyContent(InLayerGuid);
@@ -13050,8 +13036,6 @@ void ALandscapeProxy::DeleteLayer(const FGuid& InLayerGuid)
 			}
 		}
 	}
-
-	UpdateCachedHasLayersContent();
 }
 
 void ALandscapeProxy::InitializeLayerWithEmptyContent(const FGuid& InLayerGuid)
@@ -13169,7 +13153,7 @@ bool ALandscape::IsUpToDate() const
 	}
 
 #if WITH_EDITORONLY_DATA
-	if (CanHaveLayersContent() && GetWorld() != nullptr && !GetWorld()->IsGameWorld())
+	if (!IsTemplate() && GetWorld() != nullptr && !GetWorld()->IsGameWorld())
 	{
 		return LayerContentUpdateModes == 0 && !FLandscapeEditLayerReadback::HasWork();
 	}
@@ -13472,30 +13456,44 @@ TArray<const FLandscapeLayer*> ALandscape::GetLayersOfTypeConst(const TSubclassO
 	return Result;
 }
 
+// Deprecated
 void ALandscape::DeleteLayers()
 {
-	SelectedEditLayerIndex = INDEX_NONE;
+}
 
+void ALandscape::CollapseEditLayers()
+{
+	if (IsTemplate())
+	{
+		return;
+	}
+
+	// Remove all existing edit layers
 	for (int32 LayerIndex = LandscapeEditLayers.Num() - 1; LayerIndex >= 0; --LayerIndex)
 	{
 		DeleteLayer(LayerIndex);
 	}
+
+	// Create a new empty layer since all landscapes require >= 1 layer
+	CreateDefaultLayer();
+	CopyOldDataToDefaultLayer();
+	check(SelectedEditLayerIndex == 0);
 }
 
 bool ALandscape::DeleteLayer(int32 InLayerIndex)
 {
-	// ToggleCanHaveLayers sets the bCanHaveLayers flag before calling DeleteLayers()
-	// When toggling from an layer to a non-layer state, HasLayersContent becomes false but this function still needs to run
-	// ensure Selected Index is INDEX_NONE in this case
-	ensure(HasLayersContent() || SelectedEditLayerIndex == INDEX_NONE);
+	if (IsTemplate())
+	{
+		return false;
+	}
 
 	// Detect any attempt to call this in the middle of UpdateLayersContent.  If called from blueprint, log an error and return early instead of asserting.
-	if (InLayerUpdateCount > 0 && UE::Landscape::Private::InBPCallstack())
+	if (LayerUpdateCount > 0 && UE::Landscape::Private::InBPCallstack())
 	{
 		UE_LOG(LogLandscapeBP, Error, TEXT("Attempting to make illegal call to DeleteLayer during UpdateLayersContent."));
 		return false;
 	}
-	check(InLayerUpdateCount == 0);
+	check(LayerUpdateCount == 0);
 
 	const FLandscapeLayer* LayerStruct = GetLayerConst(InLayerIndex);
 	if (!LayerStruct)
@@ -13507,7 +13505,7 @@ bool ALandscape::DeleteLayer(int32 InLayerIndex)
 
 	// If the layer to delete is below the current selected layer index, shift the selected index down
 	// Ignore index update when toggling from an layer to a non-layer state 
-	if (HasLayersContent() && SelectedEditLayerIndex >= InLayerIndex)
+	if (SelectedEditLayerIndex >= InLayerIndex)
 	{
 		if (LandscapeEditLayers.IsValidIndex(InLayerIndex - 1))
 		{
@@ -13806,8 +13804,6 @@ void ALandscape::ClearEditLayer(int32 InLayerIndex, TSet<TObjectPtr<ULandscapeCo
 
 void ALandscape::ClearEditLayer(const FGuid& InLayerGuid, TSet<TObjectPtr<ULandscapeComponent>>* InComponents, ELandscapeToolTargetTypeFlags InClearMode, bool bMarkPackageDirty)
 {
-	ensure(HasLayersContent());
-
 	const ULandscapeEditLayerBase* EditLayer = GetEditLayerConst(InLayerGuid);
 	ULandscapeInfo* LandscapeInfo = GetLandscapeInfo();
 	if (!LandscapeInfo || !EditLayer)
@@ -13992,7 +13988,7 @@ uint32 ULandscapeComponent::ComputeLayerHash(bool InReturnEditingHash) const
 void ALandscape::UpdateLandscapeSplines(const FGuid& InTargetLayer, bool bInUpdateOnlySelected, bool bInForceUpdateAllCompoments)
 {
 	TRACE_CPUPROFILER_EVENT_SCOPE(LandscapeLayers_UpdateLandscapeSplines);
-	check(CanHaveLayersContent());
+	check(!IsTemplate());
 	ULandscapeInfo* LandscapeInfo = GetLandscapeInfo();
 	const ULandscapeEditLayerBase* SplinesEditLayer = FindEditLayerOfTypeConst(ULandscapeEditLayerSplines::StaticClass());
 	FGuid TargetLayerGuid = (SplinesEditLayer != nullptr) ? SplinesEditLayer->GetGuid() : InTargetLayer;
@@ -14070,7 +14066,7 @@ FScopedSetLandscapeEditingLayer::FScopedSetLandscapeEditingLayer(ALandscape* InL
 	: Landscape(InLandscape)
 	, CompletionCallback(MoveTemp(InCompletionCallback))
 {
-	if (Landscape.IsValid() && Landscape.Get()->CanHaveLayersContent())
+	if (Landscape.IsValid() && !Landscape.Get()->IsTemplate())
 	{
 		PreviousLayerGUID = Landscape->GetEditingLayer();
 		Landscape->SetEditingLayer(InLayerGUID);
@@ -14079,7 +14075,7 @@ FScopedSetLandscapeEditingLayer::FScopedSetLandscapeEditingLayer(ALandscape* InL
 
 FScopedSetLandscapeEditingLayer::~FScopedSetLandscapeEditingLayer()
 {
-	if (Landscape.IsValid() && Landscape.Get()->CanHaveLayersContent())
+	if (Landscape.IsValid() && !Landscape.Get()->IsTemplate())
 	{
 		Landscape->SetEditingLayer(PreviousLayerGUID);
 		if (CompletionCallback)
@@ -14091,7 +14087,7 @@ FScopedSetLandscapeEditingLayer::~FScopedSetLandscapeEditingLayer()
 
 void ALandscape::SetEditingLayer(const FGuid& InLayerGuid)
 {
-	ensure(CanHaveLayersContent());
+	ensure(!IsTemplate());
 
 	ULandscapeInfo* LandscapeInfo = GetLandscapeInfo();
 	if (!LandscapeInfo)
@@ -14117,14 +14113,14 @@ const FGuid& ALandscape::GetEditingLayer() const
 
 void ALandscape::SetSelectedEditLayerIndex(const int32 InEditLayerIndex)
 {
-	check(CanHaveLayersContent() ? LandscapeEditLayers.IsValidIndex(InEditLayerIndex) : InEditLayerIndex == INDEX_NONE);
+	check(LandscapeEditLayers.IsValidIndex(InEditLayerIndex) && !IsTemplate());
 	SelectedEditLayerIndex = InEditLayerIndex;
 }
 
 const int32 ALandscape::GetSelectedEditLayerIndex() const
 {
 	// When edit layers are not supported, index should always be NONE
-	check(CanHaveLayersContent() ? LandscapeEditLayers.IsValidIndex(SelectedEditLayerIndex) : SelectedEditLayerIndex == INDEX_NONE);
+	check(LandscapeEditLayers.IsValidIndex(SelectedEditLayerIndex) && !IsTemplate());
 	return SelectedEditLayerIndex;
 }
 
@@ -14135,7 +14131,7 @@ bool ALandscape::IsMaxLayersReached() const
 
 void ALandscape::CreateDefaultLayer()
 {
-	if (!CanHaveLayersContent())
+	if (IsTemplate())
 	{
 		return;
 	}
@@ -14149,7 +14145,7 @@ void ALandscape::CreateDefaultLayer()
 FLandscapeLayer* ALandscape::DuplicateLayerAndMoveBrushes(const FLandscapeLayer& InOtherLayer)
 {
 	ULandscapeInfo* LandscapeInfo = GetLandscapeInfo();
-	if (!LandscapeInfo || !CanHaveLayersContent())
+	if (!LandscapeInfo || IsTemplate())
 	{
 		return nullptr;
 	}
@@ -14192,14 +14188,14 @@ FLandscapeLayer* ALandscape::DuplicateLayerAndMoveBrushes(const FLandscapeLayer&
 int32 ALandscape::CreateLayer(FName InName, const TSubclassOf<ULandscapeEditLayerBase>& InEditLayerClass, bool bInIgnoreLayerCountLimit)
 {
 	// Detect any attempt to call this in the middle of UpdateLayersContent.  If called from blueprint, log an error and return early instead of asserting.
-	if (InLayerUpdateCount > 0 && UE::Landscape::Private::InBPCallstack())
+	if (LayerUpdateCount > 0 && UE::Landscape::Private::InBPCallstack())
 	{
 		UE_LOG(LogLandscapeBP, Error, TEXT("Attempting to make illegal call to CreateLayer during UpdateLayersContent."));
 		return INDEX_NONE;
 	}
-	check(InLayerUpdateCount == 0);
+	check(LayerUpdateCount == 0);
 
-	if ((!bInIgnoreLayerCountLimit && IsMaxLayersReached()) || !CanHaveLayersContent())
+	if ((!bInIgnoreLayerCountLimit && IsMaxLayersReached()) || IsTemplate())
 	{
 		return INDEX_NONE;
 	}
@@ -14254,7 +14250,7 @@ void ALandscape::OnLayerCreatedInternal(ULandscapeEditLayerBase* EditLayer)
 void ALandscape::AddLayersToProxy(ALandscapeProxy* InProxy)
 {
 	ULandscapeInfo* LandscapeInfo = GetLandscapeInfo();
-	if (!LandscapeInfo || !CanHaveLayersContent())
+	if (!LandscapeInfo || IsTemplate())
 	{
 		return;
 	}
@@ -14467,7 +14463,7 @@ void ALandscape::OnLayerInfoSplineFalloffModulationChanged(ULandscapeLayerInfoOb
 	}
 
 	ALandscape* Landscape = LandscapeInfo->LandscapeActor.Get();
-	if (!Landscape || !Landscape->HasLayersContent())
+	if (!Landscape || IsTemplate())
 	{
 		return;
 	}
diff --git a/Engine/Source/Runtime/Landscape/Private/LandscapeSplineRaster.cpp b/Engine/Source/Runtime/Landscape/Private/LandscapeSplineRaster.cpp
--- a/Engine/Source/Runtime/Landscape/Private/LandscapeSplineRaster.cpp
+++ b/Engine/Source/Runtime/Landscape/Private/LandscapeSplineRaster.cpp
@@ -292,7 +292,7 @@ void RasterizeHeight(int32& MinX, int32& MinY, int32& MaxX, int32& MaxY, FLandsc
 	check(Landscape != nullptr);
 
 	bool bIsEditingLayerReservedForSplines = false;
-	if (Landscape->CanHaveLayersContent())
+	if (!Landscape->IsTemplate())
 	{
 		const ULandscapeEditLayerBase* EditingLayer = Landscape->GetEditLayerConst(Landscape->GetEditingLayer());
 		check(EditingLayer != nullptr);
@@ -791,32 +791,14 @@ bool ULandscapeInfo::ApplySplinesInternal(bool bOnlySelected, TScriptInterface<I
 	}
 
 	LandscapeEdit.Flush();
-		
-	if (!CanHaveLayersContent())
-	{
-		ALandscapeProxy::InvalidateGeneratedComponentData(ModifiedComponents);
-	}
 
 	for (ULandscapeComponent* Component : ModifiedComponents)
 	{	
-		if (Component->GetLandscapeProxy()->HasLayersContent())
+		Component->RequestHeightmapUpdate();
+		Component->RequestWeightmapUpdate();
+		if (OutModifiedComponents)
 		{
-			Component->RequestHeightmapUpdate();
-			Component->RequestWeightmapUpdate();
-			if (OutModifiedComponents)
-			{
-				OutModifiedComponents->Add(Component);
-			}
-		}
-		else
-		{
-			// Recreate collision for modified components and update the navmesh
-			ULandscapeHeightfieldCollisionComponent* CollisionComponent = Component->GetCollisionComponent();
-			if (CollisionComponent)
-			{
-				CollisionComponent->RecreateCollision();
-				FNavigationSystem::UpdateComponentData(*CollisionComponent);
-			}
+			OutModifiedComponents->Add(Component);
 		}
 	}
 	
@@ -910,20 +892,6 @@ namespace LandscapeSplineRaster
 		}
 
 		LandscapeEdit.Flush();
-
-		if (!LandscapeProxy->HasLayersContent())
-		{
-			for (ULandscapeComponent* Component : ModifiedComponents)
-			{
-				// Recreate collision for modified components and update the navmesh
-				ULandscapeHeightfieldCollisionComponent* CollisionComponent = Component->GetCollisionComponent();
-				if (CollisionComponent)
-				{
-					CollisionComponent->RecreateCollision();
-					FNavigationSystem::UpdateComponentData(*CollisionComponent);
-				}
-			}
-		}
 	}
 
 	static bool LineIntersect(const FVector2D& L1Start, const FVector2D& L1End, const FVector2D& L2Start, const FVector2D& L2End, FVector2D& Intersect, float Tolerance = KINDA_SMALL_NUMBER)
diff --git a/Engine/Source/Runtime/Landscape/Private/LandscapeSplines.cpp b/Engine/Source/Runtime/Landscape/Private/LandscapeSplines.cpp
--- a/Engine/Source/Runtime/Landscape/Private/LandscapeSplines.cpp
+++ b/Engine/Source/Runtime/Landscape/Private/LandscapeSplines.cpp
@@ -120,7 +120,7 @@ struct FLandscapeFixSplines
 							}
 						});
 
-						if (Pair.Value->LandscapeActor.IsValid() && Pair.Value->LandscapeActor->HasLayersContent())
+						if (Pair.Value->LandscapeActor.IsValid() && !Pair.Value->LandscapeActor->IsTemplate())
 						{
 							Pair.Value->LandscapeActor->RequestSplineLayerUpdate();
 						}
diff --git a/Engine/Source/Runtime/Landscape/Public/LandscapeEdit.h b/Engine/Source/Runtime/Landscape/Public/LandscapeEdit.h
--- a/Engine/Source/Runtime/Landscape/Public/LandscapeEdit.h
+++ b/Engine/Source/Runtime/Landscape/Public/LandscapeEdit.h
@@ -39,11 +39,12 @@ class ULandscapeLayerInfoObject;
 class ILandscapeEdModeInterface
 {
 public:
-	virtual void PostUpdateLayerContent() = 0;
-	virtual ELandscapeToolTargetType GetLandscapeToolTargetType() const = 0;
-	virtual const ULandscapeEditLayerBase* GetLandscapeSelectedLayer() const = 0;
-	virtual ULandscapeLayerInfoObject* GetSelectedLandscapeLayerInfo() const = 0;
-	virtual void OnCanHaveLayersContentChanged() = 0;
+	virtual void PostUpdateLayerContent() PURE_VIRTUAL(ILandscapeEdModeInterface::PostUpdateLayerContent(), return;);
+	virtual ELandscapeToolTargetType GetLandscapeToolTargetType() const PURE_VIRTUAL(ILandscapeEdModeInterface::GetLandscapeToolTargetType(), return ELandscapeToolTargetType::Invalid;);
+	virtual const ULandscapeEditLayerBase* GetLandscapeSelectedLayer() const PURE_VIRTUAL(ILandscapeEdModeInterface::GetLandscapeSelectedLayer(), return nullptr;);
+	virtual ULandscapeLayerInfoObject* GetSelectedLandscapeLayerInfo() const PURE_VIRTUAL(ILandscapeEdModeInterface::GetSelectedLandscapeLayerInfo(), return nullptr;);
+	UE_DEPRECATED(5.7, "Non-edit layer landscapes are deprecated, all landscapes use the edit layer system now.")
+	virtual void OnCanHaveLayersContentChanged() PURE_VIRTUAL(ILandscapeEdModeInterface::OnCanHaveLayersContentChanged(), return;);
 };
 
 struct FLandscapeTextureDataInfo
@@ -295,7 +296,9 @@ struct FLandscapeEditDataInterface : public FLandscapeTextureDataInterface
 
 	LANDSCAPE_API const ALandscape* GetTargetLandscape() const;
 
+	UE_DEPRECATED(5.7, "Non-edit layer landscapes are deprecated, all landscapes use the edit layer system now.")
 	LANDSCAPE_API bool CanHaveLandscapeLayersContent() const;
+	UE_DEPRECATED(5.7, "Non-edit layer landscapes are deprecated, all landscapes use the edit layer system now.")
 	LANDSCAPE_API bool HasLandscapeLayersContent() const;
 
 private:
@@ -458,61 +461,8 @@ struct FHeightmapAccessor
 			{
 				Component->RequestHeightmapUpdate();
 			}
-						
-			// Notify foliage to move any attached instances
-			bool bUpdateFoliage = false;
-			bool bUpdateNormals = false;
 
-            // Landscape Layers are updates are delayed and done in  ALandscape::TickLayers
-			if (!LandscapeEdit->HasLandscapeLayersContent())
-			{
-				const bool bInvalidateLightingCache = true;
-				ALandscapeProxy::InvalidateGeneratedComponentData(Components, bInvalidateLightingCache);
-				bUpdateNormals = true;
-				for (ULandscapeComponent* Component : Components)
-				{
-					ULandscapeHeightfieldCollisionComponent* CollisionComponent = Component->GetCollisionComponent();
-					if (CollisionComponent && AInstancedFoliageActor::HasFoliageAttached(CollisionComponent))
-					{
-						bUpdateFoliage = true;
-						break;
-					}
-				}
-			}
-
-			if (bUpdateFoliage)
-			{
-				// Calculate landscape local-space bounding box of old data, to look for foliage instances.
-				TArray<ULandscapeHeightfieldCollisionComponent*> CollisionComponents;
-				CollisionComponents.Empty(Components.Num());
-				TArray<FBox> PreUpdateLocalBoxes;
-				PreUpdateLocalBoxes.Empty(Components.Num());
-
-				for (ULandscapeComponent* Component : Components)
-				{
-					ULandscapeHeightfieldCollisionComponent* CollisionComponent = Component->GetCollisionComponent();
-					if (CollisionComponent)
-					{
-						CollisionComponents.Add(CollisionComponent);
-						PreUpdateLocalBoxes.Add(FBox(FVector((float)X1, (float)Y1, Component->CachedLocalBox.Min.Z), FVector((float)X2, (float)Y2, Component->CachedLocalBox.Max.Z)));
-					}
-				}
-
-				// Update landscape.
-				LandscapeEdit->SetHeightData(X1, Y1, X2, Y2, Data, 0, bUpdateNormals);
-
-				// Snap foliage for each component.
-				for (int32 Index = 0; Index < CollisionComponents.Num(); ++Index)
-				{
-					ULandscapeHeightfieldCollisionComponent* CollisionComponent = CollisionComponents[Index];
-					CollisionComponent->SnapFoliageInstances(PreUpdateLocalBoxes[Index].TransformBy(LandscapeInfo->GetLandscapeProxy()->LandscapeActorToWorld().ToMatrixWithScale()).ExpandBy(1.0f));
-				}
-			}
-			else
-			{
-				// No foliage, just update landscape.
-				LandscapeEdit->SetHeightData(X1, Y1, X2, Y2, Data, 0, bUpdateNormals);
-			}
+			LandscapeEdit->SetHeightData(X1, Y1, X2, Y2, Data, 0, /*InCalcNormals=*/ false);
 		}
 	}
 
@@ -526,25 +476,6 @@ struct FHeightmapAccessor
 		// Flush here manually so it will release the lock of the textures, as we will re lock for other things afterward
 		Flush();
 
-		// Landscape Layers are updates are delayed and done in  ALandscape::TickLayers
-		if (!LandscapeEdit->HasLandscapeLayersContent())
-		{
-			// Update the bounds and navmesh for the components we edited
-			for (TSet<ULandscapeComponent*>::TConstIterator It(ChangedComponents); It; ++It)
-			{
-				(*It)->UpdateCachedBounds();
-				(*It)->UpdateComponentToWorld();
-
-				// Recreate collision for modified components to update the physical materials
-				ULandscapeHeightfieldCollisionComponent* CollisionComponent = (*It)->GetCollisionComponent();
-				if (CollisionComponent)
-				{
-					CollisionComponent->RecreateCollision();
-					FNavigationSystem::UpdateComponentData(*CollisionComponent);
-				}
-			}
-		}
-
 		delete LandscapeEdit;
 		LandscapeEdit = NULL;
 	}
@@ -583,28 +514,6 @@ struct FAlphamapAccessor
 		}
 	}
 
-	~FAlphamapAccessor()
-	{
-		if (!LandscapeEdit.HasLandscapeLayersContent())
-		{
-			// Recreate collision for modified components to update the physical materials
-			for (ULandscapeComponent* Component : ModifiedComponents)
-			{
-				ULandscapeHeightfieldCollisionComponent* CollisionComponent = Component->GetCollisionComponent();
-				if (CollisionComponent)
-				{
-					CollisionComponent->RecreateCollision();
-
-					// We need to trigger navigation mesh build, in case user have painted holes on a landscape
-					if (LayerInfo == ALandscapeProxy::VisibilityLayer)
-					{
-						FNavigationSystem::UpdateComponentData(*CollisionComponent);
-					}
-				}
-			}
-		}
-	}
-
 	void GetData(int32& X1, int32& Y1, int32& X2, int32& Y2, TMap<FIntPoint, uint8>& Data)
 	{
 		LandscapeEdit.GetWeightData(LayerInfo, X1, Y1, X2, Y2, Data);
@@ -642,12 +551,6 @@ struct FAlphamapAccessor
 				LandscapeComponent->RequestWeightmapUpdate();
 			}
 			
-			if (!LandscapeEdit.HasLandscapeLayersContent())
-			{
-				ALandscapeProxy::InvalidateGeneratedComponentData(Components);
-				ModifiedComponents.Append(Components);
-			}
-
 			LandscapeEdit.SetAlphaData(LayerInfo, X1, Y1, X2, Y2, Data, 0, PaintingRestriction, bBlendWeight, bUseTotalNormalize);
 			//LayerInfo->IsReferencedFromLoadedData = true;
 		}
