Landscape Editor - Fix static analysis warnings for non-edit layer deprecation CL

diff --git a/Engine/Source/Runtime/Landscape/Private/LandscapeEdit.cpp b/Engine/Source/Runtime/Landscape/Private/LandscapeEdit.cpp
--- a/Engine/Source/Runtime/Landscape/Private/LandscapeEdit.cpp
+++ b/Engine/Source/Runtime/Landscape/Private/LandscapeEdit.cpp
@@ -8279,14 +8279,13 @@ bool ALandscapeProxy::LandscapeImportHeightmapFromRenderTarget(UTextureRenderTar
 	uint64 StartCycle = FPlatformTime::Cycles64();
 
 	ALandscape* Landscape = GetLandscapeActor();
-	if (Landscape == nullptr)
+	ULandscapeInfo* LandscapeInfo = Landscape->GetLandscapeInfo();
+	if (Landscape == nullptr || LandscapeInfo == nullptr)
 	{
 		FMessageLog("Blueprint").Error(LOCTEXT("LandscapeImportHeightmapFromRenderTarget_NullLandscape", "LandscapeImportHeightmapFromRenderTarget: Landscape must be non-null."));
 		return false;
 	}
 
-	ULandscapeInfo* LandscapeInfo = Landscape->GetLandscapeInfo();
-
 	const ULandscapeEditLayerBase* EditLayer = LandscapeInfo->LandscapeActor->GetEditLayer(InEditLayerIndex);
 	if (EditLayer == nullptr)
 	{
@@ -8553,17 +8552,17 @@ bool ALandscapeProxy::LandscapeImportWeightmapFromRenderTarget(UTextureRenderTar
 	{
 		ULandscapeInfo* LandscapeInfo = Landscape->GetLandscapeInfo();
 
-		const ULandscapeEditLayerBase* EditLayer = LandscapeInfo->LandscapeActor->GetEditLayer(InEditLayerIndex);
-		if (EditLayer == nullptr)
-		{
-			FMessageLog("Blueprint").Error(LOCTEXT("LandscapeImportRenderTarget_InvalidEditLayer", "LandscapeImportWeightmapFromRenderTarget: Invalid edit layer index for import."));
-			return false;
-		}
-		const FGuid& EditLayerGuid = LandscapeInfo->LandscapeActor->GetEditLayer(InEditLayerIndex)->GetGuid();
-
 		int32 MinX, MinY, MaxX, MaxY;
 		if (LandscapeInfo && LandscapeInfo->GetLandscapeExtent(MinX, MinY, MaxX, MaxY))
 		{
+			const ULandscapeEditLayerBase* EditLayer = LandscapeInfo->LandscapeActor->GetEditLayer(InEditLayerIndex);
+			if (EditLayer == nullptr)
+			{
+				FMessageLog("Blueprint").Error(LOCTEXT("LandscapeImportRenderTarget_InvalidEditLayer", "LandscapeImportWeightmapFromRenderTarget: Invalid edit layer index for import."));
+				return false;
+			}
+			const FGuid& EditLayerGuid = LandscapeInfo->LandscapeActor->GetEditLayer(InEditLayerIndex)->GetGuid();
+
 			const uint32 LandscapeWidth = (uint32)(1 + MaxX - MinX);
 			const uint32 LandscapeHeight = (uint32)(1 + MaxY - MinY);
 			const FLinearColor SampleRect = FLinearColor(0.0f, 0.0f, static_cast<float>(LandscapeWidth), static_cast<float>(LandscapeHeight));
