Landscape Editor - Non-edit layer World Partition auto-conversion fix for content with mismatched bHasLayersContent values
* Some WP ALandscape actors have valid edit layer data but bHasLayersContent is set to false
* Implementing BugHawk feedback for ULandscapeInfo nullptr

diff --git a/Engine/Source/Runtime/Landscape/Private/Landscape.cpp b/Engine/Source/Runtime/Landscape/Private/Landscape.cpp
--- a/Engine/Source/Runtime/Landscape/Private/Landscape.cpp
+++ b/Engine/Source/Runtime/Landscape/Private/Landscape.cpp
@@ -5948,10 +5948,17 @@ void ULandscapeInfo::RegisterLandscapeActorWithProxyInternal(ALandscapeProxy* Pr
 				* During an auto conversion the proxy edit layer data is flattened and moved to the default edit layer
 			*/
 
-			const bool bProxyHasNoLayersContentPostLoad = Proxy->HadLayersContentAtPostLoadState == ELandscapeEditLayerContentState::PostLoad_NoEditLayers;
-			const bool bProxyHasBeenAutoConverted = Proxy->HadLayersContentAtPostLoadState == ELandscapeEditLayerContentState::AutoConvertedToEditLayers;
 			const bool bIsLandscapeActor = Proxy == LandscapeActor;
 
+			// Some ALandscape actors in a WP setup may be serialized with bHasLayersContent == false but have existing edit layer data
+			if (bIsLandscapeActor && !LandscapeActor->GetEditLayersConst().IsEmpty())
+			{
+				Proxy->HadLayersContentAtPostLoadState = ELandscapeEditLayerContentState::AutoConvertedToEditLayers;
+			}
+
+			const bool bProxyHasNoLayersContentPostLoad = Proxy->HadLayersContentAtPostLoadState == ELandscapeEditLayerContentState::PostLoad_NoEditLayers;
+			const bool bProxyHasBeenAutoConverted = Proxy->HadLayersContentAtPostLoadState == ELandscapeEditLayerContentState::AutoConvertedToEditLayers;
+			
 			// If the LandscapeActor did not support edit layers we need to run a full auto conversion
 			if (bIsLandscapeActor && bProxyHasNoLayersContentPostLoad && !bProxyHasBeenAutoConverted)
 			{
diff --git a/Engine/Source/Runtime/Landscape/Private/LandscapeEdit.cpp b/Engine/Source/Runtime/Landscape/Private/LandscapeEdit.cpp
--- a/Engine/Source/Runtime/Landscape/Private/LandscapeEdit.cpp
+++ b/Engine/Source/Runtime/Landscape/Private/LandscapeEdit.cpp
@@ -8279,7 +8279,7 @@ bool ALandscapeProxy::LandscapeImportHeightmapFromRenderTarget(UTextureRenderTar
 	uint64 StartCycle = FPlatformTime::Cycles64();
 
 	ALandscape* Landscape = GetLandscapeActor();
-	ULandscapeInfo* LandscapeInfo = Landscape->GetLandscapeInfo();
+	ULandscapeInfo* LandscapeInfo = Landscape ? Landscape->GetLandscapeInfo() : nullptr;
 	if (Landscape == nullptr || LandscapeInfo == nullptr)
 	{
 		FMessageLog("Blueprint").Error(LOCTEXT("LandscapeImportHeightmapFromRenderTarget_NullLandscape", "LandscapeImportHeightmapFromRenderTarget: Landscape must be non-null."));
