Landscape Editor - Adjusting FixupSharedData to run edit layer validation and fixup on levels with old versions
* Edit layer auto-conversion failed to add default edit layer for landscapes that entered the UpgradeSharedProperties path

diff --git a/Engine/Source/Runtime/Landscape/Private/Landscape.cpp b/Engine/Source/Runtime/Landscape/Private/Landscape.cpp
--- a/Engine/Source/Runtime/Landscape/Private/Landscape.cpp
+++ b/Engine/Source/Runtime/Landscape/Private/Landscape.cpp
@@ -5106,6 +5106,12 @@ void ALandscapeProxy::FixupSharedData(ALandscape* Landscape, const bool bMapChec
 	{
 		return;
 	}
+	
+	ULandscapeInfo* LandscapeInfo = GetLandscapeInfo();
+	checkf(LandscapeInfo != nullptr, TEXT("FixupSharedData can only be called after the proxies are registered to ULandscapeInfo"));
+
+	bool bUpdated = false;
+	TArray<FName> SynchronizedProperties;
 
 	const bool bUpgradeSharedPropertiesPerformedBefore = bUpgradeSharedPropertiesPerformed;
 	if (!bUpgradeSharedPropertiesPerformed && 
@@ -5117,33 +5123,30 @@ void ALandscapeProxy::FixupSharedData(ALandscape* Landscape, const bool bMapChec
 	}
 	else
 	{
-		ULandscapeInfo* LandscapeInfo = GetLandscapeInfo();
-		checkf(LandscapeInfo != nullptr, TEXT("FixupSharedData can only be called after the proxies are registered to ULandscapeInfo"));
-
-		TArray<FName> SynchronizedProperties = SynchronizeSharedProperties(Landscape);
-		bool bUpdated = !SynchronizedProperties.IsEmpty();
+		SynchronizedProperties = SynchronizeSharedProperties(Landscape);
+		bUpdated |= !SynchronizedProperties.IsEmpty();
+	}
 
-		TSet<FGuid> LayerGuids;
-		Algo::Transform(Landscape->GetEditLayersConst(), LayerGuids, [](const ULandscapeEditLayerBase* EditLayer) { return EditLayer->GetGuid(); });
-		bUpdated |= RemoveObsoleteLayers(LayerGuids);
+	TSet<FGuid> LayerGuids;
+	Algo::Transform(Landscape->GetEditLayersConst(), LayerGuids, [](const ULandscapeEditLayerBase* EditLayer) { return EditLayer->GetGuid(); });
+	bUpdated |= RemoveObsoleteLayers(LayerGuids);
 
-		for (const ULandscapeEditLayerBase* EditLayer : Landscape->GetEditLayersConst())
-		{
-			bUpdated |= AddLayer(EditLayer->GetGuid());
-		}
+	for (const ULandscapeEditLayerBase* EditLayer : Landscape->GetEditLayersConst())
+	{
+		bUpdated |= AddLayer(EditLayer->GetGuid());
+	}
 
-		if (bUpdated)
-		{
-			// In cases where LandscapeInfo is not fully ready yet, we forward the provided ALandscape. If ALandscape is available in LandscapeInfo, we let the object function naturally.
-			const ALandscape* LandscapeActor = (LandscapeInfo->LandscapeActor == nullptr) ? Landscape : nullptr;
+	if (bUpdated)
+	{
+		// In cases where LandscapeInfo is not fully ready yet, we forward the provided ALandscape. If ALandscape is available in LandscapeInfo, we let the object function naturally.
+		const ALandscape* LandscapeActor = (LandscapeInfo->LandscapeActor == nullptr) ? Landscape : nullptr;
 
-			// Force resave the proxy through the modified landscape system, so that the user can then use the Build > Save Modified Landscapes (or Build > Build Landscape) button and therefore manually trigger the re-save of all modified proxies. * /
-			bool bNeedsManualResave = LandscapeInfo->MarkObjectDirty(/*InObject = */this, /*bInForceResave = */true, LandscapeActor);
+		// Force resave the proxy through the modified landscape system, so that the user can then use the Build > Save Modified Landscapes (or Build > Build Landscape) button and therefore manually trigger the re-save of all modified proxies. * /
+		bool bNeedsManualResave = LandscapeInfo->MarkObjectDirty(/*InObject = */this, /*bInForceResave = */true, LandscapeActor);
 
-			if (bMapCheck && bNeedsManualResave && !SynchronizedProperties.IsEmpty())
-			{
-				UE::Landscape::Private::DisplaySynchronizedPropertiesMapcheckWarning(SynchronizedProperties, /*InSynchronizedProxy = */*this, *Landscape);
-			}
+		if (bMapCheck && bNeedsManualResave && !SynchronizedProperties.IsEmpty())
+		{
+			UE::Landscape::Private::DisplaySynchronizedPropertiesMapcheckWarning(SynchronizedProperties, /*InSynchronizedProxy = */*this, *Landscape);
 		}
 	}
 
