[HWRT] Disable Landscape Vertex Mask by default until the proper support for path tracer is added.

#rnx

diff --git a/Engine/Source/Runtime/Landscape/Private/LandscapeRender.cpp b/Engine/Source/Runtime/Landscape/Private/LandscapeRender.cpp
--- a/Engine/Source/Runtime/Landscape/Private/LandscapeRender.cpp
+++ b/Engine/Source/Runtime/Landscape/Private/LandscapeRender.cpp
@@ -335,8 +335,8 @@ static FAutoConsoleVariableRef CVarLandscapeRayTracingVertexLODBias(
 
 static TAutoConsoleVariable<int32> CVarRayTracingLandscapeVertexMask(
 	TEXT("r.RayTracing.Geometry.Landscape.VertexMask"),
-	1,
-	TEXT("If landscape should use vertex mask when building ray tracing dynamic geometry"),
+	0,
+	TEXT("If landscape should use vertex mask when building ray tracing dynamic geometry (works only with inline raytracing)"),
 	FConsoleVariableDelegate::CreateStatic(&OnCVarNeedingRenderStateInvalidationChanged),
 	ECVF_RenderThreadSafe
 );
