Landscape Editor - Fix edit layer registration assert to allow Resize tool to register landscapes with no edit layers before migrating data

diff --git a/Engine/Source/Runtime/Landscape/Private/Landscape.cpp b/Engine/Source/Runtime/Landscape/Private/Landscape.cpp
--- a/Engine/Source/Runtime/Landscape/Private/Landscape.cpp
+++ b/Engine/Source/Runtime/Landscape/Private/Landscape.cpp
@@ -5963,9 +5963,9 @@ void ULandscapeInfo::RegisterLandscapeActorWithProxyInternal(ALandscapeProxy* Pr
 			const bool bProxyHasBeenAutoConverted = Proxy->HadLayersContentAtPostLoadState == ELandscapeEditLayerContentState::AutoConvertedToEditLayers;
 			
 			// If the LandscapeActor did not support edit layers we need to run a full auto conversion
-			if (bIsLandscapeActor && bProxyHasNoLayersContentPostLoad && !bProxyHasBeenAutoConverted)
+			if (bIsLandscapeActor)
 			{
-				LandscapeActor->SetIsPendingEditLayerConversion(bProxyHasNoLayersContentPostLoad);
+				LandscapeActor->SetIsPendingEditLayerConversion(bProxyHasNoLayersContentPostLoad && !bProxyHasBeenAutoConverted);
 			}
 
 			// Proxies are registered many times (on Undo), but conversion can only happen once on initial load
@@ -5978,6 +5978,7 @@ void ULandscapeInfo::RegisterLandscapeActorWithProxyInternal(ALandscapeProxy* Pr
 					// Create a default edit layer and copy the ALandscape data
 					// Streaming proxy data gets copied as proxies complete registration
 					LandscapeActor->ConvertNonEditLayerLandscape();
+					check(!LandscapeActor->GetEditLayers().IsEmpty());
 				}
 				else
 				{
@@ -5996,9 +5997,6 @@ void ULandscapeInfo::RegisterLandscapeActorWithProxyInternal(ALandscapeProxy* Pr
 				Proxy->HadLayersContentAtPostLoadState = ELandscapeEditLayerContentState::AutoConvertedToEditLayers;
 			}
 		}
-
-		// Landscape needs to have edit layers by end of registration
-		check((Proxy != LandscapeActor) || !LandscapeActor->GetEditLayers().IsEmpty());
 	}
 #endif // WITH_EDITOR
 }
