Landscape : prevent PrepareTextureResourcesLimited from being skipped when calling ForceLayersFullUpdate multiple times in the same frame. In the context of creating a new WP landscape piecewise (region by region), it would fail to fully stream in the target heightmaps/weightmaps before running an edit layers merge on it.
Also added validations that said textures are fully streamed in and compiled by the time we're trying to resolve them, so that we catch such problems on the game thread

diff --git a/Engine/Source/Runtime/Landscape/Private/LandscapeEditLayers.cpp b/Engine/Source/Runtime/Landscape/Private/LandscapeEditLayers.cpp
--- a/Engine/Source/Runtime/Landscape/Private/LandscapeEditLayers.cpp
+++ b/Engine/Source/Runtime/Landscape/Private/LandscapeEditLayers.cpp
@@ -4169,8 +4169,10 @@ void ALandscape::PrintLayersDebugTextureResource(const FString& InContext, FText
 
 bool ALandscape::PrepareTextureResourcesLimited(bool bInWaitForStreaming)
 {
-	// Rate-limit to once per frame.  Note that commandlets don't advance the frame counter.
-	if (GFrameNumber != LastPrepareTextureResourcesCalled || IsRunningCommandlet())
+	// Rate-limit to once per frame. Note that commandlets don't advance the frame counter.
+	//  Also, force it when bInWaitForStreaming is true, so that we don't prevent new textures from being prepared in case PrepareTextureResourcesLimited is called 
+	//  multiple times in the same frame with some new textures being created in-between (e.g. when creating a large WP landscape via multiple successive regions)
+	if (bInWaitForStreaming || (GFrameNumber != LastPrepareTextureResourcesCalled) || IsRunningCommandlet())
 	{
 		LastPrepareTextureResourcesCalled = GFrameNumber;
 		bLastPrepareTextureResourcesResult = PrepareTextureResources(bInWaitForStreaming);
@@ -7073,7 +7075,7 @@ UE::Landscape::EditLayers::FMergeRenderContext ALandscape::PrepareEditLayersMerg
 
 				// Only consider render items which intersect with the renderer's active state (e.g. one a given renderer, a render item (A) might affect heightmaps only and another one (B) weightmaps only, 
 				//  then if performing a heightmap merge, the RendererState's mask here will be ELandscapeToolTargetTypeFlags::Heightmap, so this allows to filter out item B: 
-				const FEditLayerTargetTypeState FilteredRenderItemTargetTypeState = RenderItem.GetTargetTypeState().Intersect(EditLayerRendererRenderInfo.RendererState.GetActiveTargetTypeState());
+				const FEditLayerTargetTypeState FilteredRenderItemTargetTypeState = RenderItem.GetTargetTypeState().Intersect(EditLayerRendererRenderInfo.RendererState.GetActiveTargetTypeState()); 
 
 				bool bOutputRenderItem = false;
 				const FOutputWorldArea& OutputWorldArea = RenderItem.GetOutputWorldArea();
@@ -7935,6 +7937,9 @@ int32 ALandscape::PerformLayersHeightmapsBatchedMerge(const FUpdateLayersContent
 				FComponentCopyInfo& ComponentCopyInfo = ComponentCopyInfos.Emplace_GetRef();
 				ComponentCopyInfo.Component = Component;
 				ComponentCopyInfo.Texture = Component->GetHeightmap(false);
+				// At this point, the heightmaps we're writing to should have been fully compiled and streamed in by PrepareLayersResources, assert that this is the case : 
+				check(!ComponentCopyInfo.Texture->IsDefaultTexture() && !ComponentCopyInfo.Texture->HasPendingInitOrStreaming() && ComponentCopyInfo.Texture->IsFullyStreamedIn());
+
 				// Effective area of the texture affecting this component (because of texture sharing) :
 				ComponentCopyInfo.TextureOffset = FIntPoint(static_cast<int32>(Component->HeightmapScaleBias.Z * ComponentCopyInfo.Texture->Source.GetSizeX()), static_cast<int32>(Component->HeightmapScaleBias.W * ComponentCopyInfo.Texture->Source.GetSizeY()));
 		
@@ -11033,6 +11038,8 @@ int32 ALandscape::PerformLayersWeightmapsBatchedMerge(FUpdateLayersContentContex
 	TexturesNeedingResolve.Reserve(InUpdateLayersContentContext.WeightmapsToResolve.Num());
 	for (UTexture2D* Weightmap : InUpdateLayersContentContext.WeightmapsToResolve)
 	{ 
+		// At this point, the weightmaps we're writing to should have been fully compiled and streamed in by PrepareLayersResources, assert that this is the case : 
+		check(!Weightmap->IsDefaultTexture() && !Weightmap->HasPendingInitOrStreaming() && Weightmap->IsFullyStreamedIn());
 		TexturesNeedingResolve.Add(Weightmap, 0);
 	}
 
@@ -12756,6 +12763,10 @@ void ALandscape::ForceLayersFullUpdate()
 	StreamingManagers.UpdateResourceStreaming(GetWorld()->GetDeltaSeconds(), /* bProcessEverything */ true);
 	StreamingManagers.BlockTillAllRequestsFinished();
 
+	// Reset the frame number for limiting PrepareTextureResourcesLimited in case ForceLayersFullUpdate is called 
+	//  multiple times in a single frame (for safety, mostly, because the frame number is bypassed when bWaitForStreaming is true)
+	LastPrepareTextureResourcesCalled = 0;
+
 	RequestSplineLayerUpdate();
 	RequestLayersContentUpdateForceAll();
 	ForceUpdateLayersContent(/* bIntermediateRender */ false);
