Fix deleted landscape spline points/segments coming back after duplicating the spline.

diff --git a/Engine/Source/Runtime/Landscape/Private/LandscapeSplines.cpp b/Engine/Source/Runtime/Landscape/Private/LandscapeSplines.cpp
--- a/Engine/Source/Runtime/Landscape/Private/LandscapeSplines.cpp
+++ b/Engine/Source/Runtime/Landscape/Private/LandscapeSplines.cpp
@@ -2570,6 +2570,11 @@ void ULandscapeSplineControlPoint::DeleteSplinePoints()
 	Points.Reset();
 	Bounds = FBox(ForceInit);
 
+	// Make this control point transient to prevent duplicate from picking the disconnected, but un-garbage-collected
+	// instance and exporting it.  Otherwise the non-standard code in PostEditImport will re-attach the point to the
+	// component and it will re-appear as part of the copy.
+	Rename(nullptr, GetTransientPackage());
+
 	OuterSplines->MarkRenderStateDirty();
 
 	if (LocalMeshComponent != nullptr)
@@ -3597,6 +3602,11 @@ void ULandscapeSplineSegment::DeleteSplinePoints()
 	Points.Reset();
 	Bounds = FBox(ForceInit);
 
+	// Make this segment transient to prevent duplicate from picking the disconnected, but un-garbage-collected
+	// instance and exporting it.  Otherwise the non-standard code in PostEditImport will re-attach the segment to the
+	// component and it will re-appear as part of the copy.
+	Rename(nullptr, GetTransientPackage());
+
 	OuterSplines->MarkRenderStateDirty();
 
 	// Destroy mesh components
