Fix grid-alignment warnings when creating a landscape with more than one region worth of components.

If the component's cached local box is not updated prior to registration, the landscape group will miscalculate the position of the component, resulting in misalignment warnings being logged.

diff --git a/Engine/Source/Runtime/Landscape/Classes/LandscapeComponent.h b/Engine/Source/Runtime/Landscape/Classes/LandscapeComponent.h
--- a/Engine/Source/Runtime/Landscape/Classes/LandscapeComponent.h
+++ b/Engine/Source/Runtime/Landscape/Classes/LandscapeComponent.h
@@ -1062,6 +1062,8 @@ public:
 	/** Returns the components's LandscapeHoleMaterial, or the Component's OverrideLandscapeHoleMaterial if set */
 	LANDSCAPE_API UMaterialInterface* GetLandscapeHoleMaterial() const;
 
+	FBox ComputeApproximateBounds();
+
 #if WITH_EDITOR
 	/**
 	 * Recalculate cached bounds using height values.  Returns true when the bounds were changed.
diff --git a/Engine/Source/Runtime/Landscape/Private/LandscapeEdit.cpp b/Engine/Source/Runtime/Landscape/Private/LandscapeEdit.cpp
--- a/Engine/Source/Runtime/Landscape/Private/LandscapeEdit.cpp
+++ b/Engine/Source/Runtime/Landscape/Private/LandscapeEdit.cpp
@@ -156,6 +156,16 @@ void ULandscapeComponent::Init(int32 InBaseX, int32 InBaseY, int32 InComponentSi
 	WeightmapSubsectionOffset = (float)(SubsectionSizeQuads + 1) / (float)ComponentVerts;
 		
 	UpdatedSharedPropertiesFromActor();
+
+	// Update CachedLocalBox with approximate bounds to avoid group misalignment warnings during registration
+	CachedLocalBox = ComputeApproximateBounds();
+}
+
+FBox ULandscapeComponent::ComputeApproximateBounds()
+{
+	FVector MinBox(0, 0, LandscapeDataAccess::GetLocalHeight(0));
+	FVector MaxBox(ComponentSizeQuads + 1, ComponentSizeQuads + 1, LandscapeDataAccess::GetLocalHeight(UINT16_MAX));
+	return FBox(MinBox, MaxBox);
 }
 
 #if WITH_EDITOR
@@ -334,9 +344,7 @@ bool ULandscapeComponent::UpdateCachedBounds(bool bInApproximateBounds)
 	NewLocalBox.Init();
 	if (bInApproximateBounds)
 	{
-		FVector MinBox(0, 0, LandscapeDataAccess::GetLocalHeight(0));
-		FVector MaxBox(ComponentSizeQuads + 1, ComponentSizeQuads + 1, LandscapeDataAccess::GetLocalHeight(UINT16_MAX));
-		NewLocalBox = FBox(MinBox, MaxBox);
+		NewLocalBox = ComputeApproximateBounds();
 		check(NewLocalBox.GetExtent().Z > 0.0);
 	}
 	else
