Landscape edit layers weight blending improvements :
* The 3 weight-blending methods are now :
** "None", previously known as "No Weight Blend"
** "Final Weight Blending", previously known simply as weight-blended, aka Legacy Weight Blending
** "Advanced Weight Blending" (aka Premultiplied Alpha Blending) : improved weight blending solution that is compatible with edit layers and is applied at every blend step of the merge algorithm, instead of at the very end, relying on the weight sum of the blend group of the current layer as the alpha value.
* Removed landscape layer info's sub-menu to create a Non-Weight Blended or Weight Blended : there's now a new landscape setting that defines which one is the default. The blend settings are now moved to the landscape layer info asset
* Advanced Weight Blending also has the concept of blend group, to be able to let certain layers be weight-blended with each other but remain additive against others
* Added "Sort by [REDACTED_PERSON]" option in the target layers list of the landscape paint panel

Texture patch improvements :
* Changed landscape texture patch's texture sampler to "Border" in order to prevent "bleeding" the edge values of the patch when they are rotated
* Hid useless categories for landscape patch component in the details panel
* Exposed patch component's layer Guid (get/set), landscape/falloff/falloff mode/blend mode (get),  to BP / public dll API
* Renamed "Internal Data" to "Debug Data"

Misc:
* Added CVar landscape.DebugShaders to selectively un-optimize all landscape (tool) shaders

diff --git a/Engine/Plugins/Experimental/LandscapePatch/Shaders/Private/LandscapeTexturePatchPS.usf b/Engine/Plugins/Experimental/LandscapePatch/Shaders/Private/LandscapeTexturePatchPS.usf
--- a/Engine/Plugins/Experimental/LandscapePatch/Shaders/Private/LandscapeTexturePatchPS.usf
+++ b/Engine/Plugins/Experimental/LandscapePatch/Shaders/Private/LandscapeTexturePatchPS.usf
@@ -78,6 +78,17 @@ float GetFalloffAlpha(float FalloffWorldMargin, float2 PatchWorldDimensions, flo
 	
 	return Alpha;
 }
+
+float4 SamplePatchTexture(Texture2D<float4> InPatchTexture, SamplerState InPatchTextureSampler, float2 InPatchUVCoordinates)
+{
+	float4 TextureSample = 0.0f;
+	// We could have used a AM_Border sampler with a border color of 0, but this is not supported on all plaforms : 
+	if (all(InPatchUVCoordinates >= 0.0f) && all(InPatchUVCoordinates <= 1.0f))
+	{
+		TextureSample = InPatchTexture.Sample(InPatchTextureSampler, InPatchUVCoordinates);
+	}
+	return TextureSample;
+}
 #endif // APPLY_HEIGHT_PATCH || APPLY_WEIGHT_PATCH
 
 #if APPLY_HEIGHT_PATCH
@@ -131,7 +142,7 @@ void ApplyLandscapeTextureHeightPatch(in float4 SVPos : SV_POSITION, out float4
 	float2 HeightmapToPatchTranslate = InHeightmapToPatch._m03_m13;
 	
 	float2 PatchUVCoordinates = mul(HeightmapToPatchRotateScale, SVPos.xy) + HeightmapToPatchTranslate;
-	float4 PatchSampledValue = InHeightPatch.Sample(InHeightPatchSampler, PatchUVCoordinates);
+	float4 PatchSampledValue = SamplePatchTexture(InHeightPatch, InHeightPatchSampler, PatchUVCoordinates);
 	float PatchStoredHeight = bInputIsPackedHeight ? UnpackHeight(PatchSampledValue.xy) : PatchSampledValue.x;
 	
 	float PatchSignedHeight = InHeightScale * (PatchStoredHeight - InZeroInEncoding) + InHeightOffset;
@@ -268,7 +279,7 @@ void ApplyLandscapeTextureWeightPatch(in float4 SVPos : SV_POSITION, out float4
 	float2 WeightmapToPatchTranslate = InWeightmapToPatch._m03_m13;
 	
 	float2 PatchUVCoordinates = mul(WeightmapToPatchRotateScale, SVPos.xy) + WeightmapToPatchTranslate;
-	float4 PatchSampledValue = InWeightPatch.Sample(InWeightPatchSampler, PatchUVCoordinates);
+	float4 PatchSampledValue = SamplePatchTexture(InWeightPatch, InWeightPatchSampler, PatchUVCoordinates);
 	float PatchWeight = PatchSampledValue.x;
 	
 #if PERFORM_BLENDING
diff --git a/Engine/Plugins/Experimental/LandscapePatch/Source/LandscapePatch/Private/LandscapeCircleHeightPatch.cpp b/Engine/Plugins/Experimental/LandscapePatch/Source/LandscapePatch/Private/LandscapeCircleHeightPatch.cpp
--- a/Engine/Plugins/Experimental/LandscapePatch/Source/LandscapePatch/Private/LandscapeCircleHeightPatch.cpp
+++ b/Engine/Plugins/Experimental/LandscapePatch/Source/LandscapePatch/Private/LandscapeCircleHeightPatch.cpp
@@ -7,6 +7,7 @@
 #include "Landscape.h"
 #include "LandscapeCircleHeightPatchPS.h"
 #include "LandscapeEditLayerMergeRenderContext.h"
+#include "LandscapeEditLayerTargetLayerGroup.h"
 #include "LandscapeEditResourcesSubsystem.h" // ULandscapeScratchRenderTarget
 #include "LandscapeInfo.h"
 #include "LandscapePatchUtil.h" // GetHeightmapToWorld
@@ -288,13 +289,15 @@ bool ULandscapeCircleHeightPatch::ApplyCirclePatch(bool bInPerformBlending, UE::
 }
 
 void ULandscapeCircleHeightPatch::GetRendererStateInfo(const UE::Landscape::EditLayers::FMergeContext* InMergeContext,
-	FEditLayerTargetTypeState& OutSupported,
-	FEditLayerTargetTypeState& OutEnabled,
-	TArray<TBitArray<>>& OutTargetLayerGroups) const
+	UE::Landscape::EditLayers::FEditLayerTargetTypeState& OutSupported,
+	UE::Landscape::EditLayers::FEditLayerTargetTypeState& OutEnabled,
+	TArray<UE::Landscape::EditLayers::FTargetLayerGroup>& OutTargetLayerGroups) const
 {
 	OutSupported.AddTargetType(bEditVisibility ? ELandscapeToolTargetType::Visibility
 		: ELandscapeToolTargetType::Heightmap);
 
+	// No support for weightmaps -> no need for target layer groups
+
 	if (IsEnabled())
 	{
 		OutEnabled = OutSupported;
diff --git a/Engine/Plugins/Experimental/LandscapePatch/Source/LandscapePatch/Private/LandscapeCircleHeightPatchPS.cpp b/Engine/Plugins/Experimental/LandscapePatch/Source/LandscapePatch/Private/LandscapeCircleHeightPatchPS.cpp
--- a/Engine/Plugins/Experimental/LandscapePatch/Source/LandscapePatch/Private/LandscapeCircleHeightPatchPS.cpp
+++ b/Engine/Plugins/Experimental/LandscapePatch/Source/LandscapePatch/Private/LandscapeCircleHeightPatchPS.cpp
@@ -3,6 +3,7 @@
 #include "LandscapeCircleHeightPatchPS.h"
 
 #include "PixelShaderUtils.h"
+#include "LandscapeUtils.h"
 #include "DataDrivenShaderPlatformInfo.h"
 
 bool FLandscapeCircleHeightPatchPSBase::ShouldCompilePermutation(const FGlobalShaderPermutationParameters& Parameters)
@@ -16,6 +17,7 @@ bool FLandscapeCircleHeightPatchPSBase::ShouldCompilePermutation(const FGlobalSh
 template <bool bPerformBlending>
 void FLandscapeCircleHeightPatchPS<bPerformBlending>::ModifyCompilationEnvironment(const FGlobalShaderPermutationParameters& Parameters, FShaderCompilerEnvironment& OutEnvironment)
 {
+	UE::Landscape::ModifyShaderCompilerEnvironmentForDebug(OutEnvironment);
 	if (bPerformBlending)
 	{
 		OutEnvironment.SetDefine(TEXT("PERFORM_BLENDING"), 1);
@@ -41,6 +43,7 @@ void FLandscapeCircleHeightPatchPS<bPerformBlending>::AddToRenderGraph(FRDGEvent
 template <bool bPerformBlending>
 void FLandscapeCircleVisibilityPatchPS<bPerformBlending>::ModifyCompilationEnvironment(const FGlobalShaderPermutationParameters& Parameters, FShaderCompilerEnvironment& OutEnvironment)
 {
+	UE::Landscape::ModifyShaderCompilerEnvironmentForDebug(OutEnvironment);
 	if (bPerformBlending)
 	{
 		OutEnvironment.SetDefine(TEXT("PERFORM_BLENDING"), 1);
diff --git a/Engine/Plugins/Experimental/LandscapePatch/Source/LandscapePatch/Private/LandscapePatchManager.cpp b/Engine/Plugins/Experimental/LandscapePatch/Source/LandscapePatch/Private/LandscapePatchManager.cpp
--- a/Engine/Plugins/Experimental/LandscapePatch/Source/LandscapePatch/Private/LandscapePatchManager.cpp
+++ b/Engine/Plugins/Experimental/LandscapePatch/Source/LandscapePatch/Private/LandscapePatchManager.cpp
@@ -269,7 +269,7 @@ TArray<UE::Landscape::EditLayers::FEditLayerRendererState> ALandscapePatchManage
 	// Inherit the base brush's enabled state and intersect it with all patches' state, so that the heightmap/weightmap/visibility toggle at the brush's level applies to them too : 
 	FEditLayerTargetTypeState BaseBrushSupportedState(InMergeContext);
 	FEditLayerTargetTypeState BaseBrushEnabledState(InMergeContext);
-	TArray<TBitArray<>> BaseTargetLayerGroups;
+	TArray<FTargetLayerGroup> BaseTargetLayerGroups;
 	Super::GetRendererStateInfo(InMergeContext, BaseBrushSupportedState, BaseBrushEnabledState, BaseTargetLayerGroups);
 
 	TArray<FEditLayerRendererState> RendererStates;
diff --git a/Engine/Plugins/Experimental/LandscapePatch/Source/LandscapePatch/Private/LandscapeTexturePatch.cpp b/Engine/Plugins/Experimental/LandscapePatch/Source/LandscapePatch/Private/LandscapeTexturePatch.cpp
--- a/Engine/Plugins/Experimental/LandscapePatch/Source/LandscapePatch/Private/LandscapeTexturePatch.cpp
+++ b/Engine/Plugins/Experimental/LandscapePatch/Source/LandscapePatch/Private/LandscapeTexturePatch.cpp
@@ -7,6 +7,7 @@
 #include "Landscape.h"
 #include "LandscapeDataAccess.h"
 #include "LandscapeEditLayerMergeRenderContext.h"
+#include "LandscapeEditLayerTargetTypeState.h"
 #include "LandscapeEditResourcesSubsystem.h" // ULandscapeScratchRenderTarget
 #include "LandscapePatchLogging.h"
 #include "LandscapePatchManager.h"
@@ -1464,7 +1465,7 @@ void ULandscapeTexturePatch::SnapToLandscape()
 	double PatchExtentY = PatchScale.Y * UnscaledPatchCoverage.Y;
 	if (Num90DegreeRotations % 2)
 	{
-		// Relative to the landscape, our lenght and width are backwards...
+		// Relative to the landscape, our length and width are backwards...
 		Swap(PatchExtentX, PatchExtentY);
 	}
 
@@ -2206,8 +2207,8 @@ void ULandscapeTexturePatch::SetEditVisibilityLayer(const FName& InWeightmapLaye
 
 #if WITH_EDITOR
 void ULandscapeTexturePatch::GetRendererStateInfo(const UE::Landscape::EditLayers::FMergeContext* InMergeContext,
-	FEditLayerTargetTypeState& OutSupported, FEditLayerTargetTypeState& OutEnabled,
-	TArray<TBitArray<>>& OutTargetLayerGroups) const
+	UE::Landscape::EditLayers::FEditLayerTargetTypeState& OutSupported, UE::Landscape::EditLayers::FEditLayerTargetTypeState& OutEnabled,
+	TArray<UE::Landscape::EditLayers::FTargetLayerGroup>& OutTargetLayerGroups) const
 {
 	if (InMergeContext->IsHeightmapMerge())
 	{
@@ -2234,6 +2235,9 @@ void ULandscapeTexturePatch::GetRendererStateInfo(const UE::Landscape::EditLayer
 				}
 			}
 		}
+
+		// Build the target layer groups for this step. Since we use the generic blend, we can use BuildGenericBlendTargetLayerGroups
+		OutTargetLayerGroups = InMergeContext->BuildGenericBlendTargetLayerGroups(OutSupported.GetActiveWeightmapBitIndices());
 	}
 
 	if (IsEnabled())
diff --git a/Engine/Plugins/Experimental/LandscapePatch/Source/LandscapePatch/Private/LandscapeTexturePatchPS.cpp b/Engine/Plugins/Experimental/LandscapePatch/Source/LandscapePatch/Private/LandscapeTexturePatchPS.cpp
--- a/Engine/Plugins/Experimental/LandscapePatch/Source/LandscapePatch/Private/LandscapeTexturePatchPS.cpp
+++ b/Engine/Plugins/Experimental/LandscapePatch/Source/LandscapePatch/Private/LandscapeTexturePatchPS.cpp
@@ -16,6 +16,8 @@ bool FApplyLandscapeTextureHeightPatchPSBase::ShouldCompilePermutation(const FGl
 
 void FApplyLandscapeTextureHeightPatchPSBase::ModifyCompilationEnvironment(const FGlobalShaderPermutationParameters& Parameters, FShaderCompilerEnvironment& OutEnvironment)
 {
+	UE::Landscape::ModifyShaderCompilerEnvironmentForDebug(OutEnvironment);
+
 	OutEnvironment.SetDefine(TEXT("APPLY_HEIGHT_PATCH"), 1);
 
 	// Make our flag choices match in the shader.
@@ -36,6 +38,8 @@ bool FOffsetHeightmapPS::ShouldCompilePermutation(const FGlobalShaderPermutation
 
 void FOffsetHeightmapPS::ModifyCompilationEnvironment(const FGlobalShaderPermutationParameters& Parameters, FShaderCompilerEnvironment& OutEnvironment)
 {
+	UE::Landscape::ModifyShaderCompilerEnvironmentForDebug(OutEnvironment);
+
 	OutEnvironment.SetDefine(TEXT("OFFSET_HEIGHT_PATCH"), 1);
 }
 
@@ -62,6 +66,8 @@ bool FSimpleTextureCopyPS::ShouldCompilePermutation(const FGlobalShaderPermutati
 
 void FSimpleTextureCopyPS::ModifyCompilationEnvironment(const FGlobalShaderPermutationParameters& Parameters, FShaderCompilerEnvironment& OutEnvironment)
 {
+	UE::Landscape::ModifyShaderCompilerEnvironmentForDebug(OutEnvironment);
+
 	OutEnvironment.SetDefine(TEXT("SIMPLE_TEXTURE_COPY"), 1);
 }
 
@@ -95,6 +101,8 @@ bool FConvertToNativeLandscapePatchPS::ShouldCompilePermutation(const FGlobalSha
 
 void FConvertToNativeLandscapePatchPS::ModifyCompilationEnvironment(const FGlobalShaderPermutationParameters& Parameters, FShaderCompilerEnvironment& OutEnvironment)
 {
+	UE::Landscape::ModifyShaderCompilerEnvironmentForDebug(OutEnvironment);
+
 	OutEnvironment.SetDefine(TEXT("CONVERT_TO_NATIVE_LANDSCAPE_PATCH"), 1);
 }
 
@@ -130,6 +138,8 @@ bool FConvertBackFromNativeLandscapePatchPS::ShouldCompilePermutation(const FGlo
 
 void FConvertBackFromNativeLandscapePatchPS::ModifyCompilationEnvironment(const FGlobalShaderPermutationParameters& Parameters, FShaderCompilerEnvironment& OutEnvironment)
 {
+	UE::Landscape::ModifyShaderCompilerEnvironmentForDebug(OutEnvironment);
+
 	OutEnvironment.SetDefine(TEXT("CONVERT_BACK_FROM_NATIVE_LANDSCAPE_PATCH"), 1);
 }
 
@@ -165,6 +175,8 @@ bool FApplyLandscapeTextureWeightPatchPSBase::ShouldCompilePermutation(const FGl
 
 void FApplyLandscapeTextureWeightPatchPSBase::ModifyCompilationEnvironment(const FGlobalShaderPermutationParameters& Parameters, FShaderCompilerEnvironment& OutEnvironment)
 {
+	UE::Landscape::ModifyShaderCompilerEnvironmentForDebug(OutEnvironment);
+
 	OutEnvironment.SetDefine(TEXT("APPLY_WEIGHT_PATCH"), 1);
 
 	// Make our flag choices match in the shader.
@@ -184,6 +196,8 @@ bool FReinitializeLandscapePatchPS::ShouldCompilePermutation(const FGlobalShader
 
 void FReinitializeLandscapePatchPS::ModifyCompilationEnvironment(const FGlobalShaderPermutationParameters& Parameters, FShaderCompilerEnvironment& OutEnvironment)
 {
+	UE::Landscape::ModifyShaderCompilerEnvironmentForDebug(OutEnvironment);
+
 	OutEnvironment.SetDefine(TEXT("REINITIALIZE_PATCH"), 1);
 }
 
diff --git a/Engine/Plugins/Experimental/LandscapePatch/Source/LandscapePatch/Public/LandscapeCircleHeightPatch.h b/Engine/Plugins/Experimental/LandscapePatch/Source/LandscapePatch/Public/LandscapeCircleHeightPatch.h
--- a/Engine/Plugins/Experimental/LandscapePatch/Source/LandscapePatch/Public/LandscapeCircleHeightPatch.h
+++ b/Engine/Plugins/Experimental/LandscapePatch/Source/LandscapePatch/Public/LandscapeCircleHeightPatch.h
@@ -31,16 +31,13 @@ public:
 #if WITH_EDITOR
 	virtual UTextureRenderTarget2D* RenderLayer_Native(const FLandscapeBrushParameters& InParameters, const FTransform& HeightmapToWorld) override;
 
-	using FEditLayerTargetTypeState = UE::Landscape::EditLayers::FEditLayerTargetTypeState;
-	using FEditLayerRenderItem = UE::Landscape::EditLayers::FEditLayerRenderItem;
-
 	// ILandscapeEditLayerRenderer, via ULandscapePatchComponent
 	virtual void GetRendererStateInfo(const UE::Landscape::EditLayers::FMergeContext* InMergeContext,
-		FEditLayerTargetTypeState& OutSupportedTargetTypeState, 
-		FEditLayerTargetTypeState& OutEnabledTargetTypeState, 
-		TArray<TBitArray<>>& OutTargetLayerGroups) const override;
+		UE::Landscape::EditLayers::FEditLayerTargetTypeState& OutSupportedTargetTypeState, 
+		UE::Landscape::EditLayers::FEditLayerTargetTypeState& OutEnabledTargetTypeState, 
+		TArray<UE::Landscape::EditLayers::FTargetLayerGroup>& OutTargetLayerGroups) const override;
 	virtual FString GetEditLayerRendererDebugName() const override;
-	virtual TArray<FEditLayerRenderItem> GetRenderItems(const UE::Landscape::EditLayers::FMergeContext* InMergeContext) const override;
+	virtual TArray<UE::Landscape::EditLayers::FEditLayerRenderItem> GetRenderItems(const UE::Landscape::EditLayers::FMergeContext* InMergeContext) const override;
 	virtual UE::Landscape::EditLayers::ERenderFlags GetRenderFlags(const UE::Landscape::EditLayers::FMergeContext* InMergeContext) const override;
 	virtual bool CanGroupRenderLayerWith(TScriptInterface<ILandscapeEditLayerRenderer> InOtherRenderer) const override;
 	virtual bool RenderLayer(UE::Landscape::EditLayers::FRenderParams& RenderParams, UE::Landscape::FRDGBuilderRecorder& RDGBuilderRecorder) override;
diff --git a/Engine/Plugins/Experimental/LandscapePatch/Source/LandscapePatch/Public/LandscapePatchComponent.h b/Engine/Plugins/Experimental/LandscapePatch/Source/LandscapePatch/Public/LandscapePatchComponent.h
--- a/Engine/Plugins/Experimental/LandscapePatch/Source/LandscapePatch/Public/LandscapePatchComponent.h
+++ b/Engine/Plugins/Experimental/LandscapePatch/Source/LandscapePatch/Public/LandscapePatchComponent.h
@@ -43,7 +43,7 @@ enum class ELandscapePatchPriorityInitialization : uint8
  */
 //~ TODO: Although this doesn't generate geometry, we are likely to change this to inherit from UPrimitiveComponent
 //~ so that we can use render proxies for passing along data to the render thread or perhaps for visualization.
-UCLASS(MinimalAPI, Blueprintable, BlueprintType, Abstract)
+UCLASS(MinimalAPI, Blueprintable, BlueprintType, Abstract, HideCategories = (Tags, Activation, Cooking, Physics, Rendering, Navigation, LOD))
 class ULandscapePatchComponent : public USceneComponent
 	, public ILandscapeEditLayerRenderer
 {
@@ -101,8 +101,11 @@ public:
 	UFUNCTION(BlueprintCallable, Category = "LandscapePatch")
 	LANDSCAPEPATCH_API void SetPriority(double PriorityIn);
 
-	FGuid GetEditLayerGuid() const { return EditLayerGuid; }
-	void SetEditLayerGuid(const FGuid& GuidIn);
+	UFUNCTION(BlueprintCallable, Category = "LandscapePatch")
+	LANDSCAPEPATCH_API FGuid GetEditLayerGuid() const { return EditLayerGuid; }
+
+	UFUNCTION(BlueprintCallable, Category = "LandscapePatch")
+	LANDSCAPEPATCH_API void SetEditLayerGuid(const FGuid& GuidIn);
 
 	/**
 	 * Returns the result of GetTypeHash(GetFullName()) for this component. Useful for quick deterministic
@@ -113,6 +116,9 @@ public:
 	UFUNCTION(BlueprintCallable, Category = "LandscapePatch")
 	LANDSCAPEPATCH_API FTransform GetLandscapeHeightmapCoordsToWorld() const;
 
+	UFUNCTION(BlueprintCallable, Category = "LandscapePatch")
+	LANDSCAPEPATCH_API ALandscape* GetLandscape() const { return Landscape.Get(); }
+
 	UFUNCTION(BlueprintCallable, Category = "LandscapePatch")
 	LANDSCAPEPATCH_API virtual void SetLandscape(ALandscape* NewLandscape);
 
diff --git a/Engine/Plugins/Experimental/LandscapePatch/Source/LandscapePatch/Public/LandscapePatchManager.h b/Engine/Plugins/Experimental/LandscapePatch/Source/LandscapePatch/Public/LandscapePatchManager.h
--- a/Engine/Plugins/Experimental/LandscapePatch/Source/LandscapePatch/Public/LandscapePatchManager.h
+++ b/Engine/Plugins/Experimental/LandscapePatch/Source/LandscapePatch/Public/LandscapePatchManager.h
@@ -33,11 +33,9 @@ public:
 	virtual UTextureRenderTarget2D* RenderLayer_Native(const FLandscapeBrushParameters& InParameters) override;
 
 #if WITH_EDITOR
-	using FEditLayerRendererState = UE::Landscape::EditLayers::FEditLayerRendererState;
-
 	// IEditLayerRendererProvider
 	// Called by the batched-merge application path.
-	virtual TArray<FEditLayerRendererState> GetEditLayerRendererStates(const UE::Landscape::EditLayers::FMergeContext* InMergeContext) override;
+	virtual TArray<UE::Landscape::EditLayers::FEditLayerRendererState> GetEditLayerRendererStates(const UE::Landscape::EditLayers::FMergeContext* InMergeContext) override;
 
 	// ILandscapeEditLayerRenderer
 	//~ In batched merge, the manager relies on being a renderer provider. It does not need to have its 
@@ -46,7 +44,7 @@ public:
 	virtual void GetRendererStateInfo(const UE::Landscape::EditLayers::FMergeContext* InMergeContext,
 		UE::Landscape::EditLayers::FEditLayerTargetTypeState& OutSupportedTargetTypeState,
 		UE::Landscape::EditLayers::FEditLayerTargetTypeState& OutEnabledTargetTypeState,
-		TArray<TBitArray<>>& OutTargetLayerGroups) const override {}
+		TArray<UE::Landscape::EditLayers::FTargetLayerGroup>& OutTargetLayerGroups) const override {}
 	virtual TArray<UE::Landscape::EditLayers::FEditLayerRenderItem> GetRenderItems(const UE::Landscape::EditLayers::FMergeContext* InMergeContext) const override { return {}; }
 	virtual bool RenderLayer(UE::Landscape::EditLayers::FRenderParams& RenderParams, UE::Landscape::FRDGBuilderRecorder& RDGBuilderRecorder) override { return false; };
 
diff --git a/Engine/Plugins/Experimental/LandscapePatch/Source/LandscapePatch/Public/LandscapeTextureBackedRenderTarget.h b/Engine/Plugins/Experimental/LandscapePatch/Source/LandscapePatch/Public/LandscapeTextureBackedRenderTarget.h
--- a/Engine/Plugins/Experimental/LandscapePatch/Source/LandscapePatch/Public/LandscapeTextureBackedRenderTarget.h
+++ b/Engine/Plugins/Experimental/LandscapePatch/Source/LandscapePatch/Public/LandscapeTextureBackedRenderTarget.h
@@ -60,16 +60,16 @@ protected:
 
 	void ReinitializeRenderTarget(bool bClear);
 
-	UPROPERTY(VisibleAnywhere, Category = InternalData)
+	UPROPERTY(VisibleAnywhere, Category = DebugData)
 	TObjectPtr<UTexture2D> InternalTexture = nullptr;
 
-	UPROPERTY(VisibleAnywhere, Category = InternalData, Transient, DuplicateTransient)
+	UPROPERTY(VisibleAnywhere, Category = DebugData, Transient, DuplicateTransient)
 	TObjectPtr<UTextureRenderTarget2D> RenderTarget = nullptr;
 
-	UPROPERTY(VisibleAnywhere, Category = InternalData)
+	UPROPERTY(VisibleAnywhere, Category = DebugData)
 	int32 SizeX = 32;
 
-	UPROPERTY(VisibleAnywhere, Category = InternalData)
+	UPROPERTY(VisibleAnywhere, Category = DebugData)
 	int32 SizeY = 32;
 
 	UPROPERTY()
diff --git a/Engine/Plugins/Experimental/LandscapePatch/Source/LandscapePatch/Public/LandscapeTexturePatch.h b/Engine/Plugins/Experimental/LandscapePatch/Source/LandscapePatch/Public/LandscapeTexturePatch.h
--- a/Engine/Plugins/Experimental/LandscapePatch/Source/LandscapePatch/Public/LandscapeTexturePatch.h
+++ b/Engine/Plugins/Experimental/LandscapePatch/Source/LandscapePatch/Public/LandscapeTexturePatch.h
@@ -168,6 +168,7 @@ public:
 #endif // WITH_EDITOR
 
 protected:
+	/** Name of the target layer to affect. */
 	UPROPERTY(EditAnywhere, Category = WeightPatch)
 	FName WeightmapLayerName;
 
@@ -182,28 +183,28 @@ protected:
 	TObjectPtr<UTexture> TextureAsset = nullptr;
 
 	/** Not directly settable via detail panel- for display/debugging purposes only. */
-	UPROPERTY(VisibleAnywhere, Category = WeightPatch, Instanced, AdvancedDisplay)
+	UPROPERTY(VisibleAnywhere, Category = WeightPatch, Instanced, AdvancedDisplay, meta = (DisplayName = "Debug Data"))
 	TObjectPtr<ULandscapeWeightTextureBackedRenderTarget> InternalData = nullptr;
 
 	UPROPERTY(EditAnywhere, Category = WeightPatch, meta = (EditConditionHides, EditCondition = "false"))
 	ELandscapeTexturePatchSourceMode SourceMode = ELandscapeTexturePatchSourceMode::None;
 
-	/**
-	 * How the heightmap of the patch is stored.
-	 */
+	/** How the weightmap of the patch is stored. */
 	UPROPERTY(EditAnywhere, Category = WeightPatch, meta = (DisplayName = "Source Mode"))
 	ELandscapeTexturePatchSourceMode DetailPanelSourceMode = ELandscapeTexturePatchSourceMode::None;
 
-	//~ We could refactor things such that we always have an InternalData pointer, even when we use
-	//~ a texture asset, and then we could use the boolean inside that instead (which needs to be there
-	//~ so that we know how many channels we need). Not clear whether that will be any cleaner though.
+	/**
+	 * When true, texture alpha channel will be used when applying the weight patch. Note that the source data needs to
+	 * have an alpha channel for this to have an effect.
+	 */
 	UPROPERTY(EditAnywhere, Category = WeightPatch)
 	bool bUseAlphaChannel = false;
 
-	// Can't make TOptional a UPROPERTY, hence these two.
-	UPROPERTY(EditAnywhere, Category = WeightPatch)
+	/** Allows to override the patch component's own blend node */
+	UPROPERTY(EditAnywhere, Category = WeightPatch, meta = (InlineEditConditionToggle))
 	bool bOverrideBlendMode = false;
-	UPROPERTY(EditAnywhere, Category = WeightPatch, meta = (EditConditionHides, EditCondition = "bOverrideBlendMode"))
+	/** Determines how the patch is combined with the previous state of the landscape (overrides the patch component's own blend node). */
+	UPROPERTY(EditAnywhere, Category = WeightPatch, meta = (EditCondition = "bOverrideBlendMode"))
 	ELandscapeTexturePatchBlendMode OverrideBlendMode = ELandscapeTexturePatchBlendMode::AlphaBlend;
 
 	// TODO: We could support having different per-layer falloff modes and falloff amounts as well, as
@@ -228,16 +229,13 @@ class ULandscapeTexturePatch : public ULandscapePatchComponent
 public:
 
 #if WITH_EDITOR
-	using FEditLayerTargetTypeState = UE::Landscape::EditLayers::FEditLayerTargetTypeState;
-	using FEditLayerRenderItem = UE::Landscape::EditLayers::FEditLayerRenderItem;
-
 	// ILandscapeEditLayerRenderer, via ULandscapePatchComponent
 	LANDSCAPEPATCH_API virtual void GetRendererStateInfo(const UE::Landscape::EditLayers::FMergeContext* InMergeContext,
-		FEditLayerTargetTypeState& OutSupportedTargetTypeState,
-		FEditLayerTargetTypeState& OutEnabledTargetTypeState,
-		TArray<TBitArray<>>& OutTargetLayerGroups) const override;
+		UE::Landscape::EditLayers::FEditLayerTargetTypeState& OutSupportedTargetTypeState,
+		UE::Landscape::EditLayers::FEditLayerTargetTypeState& OutEnabledTargetTypeState,
+		TArray<UE::Landscape::EditLayers::FTargetLayerGroup>& OutTargetLayerGroups) const override;
 	LANDSCAPEPATCH_API virtual FString GetEditLayerRendererDebugName() const override;
-	LANDSCAPEPATCH_API virtual TArray<FEditLayerRenderItem> GetRenderItems(const UE::Landscape::EditLayers::FMergeContext* InMergeContext) const override;
+	LANDSCAPEPATCH_API virtual TArray<UE::Landscape::EditLayers::FEditLayerRenderItem> GetRenderItems(const UE::Landscape::EditLayers::FMergeContext* InMergeContext) const override;
 	LANDSCAPEPATCH_API virtual UE::Landscape::EditLayers::ERenderFlags GetRenderFlags(const UE::Landscape::EditLayers::FMergeContext* InMergeContext) const override;
 	LANDSCAPEPATCH_API virtual bool CanGroupRenderLayerWith(TScriptInterface<ILandscapeEditLayerRenderer> InOtherRenderer) const override;
 	LANDSCAPEPATCH_API virtual bool RenderLayer(UE::Landscape::EditLayers::FRenderParams& RenderParams, UE::Landscape::FRDGBuilderRecorder& RDGBuilderRecorder) override;
@@ -318,6 +316,8 @@ public:
 	UFUNCTION(BlueprintCallable, Category = LandscapePatch, meta = (ResolutionMultiplier = "1.0"))
 	LANDSCAPEPATCH_API virtual UPARAM(DisplayName = "Success") bool GetInitResolutionFromLandscape(float ResolutionMultiplier, FVector2D& ResolutionOut) const;
 
+	UFUNCTION(BlueprintCallable, Category = LandscapePatch)
+	LANDSCAPEPATCH_API float GetFalloff() const { return Falloff; }
 
 	UFUNCTION(BlueprintCallable, Category = LandscapePatch)
 	void SetFalloff(float FalloffIn) 
@@ -329,6 +329,9 @@ public:
 		}
 	}
 
+	UFUNCTION(BlueprintCallable, Category = LandscapePatch)
+	LANDSCAPEPATCH_API ELandscapeTexturePatchFalloffMode GetFalloffMode() const { return FalloffMode; }
+
 	UFUNCTION(BlueprintCallable, Category = LandscapePatch)
 	void SetFalloffMode(ELandscapeTexturePatchFalloffMode FalloffModeIn) 
 	{
@@ -339,6 +342,9 @@ public:
 		}
 	}
 
+	UFUNCTION(BlueprintCallable, Category = LandscapePatch)
+	LANDSCAPEPATCH_API ELandscapeTexturePatchBlendMode GetBlendMode() const { return BlendMode; }
+
 	/**
 	 * Determines how the height patch is blended into the existing terrain.
 	 */
@@ -521,6 +527,7 @@ protected:
 	UPROPERTY(EditAnywhere, Category = Settings, meta = (UIMin = "0", ClampMin = "0"))
 	FVector2D UnscaledPatchCoverage = FVector2D(2000, 2000);
 
+	/** Determines how the patch is combined with the previous state of the landscape. */
 	UPROPERTY(EditAnywhere, Category = Settings)
 	ELandscapeTexturePatchBlendMode BlendMode = ELandscapeTexturePatchBlendMode::AlphaBlend;
 
@@ -551,7 +558,7 @@ protected:
 	ELandscapeTexturePatchSourceMode DetailPanelHeightSourceMode = ELandscapeTexturePatchSourceMode::None;
 
 	/** Not directly settable via detail panel- for display/debugging purposes only. */
-	UPROPERTY(VisibleAnywhere, Category = HeightPatch, Instanced)
+	UPROPERTY(VisibleAnywhere, Category = HeightPatch, Instanced, AdvancedDisplay, meta = (DisplayName = "Debug Data"))
 	TObjectPtr<ULandscapeHeightTextureBackedRenderTarget> HeightInternalData = nullptr;
 
 	/**
diff --git a/Engine/Shaders/Private/Landscape/LandscapeEditLayersWeightmaps.usf b/Engine/Shaders/Private/Landscape/LandscapeEditLayersWeightmaps.usf
--- a/Engine/Shaders/Private/Landscape/LandscapeEditLayersWeightmaps.usf
+++ b/Engine/Shaders/Private/Landscape/LandscapeEditLayersWeightmaps.usf
@@ -8,7 +8,7 @@
 #if defined (__INTELLISENSE__)
 // Uncomment the appropriate define for enabling syntax highlighting with HLSL Tools for Visual Studio : 
 //#define MERGE_EDIT_LAYER 1
-//#define PERFORM_LEGACY_WEIGHT_BLENDING 1
+//#define PERFORM_FINAL_WEIGHT_BLENDING 1
 //#define PACK_WEIGHTMAP 1
 //#define GENERATE_MIPS 1
 #endif // defined (__INTELLISENSE__)
@@ -24,36 +24,58 @@
 
 /** EWeightmapTargetLayerFlags enum */
 #define EWEIGHTMAPTARGETLAYERFLAGS_ISVISIBILITYLAYER (1 << 0)
+// EWEIGHTMAPTARGETLAYERFLAGS_ISWEIGHTBLENDED is used for "final weight blending", whose formula is : 
+//  * BlendedWeightsSum = weight sum of all final weight-blended target layers
+//  * LayerWeight = weight of the target layer being processed
+//  ==> FinalLayerWeight = LayerWeight / BlendedWeightsSum
 #define EWEIGHTMAPTARGETLAYERFLAGS_ISWEIGHTBLENDED (1 << 1)
 #define EWEIGHTMAPTARGETLAYERFLAGS_SKIP (1 << 2)
+// EWEIGHTMAPTARGETLAYERFLAGS_ISPREMULTIPLIEDALPHAWEIGHTBLENDED is used for "premultiplied alpha blending", whose formula is : 
+//  * PreviousLayersWeight = weight of the target layer being processed, at this step of the edit layers merge algorithm (i.e. merged weight of all previous edit layers in the stack)
+//  * CurrentLayerWeight = weight of the target layer being processed at the current edit layer 
+//  * PremultipliedAlphaBlendGroupWeightSum = weight sum of all premultiplied alpha-blended target layers of the same blend group at the current edit layer 
+//  * EditLayerAlpha = alpha value of the current edit layer
+// => FinalLayerWeight = PreviousLayersWeight * (1.0f - min(PremultipliedAlphaBlendGroupWeightSum, 1.0f)) 
+//     + EditLayerAlpha * CurrentLayerWeight / max(PremultipliedAlphaBlendGroupWeightSum, 1.0f);
+#define EWEIGHTMAPTARGETLAYERFLAGS_ISPREMULTIPLIEDALPHAWEIGHTBLENDED (1 << 3)
 
 #if MERGE_EDIT_LAYER
 // MergeEditLayerPS inputs/outputs :
 
+// Per-target layer information : 
+struct FMergeEditLayerTargetLayerInfo
+{
+	uint Flags; // See EWeightmapTargetLayerFlags
+	int BlendGroupIndex; // Defines the target layer blend group that this target layer belongs to in this blend operation. -1 if not applicable
+};
+
 uint InTargetLayerIndex; // Index of the target layer being processed
+uint InNumTargetLayers; // Number of target layers (i.e. == InMergeEditLayerTargetLayerInfos.Num())
 uint InEditLayerTargetLayerBlendMode; // See EWeightmapBlendMode. Target layer's blend mode, which is per-edit layer and per-target layer and can therefore differ from one edit layer to another
 float InEditLayerAlpha; // Global alpha value of the edit layer currently being merged
 
+// Array that contains per-target layer information (e.g. target layer's flags)
+StructuredBuffer<FMergeEditLayerTargetLayerInfo> InMergeEditLayerTargetLayerInfos;
 Texture2DArray<float4> InCurrentEditLayerWeightmaps; // The target layers of the current edit layer to merge (1 slice per weightmap in the target layer group)
 Texture2DArray<float4> InPreviousEditLayersWeightmaps; // The result from the merge of all prior edit layers (1 slice per weightmap in the target layer group)
 // Array that contains per-edit layer / per-target layer information (e.g. the target layer's blend mode, which is per-edit layer and can therefore differ from one edit layer to another)
 #endif // MERGE_EDIT_LAYER
 
-#if PERFORM_LEGACY_WEIGHT_BLENDING
-// PerformLegacyWeightBlendingPS inputs/outputs :
+#if PERFORM_FINAL_WEIGHT_BLENDING
+// PerformFinalWeightBlendingPS inputs/outputs :
 
 // Per-target layer information : 
-struct FWeightmapTargetLayerInfo
+struct FFinalWeightBlendingTargetLayerInfo
 {
 	uint Flags; // See EWeightmapTargetLayerFlags
 };
 
 uint InTargetLayerIndex; // Index of the target layer being processed
-uint InNumTargetLayers; // Number of target layers (i.e. == InTargetLayerInfos.Num())
+uint InNumTargetLayers; // Number of target layers (i.e. == InFinalWeightBlendingTargetLayerInfos.Num())
 // Array that contains per-target layer information (e.g. target layer's flags)
-StructuredBuffer<FWeightmapTargetLayerInfo> InTargetLayerInfos;
+StructuredBuffer<FFinalWeightBlendingTargetLayerInfo> InFinalWeightBlendingTargetLayerInfos;
 Texture2DArray<float4> InCurrentEditLayerWeightmaps; // The texture to horizontally blend (1 slice per weightmap in the target layer group)
-#endif // PERFORM_LEGACY_WEIGHT_BLENDING
+#endif // PERFORM_FINAL_WEIGHT_BLENDING
 
 #if PACK_WEIGHTMAP
 // PackWeightmapPS inputs/outputs :
@@ -79,70 +101,249 @@ Texture2D<float4> InSourceWeightmap; // Source weightmap (containing the current
 // ----------------------------------------------------------------------------------
 // Util functions : 
 
-#if PERFORM_LEGACY_WEIGHT_BLENDING
-bool IsWeightBlendedTargetLayer(FWeightmapTargetLayerInfo InTargetLayerInfo)
+#if MERGE_EDIT_LAYER
+bool IsPremultipliedAlphaWeightBlendedTargetLayer(FMergeEditLayerTargetLayerInfo InTargetLayerInfo)
+{
+	return (InTargetLayerInfo.Flags & EWEIGHTMAPTARGETLAYERFLAGS_ISPREMULTIPLIEDALPHAWEIGHTBLENDED)
+		// For the visibility layer, deactivate weight blending altogether : 
+		&& ((InTargetLayerInfo.Flags & EWEIGHTMAPTARGETLAYERFLAGS_ISVISIBILITYLAYER) == 0)
+		// Some layers are absent from the batch, they should be skipped from blending too :
+		&& ((InTargetLayerInfo.Flags & EWEIGHTMAPTARGETLAYERFLAGS_SKIP) == 0); 
+}
+#endif // MERGE_EDIT_LAYER
+
+#if PERFORM_FINAL_WEIGHT_BLENDING
+bool IsWeightBlendedTargetLayer(FFinalWeightBlendingTargetLayerInfo InTargetLayerInfo)
 {
-	
 	return (InTargetLayerInfo.Flags & EWEIGHTMAPTARGETLAYERFLAGS_ISWEIGHTBLENDED)
 		// For the visibility layer, deactivate weight blending altogether : 
 		&& ((InTargetLayerInfo.Flags & EWEIGHTMAPTARGETLAYERFLAGS_ISVISIBILITYLAYER) == 0)
 		// Some layers are absent from the batch, they should be skipped from blending too :
 		&& ((InTargetLayerInfo.Flags & EWEIGHTMAPTARGETLAYERFLAGS_SKIP) == 0); 
 }
-#endif // PERFORM_LEGACY_WEIGHT_BLENDING
+#endif // PERFORM_FINAL_WEIGHT_BLENDING
 
 
 // ----------------------------------------------------------------------------------
 // Pixel shaders : 
 
 #if MERGE_EDIT_LAYER
+float PerformPremultipliedAlphaBlending(float InPreviousLayersWeight, float InCurrentLayerWeight, float InPremultipliedAlphaBlendGroupWeightSum)
+{
+	return
+		// Previous layers contribution, balanced with the relative weight of other blended target layers
+		InPreviousLayersWeight * (1.0f - min(InPremultipliedAlphaBlendGroupWeightSum, 1.0f)) 
+		// Current layer contribution, balanced with the relative weight of other blended target layers
+		+ InEditLayerAlpha * InCurrentLayerWeight / max(InPremultipliedAlphaBlendGroupWeightSum, 1.0f);
+}
+
 void MergeEditLayerPS(in float4 SVPos : SV_POSITION, out float4 OutPackedWeight : SV_Target0)
 {
+	check((InTargetLayerIndex >= 0) && (InTargetLayerIndex < InNumTargetLayers));
 	uint2 TextureCoordinates = floor(SVPos.xy);
-		
-	float4 CurrentLayerWeightSample = 0.0f; 
-	if (InEditLayerTargetLayerBlendMode != EWEIGHTMAPBLENDMODE_PASSTHROUGH)
-	{
-		CurrentLayerWeightSample = InCurrentEditLayerWeightmaps.Load(int4(TextureCoordinates, InTargetLayerIndex, 0)); // xy = relative coordinates, z = index in texture array, w = mip level
-	}
-	float CurrentLayerWeight = UnpackWeight(CurrentLayerWeightSample.xy);
-	float CurrentLayerWeightAlpha = 1.0f;
-	uint CurrentLayerWeightFlags = EWEIGHTMAPALPHAFLAGS_NONE;
-	UnpackWeightAlpha(CurrentLayerWeightSample.zw, CurrentLayerWeightAlpha, CurrentLayerWeightFlags);
 
 	float2 PreviousLayersWeightSample = InPreviousEditLayersWeightmaps.Load(int4(TextureCoordinates, InTargetLayerIndex, 0)).xy; // xy = relative coordinates, z = index in texture array, w = mip level
 	float PreviousLayersWeight = UnpackWeight(PreviousLayersWeightSample);
-
 	float FinalWeight = PreviousLayersWeight;
-	
-	if (InEditLayerTargetLayerBlendMode == EWEIGHTMAPBLENDMODE_ADDITIVE)
-	{
-		FinalWeight += CurrentLayerWeight * InEditLayerAlpha;
-	}
-	else if (InEditLayerTargetLayerBlendMode == EWEIGHTMAPBLENDMODE_SUBTRACTIVE)
-	{
-		FinalWeight -= CurrentLayerWeight * InEditLayerAlpha;
-	}
-	else if (InEditLayerTargetLayerBlendMode == EWEIGHTMAPBLENDMODE_ALPHABLEND)
+	if (InEditLayerTargetLayerBlendMode != EWEIGHTMAPBLENDMODE_PASSTHROUGH)
 	{
-		float FinalAlpha = InEditLayerAlpha * CurrentLayerWeightAlpha;
-		if (CurrentLayerWeightFlags == EWEIGHTMAPALPHAFLAGS_ALPHABLEND)
-		{
-			FinalWeight = lerp(PreviousLayersWeight, CurrentLayerWeight, FinalAlpha);
-		}
-		else if (CurrentLayerWeightFlags & EWEIGHTMAPALPHAFLAGS_MIN)
-		{
-			FinalWeight = lerp(PreviousLayersWeight, min(PreviousLayersWeight, CurrentLayerWeight), FinalAlpha);
-		}
-		else if (CurrentLayerWeightFlags & EWEIGHTMAPALPHAFLAGS_MAX)
+		float4 CurrentLayerWeightSample = InCurrentEditLayerWeightmaps.Load(int4(TextureCoordinates, InTargetLayerIndex, 0)); // xy = relative coordinates, z = index in texture array, w = mip level
+		float CurrentLayerWeight = UnpackWeight(CurrentLayerWeightSample.xy);
+		float CurrentLayerWeightAlpha = 1.0f;
+		uint CurrentLayerWeightFlags = EWEIGHTMAPALPHAFLAGS_NONE;
+		UnpackWeightAlpha(CurrentLayerWeightSample.zw, CurrentLayerWeightAlpha, CurrentLayerWeightFlags);
+
+		// TODO [jonathan.bard] : remove InEditLayerTargetLayerBlendMode == EWEIGHTMAPBLENDMODE_SUBTRACTIVE and add the EWEIGHTMAPALPHAFLAGS_SUBTRACTIVE per-pixel flag ?
+		// TODO [jonathan.bard] : handle subtractive for premultiplied alpha (it's yet another "exotic" blend mode)
+		if (InEditLayerTargetLayerBlendMode == EWEIGHTMAPBLENDMODE_SUBTRACTIVE)
 		{
-			FinalWeight = lerp(PreviousLayersWeight, max(PreviousLayersWeight, CurrentLayerWeight), FinalAlpha);
+			float FinalAlpha = InEditLayerAlpha * CurrentLayerWeightAlpha;
+			FinalWeight -= CurrentLayerWeight * FinalAlpha;
 		}
-		// Additive case (== no flag) :
-		// TODO [jonathan.bard] : remove InEditLayerTargetLayerBlendMode == EWEIGHTMAPBLENDMODE_SUBTRACTIVE and add the EWEIGHTMAPBLENDMODE_SUBTRACTIVE flag
-		else if (CurrentLayerWeightFlags == EWEIGHTMAPALPHAFLAGS_ADDITIVE)
+		else 
 		{
-			FinalWeight += CurrentLayerWeight * FinalAlpha;
+			FMergeEditLayerTargetLayerInfo OutputTargetLayerInfo = InMergeEditLayerTargetLayerInfos[InTargetLayerIndex];
+			if (IsPremultipliedAlphaWeightBlendedTargetLayer(OutputTargetLayerInfo))
+			{
+				// Start again from scratch 
+				FinalWeight = 0.0f;
+				check(OutputTargetLayerInfo.BlendGroupIndex >= 0);
+
+				// Compute the weight sum of all other target layers belonging to the same blend group 
+				float PremultipliedAlphaBlendGroupWeightSum = 0.0f;
+				bool bNeedsNonAdditiveBlendStep = false;
+
+				LOOP
+				for (uint OtherTargetLayerIndex = 0; OtherTargetLayerIndex < InNumTargetLayers; ++OtherTargetLayerIndex)
+				{
+					FMergeEditLayerTargetLayerInfo TargetLayerInfo = InMergeEditLayerTargetLayerInfos[OtherTargetLayerIndex];
+					if (IsPremultipliedAlphaWeightBlendedTargetLayer(TargetLayerInfo) 
+						&& (OutputTargetLayerInfo.BlendGroupIndex == TargetLayerInfo.BlendGroupIndex))
+					{
+						check(TargetLayerInfo.BlendGroupIndex >= 0);
+
+						bool bIsCurrentLayer = (OtherTargetLayerIndex == InTargetLayerIndex);
+						float4 OtherWeightSample = InCurrentEditLayerWeightmaps.Load(int4(TextureCoordinates, OtherTargetLayerIndex, 0)); // xy = relative coordinates, z = index in texture array, w = mip level
+						float OtherWeight = UnpackWeight(OtherWeightSample.xy);
+						float OtherWeightAlpha = 1.0f;
+						uint OtherWeightFlags = EWEIGHTMAPALPHAFLAGS_NONE;
+						UnpackWeightAlpha(OtherWeightSample.zw, OtherWeightAlpha, OtherWeightFlags);
+
+						float FinalOtherWeight = 0.0f;
+						if (InEditLayerTargetLayerBlendMode == EWEIGHTMAPBLENDMODE_ADDITIVE)
+						{
+							FinalOtherWeight = OtherWeight * OtherWeightAlpha;
+							PremultipliedAlphaBlendGroupWeightSum += FinalOtherWeight * InEditLayerAlpha;
+						}
+						else if (InEditLayerTargetLayerBlendMode == EWEIGHTMAPBLENDMODE_ALPHABLEND)
+						{
+							// Additive case (== no flag), we can handle with the other premultiplied alpha layers : 
+							if (OtherWeightFlags == EWEIGHTMAPALPHAFLAGS_ADDITIVE)
+							{
+								FinalOtherWeight = OtherWeight * OtherWeightAlpha;
+								PremultipliedAlphaBlendGroupWeightSum += FinalOtherWeight * InEditLayerAlpha;
+							}
+							// Other, more exotic, blend modes have to be handled in a second phase and then re-balanced with the others accordingly :
+							else
+							{
+								bNeedsNonAdditiveBlendStep = true;
+							}
+						}
+						else
+						{
+							check(false);
+						}
+
+						// Remember the weight of the current layer for the final computation : 
+						if (bIsCurrentLayer)
+						{
+							FinalWeight = FinalOtherWeight;
+						}
+					}
+				}
+
+				// If all layers are premultiplied alpha-blended, perform the final computation already and be done with it : 
+				if (!bNeedsNonAdditiveBlendStep)
+				{
+					FinalWeight = PerformPremultipliedAlphaBlending(PreviousLayersWeight, FinalWeight, PremultipliedAlphaBlendGroupWeightSum);
+				}
+				// Otherwise (some of the target layers had a non-standard blend), we have to compute each of the layers' final values in order to compute a new sum and adjust the final result accordingly :
+				else
+				{
+					// Start again from scratch : 
+					FinalWeight = 0.0;
+					float AdjustedBlendGroupWeightSum = 0.0f;
+					LOOP
+					for (uint OtherTargetLayerIndex = 0; OtherTargetLayerIndex < InNumTargetLayers; ++OtherTargetLayerIndex)
+					{
+						FMergeEditLayerTargetLayerInfo TargetLayerInfo = InMergeEditLayerTargetLayerInfos[OtherTargetLayerIndex];
+						if (IsPremultipliedAlphaWeightBlendedTargetLayer(TargetLayerInfo) 
+							&& (OutputTargetLayerInfo.BlendGroupIndex == TargetLayerInfo.BlendGroupIndex))
+						{
+							check(TargetLayerInfo.BlendGroupIndex >= 0);
+
+							bool bIsCurrentLayer = (OtherTargetLayerIndex == InTargetLayerIndex);
+							float4 OtherWeightSample = InCurrentEditLayerWeightmaps.Load(int4(TextureCoordinates, OtherTargetLayerIndex, 0)); // xy = relative coordinates, z = index in texture array, w = mip level
+							float OtherWeight = UnpackWeight(OtherWeightSample.xy);
+							float OtherWeightAlpha = 1.0f;
+							uint OtherWeightFlags = EWEIGHTMAPALPHAFLAGS_NONE;
+							UnpackWeightAlpha(OtherWeightSample.zw, OtherWeightAlpha, OtherWeightFlags);
+
+							float4 PreviousOtherWeightSample = InPreviousEditLayersWeightmaps.Load(int4(TextureCoordinates, OtherTargetLayerIndex, 0)); // xy = relative coordinates, z = index in texture array, w = mip level
+							float PreviousOtherWeight = UnpackWeight(PreviousOtherWeightSample.xy);
+
+							float FinalOtherWeight = 0.0f;
+							if (InEditLayerTargetLayerBlendMode == EWEIGHTMAPBLENDMODE_ADDITIVE)
+							{
+								// Compute the final weight value of this layer here and add it to the final sum so that we can re-adjust it against the other non-additive cases at the end :
+								// Compute (again) the layer's additive contribution
+								FinalOtherWeight = OtherWeight * OtherWeightAlpha;
+								// Then perform premultiplied alpha blending to get the final layer's contribution, as if all layers were premultiplied alpha blended :
+								FinalOtherWeight = PerformPremultipliedAlphaBlending(PreviousOtherWeight, FinalOtherWeight, PremultipliedAlphaBlendGroupWeightSum);
+							}
+							else if (InEditLayerTargetLayerBlendMode == EWEIGHTMAPBLENDMODE_ALPHABLEND)
+							{
+								// Additive case (== no flag) : compute the final weight value of this layer here and add it to the final sum so that we can re-adjust it against the other non-additive cases at the end :
+								if (OtherWeightFlags == EWEIGHTMAPALPHAFLAGS_ADDITIVE)
+								{
+									// Compute the final weight value of this layer here and add it to the final sum so that we can re-adjust it against the other non-additive cases at the end :
+									// Compute (again) the layer's additive contribution
+									FinalOtherWeight = OtherWeight * OtherWeightAlpha;
+									// Then perform premultiplied alpha blending to get the final layer's contribution, as if all layers were premultiplied alpha blended :
+									FinalOtherWeight = PerformPremultipliedAlphaBlending(PreviousOtherWeight, FinalOtherWeight, PremultipliedAlphaBlendGroupWeightSum);
+								}
+								else
+								{
+									// Handle exotic alpha blend cases (min/max/alphablend) here 
+									if (OtherWeightFlags == EWEIGHTMAPALPHAFLAGS_ALPHABLEND)
+									{
+										FinalOtherWeight = lerp(PreviousOtherWeight, OtherWeight, OtherWeightAlpha);
+									}
+									else if (OtherWeightFlags & EWEIGHTMAPALPHAFLAGS_MIN)
+									{
+										FinalOtherWeight = lerp(PreviousOtherWeight, min(PreviousOtherWeight, OtherWeight), OtherWeightAlpha);
+									}
+									else if (OtherWeightFlags & EWEIGHTMAPALPHAFLAGS_MAX)
+									{
+										FinalOtherWeight = lerp(PreviousOtherWeight, max(PreviousOtherWeight, OtherWeight), OtherWeightAlpha);
+									}
+									else
+									{
+										check(false);
+									}
+								}
+							}
+							else
+							{
+								check(false);
+							}
+							AdjustedBlendGroupWeightSum += FinalOtherWeight;
+
+							// Remember the weight of the current layer for the final computation : 
+							if (bIsCurrentLayer)
+							{
+								FinalWeight = FinalOtherWeight;
+							}
+						}
+					}
+
+					// Finally, perform a final, standard, weight-blending step to make sure all target layers are balanced and their cumulative weight is <= 1.0 : 
+					FinalWeight /= max(AdjustedBlendGroupWeightSum, 1.0f);
+				}
+			}
+			// Standard case : no inter-target layer blending to perform
+			else
+			{
+				float FinalAlpha = InEditLayerAlpha * CurrentLayerWeightAlpha;
+				if (InEditLayerTargetLayerBlendMode == EWEIGHTMAPBLENDMODE_ADDITIVE)
+				{
+					FinalWeight = PreviousLayersWeight + CurrentLayerWeight * FinalAlpha;
+				}
+				else if (InEditLayerTargetLayerBlendMode == EWEIGHTMAPBLENDMODE_ALPHABLEND)
+				{
+					if (CurrentLayerWeightFlags == EWEIGHTMAPALPHAFLAGS_ALPHABLEND)
+					{
+						FinalWeight = lerp(PreviousLayersWeight, CurrentLayerWeight, FinalAlpha);
+					}
+					else if (CurrentLayerWeightFlags & EWEIGHTMAPALPHAFLAGS_MIN)
+					{
+						FinalWeight = lerp(PreviousLayersWeight, min(PreviousLayersWeight, CurrentLayerWeight), FinalAlpha);
+					}
+					else if (CurrentLayerWeightFlags & EWEIGHTMAPALPHAFLAGS_MAX)
+					{
+						FinalWeight = lerp(PreviousLayersWeight, max(PreviousLayersWeight, CurrentLayerWeight), FinalAlpha);
+					}
+					// Additive case (== no flag) :
+					else if (CurrentLayerWeightFlags == EWEIGHTMAPALPHAFLAGS_ADDITIVE)
+					{
+						FinalWeight = PreviousLayersWeight + CurrentLayerWeight * FinalAlpha;
+					}
+				}
+				else
+				{
+					check(false);
+				}
+			}
 		}
 	}
 
@@ -150,13 +351,15 @@ void MergeEditLayerPS(in float4 SVPos : SV_POSITION, out float4 OutPackedWeight
 }
 #endif // MERGE_EDIT_LAYER
 
-#if PERFORM_LEGACY_WEIGHT_BLENDING
-void PerformLegacyWeightBlendingPS(in float4 SVPos : SV_POSITION, out float4 OutColor : SV_Target0)
+#if PERFORM_FINAL_WEIGHT_BLENDING
+void PerformFinalWeightBlendingPS(in float4 SVPos : SV_POSITION, out float4 OutColor : SV_Target0)
 {
+	check((InTargetLayerIndex >= 0) && (InTargetLayerIndex < InNumTargetLayers));
+
 	uint2 TextureCoordinates = floor(SVPos.xy);
 	float FinalWeight = 0.0f;
 	
-	FWeightmapTargetLayerInfo ActiveTargetLayerInfo = InTargetLayerInfos[InTargetLayerIndex];
+	FFinalWeightBlendingTargetLayerInfo ActiveTargetLayerInfo = InFinalWeightBlendingTargetLayerInfos[InTargetLayerIndex];
 	if (!IsWeightBlendedTargetLayer(ActiveTargetLayerInfo))
 	{
 		FinalWeight = UnpackWeight(InCurrentEditLayerWeightmaps.Load(int4(TextureCoordinates, InTargetLayerIndex, 0)).xy); // xy = relative coordinates, z = index in texture array, w = mip level
@@ -170,7 +373,7 @@ void PerformLegacyWeightBlendingPS(in float4 SVPos : SV_POSITION, out float4 Out
 		for (uint i = 0; i < InNumTargetLayers; ++i)
 		{
 			bool bIsOutputTargetLayer = (i == InTargetLayerIndex);
-			FWeightmapTargetLayerInfo TargetLayerInfo = InTargetLayerInfos[i];
+			FFinalWeightBlendingTargetLayerInfo TargetLayerInfo = InFinalWeightBlendingTargetLayerInfos[i];
 
 			// Only weight blended (and non-visibility) target layers participate to weight blending : 
 			if (IsWeightBlendedTargetLayer(TargetLayerInfo))
@@ -189,7 +392,7 @@ void PerformLegacyWeightBlendingPS(in float4 SVPos : SV_POSITION, out float4 Out
 
 	OutColor = float4(PackWeight(FinalWeight), 0.0f, 0.0f);
 }
-#endif // PERFORM_LEGACY_WEIGHT_BLENDING
+#endif // PERFORM_FINAL_WEIGHT_BLENDING
 
 // ----------------------------------------------------------------------------------
 
diff --git a/Engine/Source/Editor/EditorStyle/Private/StarshipStyle.cpp b/Engine/Source/Editor/EditorStyle/Private/StarshipStyle.cpp
--- a/Engine/Source/Editor/EditorStyle/Private/StarshipStyle.cpp
+++ b/Engine/Source/Editor/EditorStyle/Private/StarshipStyle.cpp
@@ -7313,13 +7313,13 @@ void FStarshipEditorStyle::FStyle::SetupLandscapeEditorStyle()
 		Set("LandscapeEditor.Target_Visibility", new IMAGE_BRUSH_SVG("Icons/icon_Landscape_Target_Visibility_64x", Icon48x48));
 		Set("LandscapeEditor.Target_Unknown",    new IMAGE_BRUSH_SVG("Icons/icon_Landscape_Target_Unknown_64x",    Icon48x48));
 
-		Set("LandscapeEditor.Target_Create",     new IMAGE_BRUSH("Icons/icon_Landscape_Target_Create_12x", Icon12x12));
-		Set("LandscapeEditor.Target_MakePublic", new IMAGE_BRUSH("Icons/assign_right_12x",                 Icon12x12));
-		Set("LandscapeEditor.Target_Delete",     new IMAGE_BRUSH("Icons/Cross_12x",                        Icon12x12));
+		Set("LandscapeEditor.Target_Create",     new IMAGE_BRUSH("Icons/PlusSymbol_12x", Icon12x12));
+		Set("LandscapeEditor.Target_Delete",     new IMAGE_BRUSH("Icons/Cross_12x", Icon12x12));
 
 		Set("LandscapeEditor.Target_DisplayOrder.Default", new IMAGE_BRUSH("Icons/icon_landscape_sort_base", Icon16x16));
 		Set("LandscapeEditor.Target_DisplayOrder.Alphabetical", new IMAGE_BRUSH("Icons/icon_landscape_sort_alphabetical", Icon16x16));
 		Set("LandscapeEditor.Target_DisplayOrder.Custom", new IMAGE_BRUSH("Icons/icon_landscape_sort_custom", Icon16x16));
+		Set("LandscapeEditor.Target_DisplayOrder.ByBlendMethod", new IMAGE_BRUSH("Icons/icon_landscape_sort_byblendmethod", Icon16x16));
 
 		Set("LandscapeEditor.TargetList.DropZone.Above", new BOX_BRUSH("Common/VerticalBoxDropZoneIndicator_Above", FMargin(10.0f / 16.0f, 10.0f / 16.0f, 0, 0), SelectionColor_Subdued));
 		Set("LandscapeEditor.TargetList.DropZone.Below", new BOX_BRUSH("Common/VerticalBoxDropZoneIndicator_Below", FMargin(10.0f / 16.0f, 0, 0, 10.0f / 16.0f), SelectionColor_Subdued));
diff --git a/Engine/Source/Editor/LandscapeEditor/Private/LandscapeEdMode.cpp b/Engine/Source/Editor/LandscapeEditor/Private/LandscapeEdMode.cpp
--- a/Engine/Source/Editor/LandscapeEditor/Private/LandscapeEdMode.cpp
+++ b/Engine/Source/Editor/LandscapeEditor/Private/LandscapeEdMode.cpp
@@ -3,6 +3,7 @@
 #include "LandscapeEdMode.h"
 
 #include "Algo/Accumulate.h"
+#include "Algo/Transform.h"
 #include "MaterialDomain.h"
 #include "Materials/Material.h"
 #include "SceneView.h"
@@ -2778,6 +2779,73 @@ void FEdModeLandscape::UpdateTargetLayerDisplayOrder(ELandscapeLayerDisplayMode
 			}
 		}
 		break;
+
+		case ELandscapeLayerDisplayMode::ByBlendMethod:
+		{
+			SavedTargetNameList.Empty();
+
+			TArray<TSharedRef<FLandscapeTargetListInfo>> SortedTargetList;
+			// Add only layers to be able to sort them by blend type
+			for (int32 i = GetTargetLayerStartingIndex(); i < LandscapeTargetList.Num(); ++i)
+			{
+				SortedTargetList.Add(LandscapeTargetList[i]);
+			}
+
+			SortedTargetList.Sort([](const TSharedRef<FLandscapeTargetListInfo>& InLHS, const TSharedRef<FLandscapeTargetListInfo>& InRHS)
+				{
+					// Put invalid layers at the end : 
+					if (InLHS->LayerInfoObj == nullptr)
+					{
+						return (InRHS->LayerInfoObj == nullptr) ? InLHS->LayerName.LexicalLess(InRHS->LayerName) : false;
+					}
+					else if (InRHS->LayerInfoObj == nullptr)
+					{
+						return true;
+					}
+
+					if (InLHS->LayerInfoObj->GetBlendMethod() == InRHS->LayerInfoObj->GetBlendMethod())
+					{
+						// For a given blend method, sort by blend group 
+						if (InLHS->LayerInfoObj->GetBlendGroup() != InRHS->LayerInfoObj->GetBlendGroup())
+						{
+							// But put empty group first (most frequent case)
+							if (InLHS->LayerInfoObj->GetBlendGroup().IsNone())
+							{
+								return true;
+							}
+							else if (InRHS->LayerInfoObj->GetBlendGroup().IsNone())
+							{
+								return false;
+							}
+							
+							return InLHS->LayerInfoObj->GetBlendGroup().LexicalLess(InRHS->LayerInfoObj->GetBlendGroup());
+						}
+
+						return InLHS->LayerName.LexicalLess(InRHS->LayerName);
+					}
+					else
+					{
+						static constexpr int32 BlendMethodSortOrder[static_cast<int32>(ELandscapeTargetLayerBlendMethod::Count)] = 
+						{
+							2, /* ELandscapeTargetLayerBlendMethod::None */
+							1, /* ELandscapeTargetLayerBlendMethod::FinalWeightBlending */
+							0, /* ELandscapeTargetLayerBlendMethod::PremultipliedAlphaBlending */
+						};
+
+						return (BlendMethodSortOrder[static_cast<int32>(InLHS->LayerInfoObj->GetBlendMethod())] < BlendMethodSortOrder[static_cast<int32>(InRHS->LayerInfoObj->GetBlendMethod())]);
+					}
+				});
+			Algo::Transform(SortedTargetList, SavedTargetNameList, [](const TSharedRef<FLandscapeTargetListInfo>& InTargetListInfo) { return InTargetListInfo->GetLayerName(); });
+
+			// Then insert the non layer target that shouldn't be sorted
+			for (int32 i = 0; i < GetTargetLayerStartingIndex(); ++i)
+			{
+				SavedTargetNameList.Insert(LandscapeTargetList[i]->LayerName, i);
+			}
+
+			DetailPanelRefreshRequired = true;
+		}
+		break;
 	}	
 
 	if (DetailPanelRefreshRequired)
diff --git a/Engine/Source/Editor/LandscapeEditor/Private/LandscapeEditorDetailCustomization_ImportLayers.cpp b/Engine/Source/Editor/LandscapeEditor/Private/LandscapeEditorDetailCustomization_ImportLayers.cpp
--- a/Engine/Source/Editor/LandscapeEditor/Private/LandscapeEditorDetailCustomization_ImportLayers.cpp
+++ b/Engine/Source/Editor/LandscapeEditor/Private/LandscapeEditorDetailCustomization_ImportLayers.cpp
@@ -7,8 +7,9 @@
 #include "Editor.h"
 #include "LandscapeEditorModule.h"
 #include "LandscapeEditorObject.h"
-#include "LandscapeUtils.h"
 #include "LandscapeEditorUtils.h"
+#include "LandscapeSettings.h"
+#include "LandscapeUtils.h"
 #include "SLandscapeEditor.h"
 
 #include "IDetailChildrenBuilder.h"
@@ -38,17 +39,17 @@ bool FLandscapeEditorStructCustomization_FLandscapeImportLayer::IsImporting()
 {
 	if (FEdModeLandscape* LandscapeEdMode = GetEditorMode())
 	{
-		if (IsToolActive("ImportExport"))
-		{
-			return LandscapeEdMode->ImportExportMode == EImportExportMode::Import;
-		}
+	if (IsToolActive("ImportExport"))
+	{
+		return LandscapeEdMode->ImportExportMode == EImportExportMode::Import;
+	}
 
 		if (IsToolActive("NewLandscape"))
 		{
 			if (LandscapeEdMode->NewLandscapePreviewMode == ENewLandscapePreviewMode::ImportLandscape)
 			{
-				return true;
-			}
+	return true;
+}
 		}
 	}
 
@@ -71,15 +72,15 @@ void FLandscapeEditorStructCustomization_FLandscapeImportLayer::CustomizeChildre
 	TSharedRef<IPropertyHandle> PropertyHandle_Selected = StructPropertyHandle->GetChildHandle(GET_MEMBER_NAME_CHECKED(FLandscapeImportLayer, bSelected)).ToSharedRef();
 	
 	PropertyHandle_SourceFilePath->SetOnPropertyValueChanged(FSimpleDelegate::CreateLambda([PropertyHandle_SourceFilePath]()
-	{
-		FLandscapeEditorDetailCustomization_ImportExport::FormatFilename(PropertyHandle_SourceFilePath, /*bForExport = */false);
-		FLandscapeEditorStructCustomization_FLandscapeImportLayer::OnImportWeightmapFilenameChanged();
-	}));
+		{
+			FLandscapeEditorDetailCustomization_ImportExport::FormatFilename(PropertyHandle_SourceFilePath, /*bForExport = */false);
+			FLandscapeEditorStructCustomization_FLandscapeImportLayer::OnImportWeightmapFilenameChanged();
+		}));
 
 	PropertyHandle_ExportFilePath->SetOnPropertyValueChanged(FSimpleDelegate::CreateLambda([PropertyHandle_ExportFilePath]()
-	{
-		FLandscapeEditorDetailCustomization_ImportExport::FormatFilename(PropertyHandle_ExportFilePath, /*bForExport = */true);
-	}));
+		{
+			FLandscapeEditorDetailCustomization_ImportExport::FormatFilename(PropertyHandle_ExportFilePath, /*bForExport = */true);
+		}));
 
 	FName LayerName;
 	FText LayerNameText;
@@ -97,131 +98,130 @@ void FLandscapeEditorStructCustomization_FLandscapeImportLayer::CustomizeChildre
 	checkSlow(Result == FPropertyAccess::Success);
 
 	auto CreateFilePathNameContentWidget = [&](bool bImport)
-	{
-		return SNew(SHorizontalBox)
-			.Visibility_Static(&GetImportExportVisibility, bImport)
-			+ SHorizontalBox::Slot()
-			.AutoWidth()
-			.VAlign(VAlign_Center)
-			.HAlign(HAlign_Left)
-			.Padding(FMargin(2))
-			[
-				SNew(SCheckBox)
-					.Visibility_Static(GetImportLayerSelectionVisibility)
-					.IsEnabled_Static(&FLandscapeEditorStructCustomization_FLandscapeImportLayer::IsValidLayerInfo, PropertyHandle_LayerInfo)
-					.IsChecked_Static(&FLandscapeEditorStructCustomization_FLandscapeImportLayer::GetImportLayerSelectedCheckState, PropertyHandle_Selected, PropertyHandle_LayerInfo)
-					.OnCheckStateChanged_Static(&FLandscapeEditorStructCustomization_FLandscapeImportLayer::OnImportLayerSelectedCheckStateChanged, PropertyHandle_Selected)
-					.ToolTipText_Static(&FLandscapeEditorStructCustomization_FLandscapeImportLayer::GetImportLayerSelectedToolTip, PropertyHandle_Selected, PropertyHandle_LayerInfo)
-			]
-			+ SHorizontalBox::Slot()
-			.AutoWidth()
-			.VAlign(VAlign_Center)
-			.Padding(FMargin(2))
-			[
-				SNew(SLandscapeAssetThumbnail, ThumbnailMIC, StructCustomizationUtils.GetThumbnailPool().ToSharedRef(), FName("LandscapeEditor.Target_Weightmap"))
-					.ThumbnailSize(FIntPoint(48, 48))
-			]
+		{
+			return SNew(SHorizontalBox)
+				.Visibility_Static(&GetImportExportVisibility, bImport)
+				+ SHorizontalBox::Slot()
+				.AutoWidth()
+				.VAlign(VAlign_Center)
+				.HAlign(HAlign_Left)
+				.Padding(FMargin(2))
+				[
+					SNew(SCheckBox)
+						.Visibility_Static(GetImportLayerSelectionVisibility)
+						.IsEnabled_Static(&FLandscapeEditorStructCustomization_FLandscapeImportLayer::IsValidLayerInfo, PropertyHandle_LayerInfo)
+						.IsChecked_Static(&FLandscapeEditorStructCustomization_FLandscapeImportLayer::GetImportLayerSelectedCheckState, PropertyHandle_Selected, PropertyHandle_LayerInfo)
+						.OnCheckStateChanged_Static(&FLandscapeEditorStructCustomization_FLandscapeImportLayer::OnImportLayerSelectedCheckStateChanged, PropertyHandle_Selected)
+						.ToolTipText_Static(&FLandscapeEditorStructCustomization_FLandscapeImportLayer::GetImportLayerSelectedToolTip, PropertyHandle_Selected, PropertyHandle_LayerInfo)
+				]
 			+ SHorizontalBox::Slot()
-			.FillWidth(1)
-			.VAlign(VAlign_Center)
-			.Padding(FMargin(2))
-			[
-				SNew(STextBlock)
-					.Font(StructCustomizationUtils.GetRegularFont())
-					.Text(LayerNameText)
-			];
-	};
+				.AutoWidth()
+				.VAlign(VAlign_Center)
+				.Padding(FMargin(2))
+				[
+					SNew(SLandscapeAssetThumbnail, ThumbnailMIC, StructCustomizationUtils.GetThumbnailPool().ToSharedRef(), FName("LandscapeEditor.Target_Weightmap"))
+						.ThumbnailSize(FIntPoint(48, 48))
+				]
+				+ SHorizontalBox::Slot()
+				.FillWidth(1)
+				.VAlign(VAlign_Center)
+				.Padding(FMargin(2))
+				[
+					SNew(STextBlock)
+						.Font(StructCustomizationUtils.GetRegularFont())
+						.Text(LayerNameText)
+				];
+		};
 
 	auto CreateFilePathContentWidget = [&](TSharedRef<IPropertyHandle> PropertyHandle_Filepath, bool bImport)
-	{
-		return SNew(SBox)
-			.VAlign(VAlign_Center)
-			.Visibility_Static(&GetImportExportVisibility, bImport)
-			.IsEnabled_Static(&FLandscapeEditorStructCustomization_FLandscapeImportLayer::IsImportLayerSelected, PropertyHandle_Selected, PropertyHandle_LayerInfo)
-			[
-				SNew(SVerticalBox)
-					+ SVerticalBox::Slot()
-					.AutoHeight()
-					[
-						SNew(SHorizontalBox)
-							.Visibility_Static(&FLandscapeEditorStructCustomization_FLandscapeImportLayer::GetLayerInfoAssignVisibility)
-							+ SHorizontalBox::Slot()
-							[
-								SNew(SObjectPropertyEntryBox)
-									.AllowedClass(ULandscapeLayerInfoObject::StaticClass())
-									.PropertyHandle(PropertyHandle_LayerInfo)
-									.OnShouldFilterAsset_Static(&ShouldFilterLayerInfo, LayerName)
-									.AllowCreate(false)
-							]
+		{
+			return SNew(SBox)
+				.VAlign(VAlign_Center)
+				.Visibility_Static(&GetImportExportVisibility, bImport)
+				.IsEnabled_Static(&FLandscapeEditorStructCustomization_FLandscapeImportLayer::IsImportLayerSelected, PropertyHandle_Selected, PropertyHandle_LayerInfo)
+				[
+					SNew(SVerticalBox)
+						+ SVerticalBox::Slot()
+						.AutoHeight()
+						[
+							SNew(SHorizontalBox)
+								.Visibility_Static(&FLandscapeEditorStructCustomization_FLandscapeImportLayer::GetLayerInfoAssignVisibility)
+								+ SHorizontalBox::Slot()
+								[
+									SNew(SObjectPropertyEntryBox)
+										.AllowedClass(ULandscapeLayerInfoObject::StaticClass())
+										.PropertyHandle(PropertyHandle_LayerInfo)
+										.OnShouldFilterAsset_Static(&ShouldFilterLayerInfo, LayerName)
+										.AllowCreate(false)
+								]
 							+ SHorizontalBox::Slot()
-							.AutoWidth()
-							.VAlign(VAlign_Center)
-							[
-								SNew(SComboButton)
-									.ButtonStyle(FAppStyle::Get(), "HoverHintOnly")
-									.HasDownArrow(false)
-									.ContentPadding(4.0f)
-									.ForegroundColor(FSlateColor::UseForeground())
-									.IsFocusable(false)
-									.ToolTipText(LOCTEXT("Target_Create", "Create Layer Info"))
-									.Visibility_Static(&GetImportLayerCreateVisibility, PropertyHandle_LayerInfo)
-									.OnGetMenuContent_Static(&OnGetImportLayerCreateMenu, PropertyHandle_LayerInfo, LayerName)
-									.ButtonContent()
-									[
-										SNew(SImage)
-											.Image(FAppStyle::GetBrush("LandscapeEditor.Target_Create"))
-											.ColorAndOpacity(FSlateColor::UseForeground())
-									]
-							]
-					]
+								.AutoWidth()
+								.VAlign(VAlign_Center)
+								[
+									SNew(SButton)
+										.ButtonStyle(FAppStyle::Get(), "HoverHintOnly")
+										.ContentPadding(4.0f)
+										.ForegroundColor(FSlateColor::UseForeground())
+										.IsFocusable(false)
+										.ToolTipText(LOCTEXT("Target_Create", "Create Layer Info"))
+										.Visibility_Static(&GetImportLayerCreateVisibility, PropertyHandle_LayerInfo)
+										.OnClicked_Lambda([PropertyHandle_LayerInfo, LayerName]() { return FLandscapeEditorStructCustomization_FLandscapeImportLayer::OnImportLayerCreateClicked(PropertyHandle_LayerInfo, LayerName); })
+										[
+											SNew(SImage)
+												.Image(FAppStyle::GetBrush("LandscapeEditor.Target_Create"))
+												.ColorAndOpacity(FSlateColor::UseForeground())
+										]
+								]
+						]
 					+ SVerticalBox::Slot()
-					.AutoHeight()
-					[
-						SNew(SVerticalBox)
-							.Visibility_Static(&FLandscapeEditorStructCustomization_FLandscapeImportLayer::GetImportLayerVisibility)
-							.IsEnabled_Static(&FLandscapeEditorStructCustomization_FLandscapeImportLayer::IsImportLayerSelected, PropertyHandle_Selected, PropertyHandle_LayerInfo)
-							+ SVerticalBox::Slot()
-							.AutoHeight()
-							[
-								SNew(SHorizontalBox)
-									+ SHorizontalBox::Slot()
-									.AutoWidth()
-									.Padding(0, 0, 2, 0)
-									[
-										SNew(SErrorText)
-											.Visibility_Static(&GetErrorVisibility, PropertyHandle_ImportResult)
-											.BackgroundColor_Static(&GetErrorColor, PropertyHandle_ImportResult)
-											.ErrorText(NSLOCTEXT("UnrealEd", "Error", "!"))
-											.ToolTip(
-												SNew(SToolTip)
-												.Text_Static(&GetErrorText, PropertyHandle_ErrorMessage)
-											)
-									]
-									+ SHorizontalBox::Slot()
-									[
-										PropertyHandle_Filepath->CreatePropertyValueWidget()
-									]
+						.AutoHeight()
+						[
+							SNew(SVerticalBox)
+								.Visibility_Static(&FLandscapeEditorStructCustomization_FLandscapeImportLayer::GetImportLayerVisibility)
+								.IsEnabled_Static(&FLandscapeEditorStructCustomization_FLandscapeImportLayer::IsImportLayerSelected, PropertyHandle_Selected, PropertyHandle_LayerInfo)
+								+ SVerticalBox::Slot()
+								.AutoHeight()
+								[
+									SNew(SHorizontalBox)
+										+ SHorizontalBox::Slot()
+										.AutoWidth()
+										.Padding(0, 0, 2, 0)
+										[
+											SNew(SErrorText)
+												.Visibility_Static(&GetErrorVisibility, PropertyHandle_ImportResult)
+												.BackgroundColor_Static(&GetErrorColor, PropertyHandle_ImportResult)
+												.ErrorText(NSLOCTEXT("UnrealEd", "Error", "!"))
+												.ToolTip(
+													SNew(SToolTip)
+													.Text_Static(&GetErrorText, PropertyHandle_ErrorMessage)
+												)
+										]
 									+ SHorizontalBox::Slot()
-									.AutoWidth()
-									.Padding(1, 0, 0, 0)
-									[
-										SNew(SButton)
-											.ContentPadding(FMargin(4, 0))
-											.Text(NSLOCTEXT("UnrealEd", "GenericOpenDialog", "..."))
-											.OnClicked_Static(&FLandscapeEditorStructCustomization_FLandscapeImportLayer::OnLayerFilenameButtonClicked, PropertyHandle_Filepath)
-									]
-							]
-					]
-			];
-	};
+										[
+											PropertyHandle_Filepath->CreatePropertyValueWidget()
+										]
+										+ SHorizontalBox::Slot()
+										.AutoWidth()
+										.Padding(1, 0, 0, 0)
+										[
+											SNew(SButton)
+												.ContentPadding(FMargin(4, 0))
+												.Text(NSLOCTEXT("UnrealEd", "GenericOpenDialog", "..."))
+												.OnClicked_Static(&FLandscapeEditorStructCustomization_FLandscapeImportLayer::OnLayerFilenameButtonClicked, PropertyHandle_Filepath)
+										]
+								]
+						]
+				];
+		};
 
 	// Here we add both the Source (Import) and export file paths, one is always hidden by the visibility lambda
 	// Each property handle needs to be registered to a custom row for the Reset To Default arrows to be shown
 	IDetailPropertyRow& ImportFilePathRow = ChildBuilder.AddProperty(PropertyHandle_SourceFilePath)
-		.Visibility(TAttribute<EVisibility>::CreateLambda([PropertyHandle_Selected]() {
-			return GetImportExportVisibility(/*bImport=*/true);
-		}));
-	
+		.Visibility(TAttribute<EVisibility>::CreateLambda([PropertyHandle_Selected]()
+			{
+				return GetImportExportVisibility(/*bImport=*/true);
+			}));
+
 	ImportFilePathRow.CustomWidget()
 		.NameContent()
 		[
@@ -235,20 +235,22 @@ void FLandscapeEditorStructCustomization_FLandscapeImportLayer::CustomizeChildre
 		]
 		.OverrideResetToDefault(FResetToDefaultOverride::Create(
 			FIsResetToDefaultVisible::CreateLambda([this, PropertyHandle_Selected, PropertyHandle_SourceFilePath, PropertyHandle_ExportFilePath](TSharedPtr<IPropertyHandle> InProperty)
-			{
-				return IsResetFilePathButtonShown(PropertyHandle_Selected, PropertyHandle_SourceFilePath, PropertyHandle_ExportFilePath);
-			}),
+				{
+					return IsResetFilePathButtonShown(PropertyHandle_Selected, PropertyHandle_SourceFilePath, PropertyHandle_ExportFilePath);
+				}),
 			FResetToDefaultHandler::CreateLambda([this, PropertyHandle_LayerInfo, PropertyHandle_Selected, PropertyHandle_SourceFilePath, PropertyHandle_ExportFilePath](TSharedPtr<IPropertyHandle> InProperty)
-			{
-				OnResetFilePathButtonClicked(PropertyHandle_LayerInfo, PropertyHandle_Selected, PropertyHandle_SourceFilePath, PropertyHandle_ExportFilePath);
-			})
+				{
+					OnResetFilePathButtonClicked(PropertyHandle_LayerInfo, PropertyHandle_Selected, PropertyHandle_SourceFilePath, PropertyHandle_ExportFilePath);
+				})
 		));
 
 	// Export file path (row hidden while in import mode)
 	IDetailPropertyRow& ExportFilePathRow = ChildBuilder.AddProperty(PropertyHandle_ExportFilePath)
-		.Visibility(TAttribute<EVisibility>::CreateLambda([PropertyHandle_Selected]() {
-			return GetImportExportVisibility(/*bImport=*/false);
-		}));
+		.Visibility(TAttribute<EVisibility>::CreateLambda([PropertyHandle_Selected]()
+			{
+				return GetImportExportVisibility(/*bImport=*/false);
+			}));
+
 	ExportFilePathRow.CustomWidget()
 		.NameContent()
 		[
@@ -262,13 +264,13 @@ void FLandscapeEditorStructCustomization_FLandscapeImportLayer::CustomizeChildre
 		]
 		.OverrideResetToDefault(FResetToDefaultOverride::Create(
 			FIsResetToDefaultVisible::CreateLambda([this, PropertyHandle_Selected, PropertyHandle_SourceFilePath, PropertyHandle_ExportFilePath](TSharedPtr<IPropertyHandle> InProperty)
-			{
-				return IsResetFilePathButtonShown(PropertyHandle_Selected, PropertyHandle_SourceFilePath, PropertyHandle_ExportFilePath);
-			}),
+				{
+					return IsResetFilePathButtonShown(PropertyHandle_Selected, PropertyHandle_SourceFilePath, PropertyHandle_ExportFilePath);
+				}),
 			FResetToDefaultHandler::CreateLambda([this, PropertyHandle_LayerInfo, PropertyHandle_Selected, PropertyHandle_SourceFilePath, PropertyHandle_ExportFilePath](TSharedPtr<IPropertyHandle> InProperty)
-			{
-				OnResetFilePathButtonClicked(PropertyHandle_LayerInfo, PropertyHandle_Selected, PropertyHandle_SourceFilePath, PropertyHandle_ExportFilePath);
-			})
+				{
+					OnResetFilePathButtonClicked(PropertyHandle_LayerInfo, PropertyHandle_Selected, PropertyHandle_SourceFilePath, PropertyHandle_ExportFilePath);
+				})
 		));
 }
 
@@ -376,25 +378,12 @@ EVisibility FLandscapeEditorStructCustomization_FLandscapeImportLayer::GetImport
 	return IsImporting() == bImport ? EVisibility::Visible : EVisibility::Collapsed;
 }
 
-TSharedRef<SWidget> FLandscapeEditorStructCustomization_FLandscapeImportLayer::OnGetImportLayerCreateMenu(TSharedRef<IPropertyHandle> PropertyHandle_LayerInfo, FName LayerName)
-{
-	FMenuBuilder MenuBuilder(true, nullptr);
-
-	MenuBuilder.AddMenuEntry(LOCTEXT("Target_Create_Blended", "Weight-Blended Layer"), FText(), FSlateIcon(),
-		FUIAction(FExecuteAction::CreateStatic(&OnImportLayerCreateClicked, PropertyHandle_LayerInfo, LayerName, false)));
-
-	MenuBuilder.AddMenuEntry(LOCTEXT("Target_Create_NoWeightBlend", "Non Weight-Blended Layer"), FText(), FSlateIcon(),
-		FUIAction(FExecuteAction::CreateStatic(&OnImportLayerCreateClicked, PropertyHandle_LayerInfo, LayerName, true)));
-
-	return MenuBuilder.MakeWidget();
-}
-
-void FLandscapeEditorStructCustomization_FLandscapeImportLayer::OnImportLayerCreateClicked(TSharedRef<IPropertyHandle> PropertyHandle_LayerInfo, FName LayerName, bool bNoWeightBlend)
+FReply FLandscapeEditorStructCustomization_FLandscapeImportLayer::OnImportLayerCreateClicked(const TSharedRef<IPropertyHandle> PropertyHandle_LayerInfo, FName LayerName)
 {
 	FEdModeLandscape* LandscapeEdMode = GetEditorMode();
 	if (!LandscapeEdMode)
 	{
-		return;
+		return FReply::Handled();
 	}
 	
 	// Build default layer object name and package name
@@ -413,7 +402,7 @@ void FLandscapeEditorStructCustomization_FLandscapeImportLayer::OnImportLayerCre
 
 		// Creates a new layer info object asset, using the default if available, or a new empty one
 		ULandscapeLayerInfoObject* LayerInfo = UE::Landscape::CreateTargetLayerInfo(LayerName, PackageName, *FileName.ToString());
-		LayerInfo->bNoWeightBlend = bNoWeightBlend;
+		LayerInfo->SetBlendMethod(GetDefault<ULandscapeSettings>()->GetTargetLayerDefaultBlendMethod(), /*bInModify = */false);
 
 		const UObject* LayerInfoAsUObject = LayerInfo; // HACK: If SetValue took a reference to a const ptr (T* const &) or a non-reference (T*) then this cast wouldn't be necessary
 		ensure(PropertyHandle_LayerInfo->SetValue(LayerInfoAsUObject) == FPropertyAccess::Success);
@@ -424,6 +413,7 @@ void FLandscapeEditorStructCustomization_FLandscapeImportLayer::OnImportLayerCre
 		GEditor->SyncBrowserToObjects(Objects);
 	}
 	
+	return FReply::Handled();
 }
 
 EVisibility FLandscapeEditorStructCustomization_FLandscapeImportLayer::GetLayerInfoAssignVisibility()
diff --git a/Engine/Source/Editor/LandscapeEditor/Private/LandscapeEditorDetailCustomization_ImportLayers.h b/Engine/Source/Editor/LandscapeEditor/Private/LandscapeEditorDetailCustomization_ImportLayers.h
--- a/Engine/Source/Editor/LandscapeEditor/Private/LandscapeEditorDetailCustomization_ImportLayers.h
+++ b/Engine/Source/Editor/LandscapeEditor/Private/LandscapeEditorDetailCustomization_ImportLayers.h
@@ -26,7 +26,7 @@ private:
 
 	static EVisibility GetImportLayerCreateVisibility(TSharedRef<IPropertyHandle> PropertyHandle_LayerInfo);
 	static TSharedRef<SWidget> OnGetImportLayerCreateMenu(TSharedRef<IPropertyHandle> PropertyHandle_LayerInfo, FName LayerName);
-	static void OnImportLayerCreateClicked(TSharedRef<IPropertyHandle> PropertyHandle_LayerInfo, FName LayerName, bool bNoWeightBlend);
+	static FReply OnImportLayerCreateClicked(const TSharedRef<IPropertyHandle> PropertyHandle_LayerInfo, FName LayerName);
 	static EVisibility GetImportLayerSelectionVisibility();
 	static EVisibility GetImportExportVisibility(bool bImport);
 
diff --git a/Engine/Source/Editor/LandscapeEditor/Private/LandscapeEditorDetailCustomization_TargetLayers.cpp b/Engine/Source/Editor/LandscapeEditor/Private/LandscapeEditorDetailCustomization_TargetLayers.cpp
--- a/Engine/Source/Editor/LandscapeEditor/Private/LandscapeEditorDetailCustomization_TargetLayers.cpp
+++ b/Engine/Source/Editor/LandscapeEditor/Private/LandscapeEditorDetailCustomization_TargetLayers.cpp
@@ -35,6 +35,7 @@
 #include "DetailCategoryBuilder.h"
 #include "PropertyCustomizationHelpers.h"
 #include "Subsystems/AssetEditorSubsystem.h"
+#include "Algo/AnyOf.h"
 #include "Algo/Count.h"
 #include "AssetRegistry/AssetRegistryModule.h"
 #include "Materials/MaterialExpressionLandscapeVisibilityMask.h"
@@ -259,6 +260,19 @@ void FLandscapeEditorDetailCustomization_TargetLayers::CustomizeDetails(IDetailL
 				.Text(LOCTEXT("FilteredTargetLayers_Tip", "All target layers assigned to this landscape are currently filtered. Use the buttons and/or the filter above to un-hide them."))
 				.AutoWrapText(true)
 		];
+
+	TargetsCategory.AddCustomRow(FText())
+		.Visibility(TAttribute<EVisibility>::Create(TAttribute<EVisibility>::FGetter::CreateStatic(&FLandscapeEditorDetailCustomization_TargetLayers::GetAdvancedWeightBlendingTipVisibility)))
+		[
+			SNew(SMultiLineEditableTextBox)
+				.IsReadOnly(true)
+				.Font(DetailBuilder.GetDetailFontBold())
+				.BackgroundColor(FAppStyle::GetColor("ErrorReporting.BackgroundColor"))
+				.Text(LOCTEXT("AdvancedWeightBlending_Tip", "Some target layers assigned to this landscape are currently using the Advanced Weight Blending the blend method, while edit layers are "
+					"not active or the landscape edit layers merge technique is not set to \"batched merge\". This blend technique is only available when edit layers are active and when batched "
+					"merge is in use (console variable : landscape.EditLayersLocalMerge.Enable 2). Please enable this, otherwise those layers will be treated as Legacy Weight Blending."))
+				.AutoWrapText(true)
+		];
 }
 END_SLATE_FUNCTION_BUILD_OPTIMIZATION
 
@@ -388,6 +402,28 @@ EVisibility FLandscapeEditorDetailCustomization_TargetLayers::GetFilteredTargetL
 	return EVisibility::Collapsed;
 }
 
+// TODO [jonathan.bard] : remove entirely once batched merge is phased out
+EVisibility FLandscapeEditorDetailCustomization_TargetLayers::GetAdvancedWeightBlendingTipVisibility()
+{
+	if (FEdModeLandscape* LandscapeEdMode = GetEditorMode(); (LandscapeEdMode != nullptr) && LandscapeEdMode->CurrentToolTarget.LandscapeInfo.IsValid())
+	{
+		if (ALandscape* LandscapeActor = LandscapeEdMode->CurrentToolTarget.LandscapeInfo->LandscapeActor.Get(); (LandscapeActor != nullptr) && (LandscapeEdMode->CurrentToolTarget.TargetType == ELandscapeToolTargetType::Weightmap))
+		{
+			const TArray<TSharedRef<FLandscapeTargetListInfo>> TargetDisplayList = FLandscapeEditorCustomNodeBuilder_TargetLayers::PrepareTargetLayerList(/*bInSort =*/ false, /*bInFilter = */false);
+			const bool bAtLeastOneLayerUsesPremultipliedAlphaBlending = Algo::AnyOf(TargetDisplayList, [](const TSharedRef<FLandscapeTargetListInfo>& InTargetListInfo)
+				{
+					return (InTargetListInfo->LayerInfoObj != nullptr) && (InTargetListInfo->LayerInfoObj->GetBlendMethod() == ELandscapeTargetLayerBlendMethod::PremultipliedAlphaBlending);
+				});
+			if (bAtLeastOneLayerUsesPremultipliedAlphaBlending && (LandscapeActor->GetEditLayersMergeMode() != ELandscapeEditLayersMergeMode::BatchedMerge))
+			{
+				return EVisibility::Visible;
+			}
+		}
+	}
+
+	return EVisibility::Collapsed;
+}
+
 
 //////////////////////////////////////////////////////////////////////////
 
@@ -718,7 +754,7 @@ FReply FLandscapeEditorCustomNodeBuilder_TargetLayers::HandleAutoFillTargetLayer
 				FName FileName;
 				const FString PackageName = UE::Landscape::GetLayerInfoObjectPackageName(TargetInfo->LayerName, TargetLayerAssetFilePath, FileName);
 
-				CreateTargetLayerInfoAsset(TargetInfo, TargetLayerAssetFilePath, FileName.ToString(), /* bNoWeightBlend */ true);
+				CreateTargetLayerInfoAsset(TargetInfo, TargetLayerAssetFilePath, FileName.ToString());
 			}
 		}
 	}
@@ -793,6 +829,19 @@ TSharedRef<SWidget> FLandscapeEditorCustomNodeBuilder_TargetLayers::GetTargetLay
 			EUserInterfaceActionType::RadioButton
 		);
 
+		MenuBuilder.AddMenuEntry(
+			LOCTEXT("TargetLayerDisplayOrderBlendType", "By Blend Method"),
+			LOCTEXT("TargetLayerDisplayOrderBlendTypeToolTip", "Sort using the blending method."),
+			FSlateIcon(),
+			FUIAction(
+				FExecuteAction::CreateSP(this, &FLandscapeEditorCustomNodeBuilder_TargetLayers::SetSelectedDisplayOrder, ELandscapeLayerDisplayMode::ByBlendMethod),
+				FCanExecuteAction(),
+				FIsActionChecked::CreateSP(this, &FLandscapeEditorCustomNodeBuilder_TargetLayers::IsSelectedDisplayOrder, ELandscapeLayerDisplayMode::ByBlendMethod)
+			),
+			NAME_None,
+			EUserInterfaceActionType::RadioButton
+		);
+
 		MenuBuilder.AddMenuEntry(
 			LOCTEXT("TargetLayerDisplayOrderCustom", "Custom"),
 			LOCTEXT("TargetLayerDisplayOrderCustomToolTip", "This sort options will be set when changing manually display order by dragging layers"),
@@ -899,6 +948,7 @@ const FSlateBrush* FLandscapeEditorCustomNodeBuilder_TargetLayers::GetTargetLaye
 		{
 			case ELandscapeLayerDisplayMode::Alphabetical: return FAppStyle::Get().GetBrush("LandscapeEditor.Target_DisplayOrder.Alphabetical");
 			case ELandscapeLayerDisplayMode::UserSpecific: return FAppStyle::Get().GetBrush("LandscapeEditor.Target_DisplayOrder.Custom");
+			case ELandscapeLayerDisplayMode::ByBlendMethod: return FAppStyle::Get().GetBrush("LandscapeEditor.Target_DisplayOrder.ByBlendMethod");
 			default:
 				return nullptr;
 		}
@@ -1185,18 +1235,18 @@ TSharedPtr<SWidget> FLandscapeEditorCustomNodeBuilder_TargetLayers::GenerateRow(
 					.Padding(FMargin(2))
 					[
 						SNew(SLandscapeAssetThumbnail, Target->bIsLayerReferencedByMaterial ? Target->ThumbnailMIC.Get() : nullptr, ThumbnailPool, Target->bIsLayerReferencedByMaterial ? FName("LandscapeEditor.Target_Weightmap") : FName("LandscapeEditor.Target_Unknown"))
-						.Visibility_Static(&FLandscapeEditorCustomNodeBuilder_TargetLayers::GetDebugModeLayerUsageVisibility_Invert, Target)
-						.ThumbnailSize(FIntPoint(48, 48))
-						// Open landscape layer info asset on double-click on the thumbnail : 
-						.OnAccessAsset_Lambda([Target](UObject* InObject)
-						{ 
-							// Note : the object being returned here is the landscape MIC so it's not what we use for opening the landscape layer info asset : 
-							if ((Target->TargetType == ELandscapeToolTargetType::Weightmap) && (Target->LayerInfoObj != nullptr))
-							{
-								UAssetEditorSubsystem* AssetEditorSubsystem = GEditor->GetEditorSubsystem<UAssetEditorSubsystem>();
-								return AssetEditorSubsystem->OpenEditorForAsset(Target->LayerInfoObj.Get());
-							}
-							return false;
+								.Visibility_Static(&FLandscapeEditorCustomNodeBuilder_TargetLayers::GetDebugModeLayerUsageVisibility_Invert, Target)
+								.ThumbnailSize(FIntPoint(48, 48))
+								// Open landscape layer info asset on double-click on the thumbnail : 
+								.OnAccessAsset_Lambda([Target](UObject* InObject)
+								{ 
+									// Note : the object being returned here is the landscape MIC so it's not what we use for opening the landscape layer info asset : 
+									if ((Target->TargetType == ELandscapeToolTargetType::Weightmap) && (Target->LayerInfoObj != nullptr))
+									{
+										UAssetEditorSubsystem* AssetEditorSubsystem = GEditor->GetEditorSubsystem<UAssetEditorSubsystem>();
+										return AssetEditorSubsystem->OpenEditorForAsset(Target->LayerInfoObj.Get());
+									}
+									return false;
 						})
 					]
 					+ SHorizontalBox::Slot()
@@ -1277,9 +1327,9 @@ TSharedPtr<SWidget> FLandscapeEditorCustomNodeBuilder_TargetLayers::GenerateRow(
 									.HAlign(HAlign_Right)
 									[
 										SNew(STextBlock)
-											.Visibility_Lambda([=] { return (Target->LayerInfoObj.IsValid() && Target->LayerInfoObj->bNoWeightBlend) ? EVisibility::Visible : EVisibility::Collapsed; })
 											.Font(IDetailLayoutBuilder::GetDetailFont())
-											.Text(LOCTEXT("NoWeightBlend", "No Weight-Blend"))
+											.Text_Static(&FLandscapeEditorCustomNodeBuilder_TargetLayers::GetTargetBlendMethodText, Target)
+											.ToolTipText_Static(&FLandscapeEditorCustomNodeBuilder_TargetLayers::GetTargetBlendMethodTooltipText, Target)
 											.ColorAndOpacity_Static(&FLandscapeEditorCustomNodeBuilder_TargetLayers::GetTargetTextColor, Target)
 									]
 							]
@@ -1304,16 +1354,14 @@ TSharedPtr<SWidget> FLandscapeEditorCustomNodeBuilder_TargetLayers::GenerateRow(
 									.AutoWidth()
 									.VAlign(VAlign_Center)
 									[
-										SNew(SComboButton)
+										SNew(SButton)
 											.ButtonStyle(FAppStyle::Get(), "HoverHintOnly")
-											.HasDownArrow(false)
 											.ContentPadding(4.0f)
 											.ForegroundColor(FSlateColor::UseForeground())
 											.IsFocusable(false)
 											.ToolTipText(LOCTEXT("Tooltip_Create", "Create Layer Info"))
 											.IsEnabled_Static(&FLandscapeEditorCustomNodeBuilder_TargetLayers::GetTargetLayerCreateEnabled, Target)
-											.OnGetMenuContent_Static(&FLandscapeEditorCustomNodeBuilder_TargetLayers::OnGetTargetLayerCreateMenu, Target)
-											.ButtonContent()
+											.OnClicked_Static(&FLandscapeEditorCustomNodeBuilder_TargetLayers::OnTargetLayerCreateClicked, Target)
 											[
 												SNew(SImage)
 													.Image(FAppStyle::GetBrush("LandscapeEditor.Target_Create"))
@@ -1860,8 +1908,8 @@ void FLandscapeEditorCustomNodeBuilder_TargetLayers::OnTargetLayerSetObject(cons
 			LandscapeEdMode->SetCurrentTargetLayer(Target->LayerName, nullptr);
 			LandscapeInfo->UpdateLayerInfoMap();
 			LandscapeEdMode->UpdateTargetList();
-			return;
-		}
+		return;
+	}
 	}
 
 	ULandscapeLayerInfoObject* SelectedLayerInfo = const_cast<ULandscapeLayerInfoObject*>(CastChecked<ULandscapeLayerInfoObject>(Object));
@@ -1936,20 +1984,7 @@ bool FLandscapeEditorCustomNodeBuilder_TargetLayers::GetTargetLayerCreateEnabled
 	return false;
 }
 
-TSharedRef<SWidget> FLandscapeEditorCustomNodeBuilder_TargetLayers::OnGetTargetLayerCreateMenu(const TSharedRef<FLandscapeTargetListInfo> Target)
-{
-	FMenuBuilder MenuBuilder(true, nullptr);
-
-	MenuBuilder.AddMenuEntry(LOCTEXT("Menu_Create_Blended", "Weight-Blended Layer"), FText(), FSlateIcon(),
-		FUIAction(FExecuteAction::CreateStatic(&FLandscapeEditorCustomNodeBuilder_TargetLayers::OnTargetLayerCreateClicked, Target, false)));
-
-	MenuBuilder.AddMenuEntry(LOCTEXT("Menu_Create_NoWeightBlend", "Non Weight-Blended Layer"), FText(), FSlateIcon(),
-		FUIAction(FExecuteAction::CreateStatic(&FLandscapeEditorCustomNodeBuilder_TargetLayers::OnTargetLayerCreateClicked, Target, true)));
-
-	return MenuBuilder.MakeWidget();
-}
-
-void FLandscapeEditorCustomNodeBuilder_TargetLayers::OnTargetLayerCreateClicked(const TSharedRef<FLandscapeTargetListInfo> Target, bool bNoWeightBlend)
+FReply FLandscapeEditorCustomNodeBuilder_TargetLayers::OnTargetLayerCreateClicked(const TSharedRef<FLandscapeTargetListInfo> Target)
 {
 	check(!Target->LayerInfoObj.IsValid());
 
@@ -1971,19 +2006,20 @@ void FLandscapeEditorCustomNodeBuilder_TargetLayers::OnTargetLayerCreateClicked(
 			PackageName = NewLayerDlg->GetAssetPath().ToString();
 			FileName = FName(*NewLayerDlg->GetAssetName().ToString());
 
-			CreateTargetLayerInfoAsset(Target, PackageName, FileName.ToString(), bNoWeightBlend);
+			CreateTargetLayerInfoAsset(Target, PackageName, FileName.ToString());
 		}
 	}
+
+	return FReply::Handled();
 }
 
-void FLandscapeEditorCustomNodeBuilder_TargetLayers::CreateTargetLayerInfoAsset(const TSharedRef<FLandscapeTargetListInfo> Target, const FString& PackageName, const FString& FileName, const bool bNoWeightBlend)
+void FLandscapeEditorCustomNodeBuilder_TargetLayers::CreateTargetLayerInfoAsset(const TSharedRef<FLandscapeTargetListInfo> Target, const FString& PackageName, const FString& FileName)
 {
 	if (FEdModeLandscape* LandscapeEdMode = GetEditorMode())
 	{
 		const FName LayerName = Target->GetLayerName();
 
 		ULandscapeLayerInfoObject* LayerInfo = UE::Landscape::CreateTargetLayerInfo(LayerName, PackageName, FileName);
-		LayerInfo->bNoWeightBlend = bNoWeightBlend;
 
 		if (LandscapeEdMode->CurrentToolTarget.LayerName == Target->LayerName
 			&& LandscapeEdMode->CurrentToolTarget.LayerInfo == Target->LayerInfoObj)
@@ -2063,9 +2099,9 @@ EVisibility FLandscapeEditorCustomNodeBuilder_TargetLayers::GetLayersSubstractiv
 	if (LandscapeEdMode != nullptr && Target->TargetType == ELandscapeToolTargetType::Weightmap && Target->LayerInfoObj.IsValid())
 	{
 		if (ULandscapeEditLayerBase* EditLayer = LandscapeEdMode->GetEditLayer(LandscapeEdMode->GetSelectedEditLayerIndex()))
-		{
-			return EVisibility::Visible;
-		}
+	{
+		return EVisibility::Visible;
+	}
 	}
 
 	return EVisibility::Collapsed;
@@ -2081,7 +2117,7 @@ ECheckBoxState FLandscapeEditorCustomNodeBuilder_TargetLayers::IsLayersSubstract
 			const bool* AllocationBlend = EditLayer->GetWeightmapLayerAllocationBlend().Find(Target.Get().LayerInfoObj.Get());
 
 			if (AllocationBlend != nullptr)
-			{
+	{
 				return (*AllocationBlend) ? ECheckBoxState::Checked : ECheckBoxState::Unchecked;
 			}
 		}
@@ -2155,6 +2191,37 @@ void FLandscapeEditorCustomNodeBuilder_TargetLayers::OnDebugModeColorChannelChan
 	}
 }
 
+FText FLandscapeEditorCustomNodeBuilder_TargetLayers::GetTargetBlendMethodText(const TSharedRef<FLandscapeTargetListInfo> InTarget)
+{
+	const ULandscapeLayerInfoObject* LayerInfo = InTarget->LayerInfoObj.Get();
+	if (LayerInfo == nullptr)
+	{
+		return FText::GetEmpty();
+	}
+
+	const UEnum* ModeEnum = StaticEnum<ELandscapeTargetLayerBlendMethod>();
+	check(ModeEnum != nullptr);
+	FText Result = ModeEnum->GetDisplayNameTextByValue(static_cast<int64>(LayerInfo->GetBlendMethod()));
+	if ((LayerInfo->GetBlendMethod() == ELandscapeTargetLayerBlendMethod::PremultipliedAlphaBlending) && !LayerInfo->GetBlendGroup().IsNone())
+	{
+		Result = FText::Format(LOCTEXT("PremultipliedAlphaBlendingBlendGroupName", "{0} ({1})"), Result, FText::FromName(LayerInfo->GetBlendGroup()));
+	}
+	return Result;
+}
+
+FText FLandscapeEditorCustomNodeBuilder_TargetLayers::GetTargetBlendMethodTooltipText(const TSharedRef<FLandscapeTargetListInfo> InTarget)
+{
+	const ULandscapeLayerInfoObject* LayerInfo = InTarget->LayerInfoObj.Get();
+	if (LayerInfo == nullptr)
+	{
+		return FText::GetEmpty();
+	}
+
+	const UEnum* ModeEnum = StaticEnum<ELandscapeTargetLayerBlendMethod>();
+	check(ModeEnum != nullptr);
+	return ModeEnum->GetToolTipTextByIndex(ModeEnum->GetIndexByValue(static_cast<int64>(LayerInfo->GetBlendMethod())));
+}
+
 FSlateColor FLandscapeEditorCustomNodeBuilder_TargetLayers::GetTargetTextColor(const TSharedRef<FLandscapeTargetListInfo> InTarget)
 {
 	return FLandscapeEditorCustomNodeBuilder_TargetLayers::GetTargetLayerIsSelected(InTarget) ? FStyleColors::ForegroundHover : FSlateColor::UseForeground();
diff --git a/Engine/Source/Editor/LandscapeEditor/Private/LandscapeEditorDetailCustomization_TargetLayers.h b/Engine/Source/Editor/LandscapeEditor/Private/LandscapeEditorDetailCustomization_TargetLayers.h
--- a/Engine/Source/Editor/LandscapeEditor/Private/LandscapeEditorDetailCustomization_TargetLayers.h
+++ b/Engine/Source/Editor/LandscapeEditor/Private/LandscapeEditorDetailCustomization_TargetLayers.h
@@ -60,6 +60,7 @@ protected:
 	static EVisibility GetPopulateTargetLayersInfoTipVisibility();
 	static EVisibility GetTargetLayersInvalidInfoAssetTipVisibility();
 	static EVisibility GetFilteredTargetLayersListInfoTipVisibility();
+	static EVisibility GetAdvancedWeightBlendingTipVisibility();
 };
 
 class FLandscapeEditorCustomNodeBuilder_TargetLayers : public IDetailCustomNodeBuilder, public TSharedFromThis<FLandscapeEditorCustomNodeBuilder_TargetLayers>
@@ -107,9 +108,8 @@ protected:
 	static void OnTargetLayerSetObject(const FAssetData& AssetData, const TSharedRef<FLandscapeTargetListInfo> Target); 
 	static EVisibility GetTargetLayerInfoSelectorVisibility(const TSharedRef<FLandscapeTargetListInfo> Target);
 	static bool GetTargetLayerCreateEnabled(const TSharedRef<FLandscapeTargetListInfo> Target);
-	static TSharedRef<SWidget> OnGetTargetLayerCreateMenu(const TSharedRef<FLandscapeTargetListInfo> Target);
-	static void OnTargetLayerCreateClicked(const TSharedRef<FLandscapeTargetListInfo> Target, bool bNoWeightBlend);
-	static void CreateTargetLayerInfoAsset(const TSharedRef<FLandscapeTargetListInfo> Target, const FString& PackageName, const FString& FileName, const bool bNoWeightBlend);
+	static FReply OnTargetLayerCreateClicked(const TSharedRef<FLandscapeTargetListInfo> Target);
+	static void CreateTargetLayerInfoAsset(const TSharedRef<FLandscapeTargetListInfo> Target, const FString& PackageName, const FString& FileName);
 	static FReply OnTargetLayerDeleteClicked(const TSharedRef<FLandscapeTargetListInfo> Target);
 	static FSlateColor GetLayerUsageDebugColor(const TSharedRef<FLandscapeTargetListInfo> Target);
 	static EVisibility GetDebugModeLayerUsageVisibility(const TSharedRef<FLandscapeTargetListInfo> Target);
@@ -117,6 +117,8 @@ protected:
 	static EVisibility GetDebugModeColorChannelVisibility(const TSharedRef<FLandscapeTargetListInfo> Target);
 	static ECheckBoxState DebugModeColorChannelIsChecked(const TSharedRef<FLandscapeTargetListInfo> Target, int32 Channel);
 	static void OnDebugModeColorChannelChanged(ECheckBoxState NewCheckedState, const TSharedRef<FLandscapeTargetListInfo> Target, int32 Channel);
+	static FText GetTargetBlendMethodText(const TSharedRef<FLandscapeTargetListInfo> InTarget);
+	static FText GetTargetBlendMethodTooltipText(const TSharedRef<FLandscapeTargetListInfo> InTarget);
 	static FSlateColor GetTargetTextColor(const TSharedRef<FLandscapeTargetListInfo> InTarget);
 
 	static EVisibility GetLayersSubstractiveBlendVisibility(const TSharedRef<FLandscapeTargetListInfo> Target);
diff --git a/Engine/Source/Editor/LandscapeEditor/Private/LandscapeEditorObject.cpp b/Engine/Source/Editor/LandscapeEditor/Private/LandscapeEditorObject.cpp
--- a/Engine/Source/Editor/LandscapeEditor/Private/LandscapeEditorObject.cpp
+++ b/Engine/Source/Editor/LandscapeEditor/Private/LandscapeEditorObject.cpp
@@ -423,7 +423,7 @@ void ULandscapeEditorObject::Save()
 	GConfig->SetBool(TEXT("LandscapeEdit"), TEXT("bCombinedLayersOperation"), bCombinedLayersOperation, GEditorPerProjectIni);
 	GConfig->SetBool(TEXT("LandscapeEdit"), TEXT("bApplyWithoutMovingSculpt"), bApplyWithoutMovingSculpt, GEditorPerProjectIni);
 	GConfig->SetBool(TEXT("LandscapeEdit"), TEXT("bApplyWithoutMovingPaint"), bApplyWithoutMovingPaint, GEditorPerProjectIni);
-	
+
 	GConfig->SetInt(TEXT("LandscapeEdit"), TEXT("ErodeThresh"), ErodeThresh, GEditorPerProjectIni);
 	GConfig->SetInt(TEXT("LandscapeEdit"), TEXT("ErodeIterationNum"), ErodeIterationNum, GEditorPerProjectIni);
 	GConfig->SetInt(TEXT("LandscapeEdit"), TEXT("ErodeSurfaceThickness"), ErodeSurfaceThickness, GEditorPerProjectIni);
@@ -819,7 +819,9 @@ ELandscapeImportResult ULandscapeEditorObject::CreateNewLayersInfo(TArray<FLands
 	}
 
 	// Fill the first weight-blended layer to 100%
-	if (FLandscapeImportLayerInfo* FirstBlendedLayer = OutNewLayerInfos.FindByPredicate([](const FLandscapeImportLayerInfo& ImportLayer) { return ImportLayer.LayerInfo && !ImportLayer.LayerInfo->bNoWeightBlend; }))
+	// TODO [jonathan.bard] : remove all this once we get rid of non-edit layers, weight balancing at the data level is not reliable and weight blending with edit layers active 
+	//  is providing the same service in a non-destructive way
+	if (FLandscapeImportLayerInfo* FirstBlendedLayer = OutNewLayerInfos.FindByPredicate([](const FLandscapeImportLayerInfo& ImportLayer) { return ImportLayer.LayerInfo && (ImportLayer.LayerInfo->GetBlendMethod() == ELandscapeTargetLayerBlendMethod::FinalWeightBlending); }))
 	{
 		const int32 DataSize = SizeX * SizeY;
 		FirstBlendedLayer->LayerData.AddUninitialized(DataSize);
@@ -970,7 +972,6 @@ void ULandscapeEditorObject::AutoFillTargetLayers(const FString& DefaultTargetLa
 		{
 			// Creates a new layer info object, using the default if available, or a new empty one
 			ULandscapeLayerInfoObject* LayerInfo = UE::Landscape::CreateTargetLayerInfo(NewLayer.LayerName, DefaultTargetLayerAssetPath);
-			LayerInfo->bNoWeightBlend = true;
 
 			NewLayer.LayerInfo = LayerInfo;
 
diff --git a/Engine/Source/Editor/WorldBrowser/Private/Tiles/WorldTileCollectionModel.cpp b/Engine/Source/Editor/WorldBrowser/Private/Tiles/WorldTileCollectionModel.cpp
--- a/Engine/Source/Editor/WorldBrowser/Private/Tiles/WorldTileCollectionModel.cpp
+++ b/Engine/Source/Editor/WorldBrowser/Private/Tiles/WorldTileCollectionModel.cpp
@@ -144,10 +144,10 @@ void FWorldTileCollectionModel::UnloadLevels(const FLevelModelList& InLevelList)
 
 		// We need to unload all children of an dirty tiles,
 		// to make sure that relative positions will be correct after parent tile information is discarded
-		struct FHierachyCollector : public FLevelModelVisitor
+		struct FHierarchyCollector : public FLevelModelVisitor
 		{
 			FLevelModelList& DirtyHierarchy;
-			FHierachyCollector(FLevelModelList& LevelToUnload) : DirtyHierarchy(LevelToUnload) {}
+			FHierarchyCollector(FLevelModelList& LevelToUnload) : DirtyHierarchy(LevelToUnload) {}
 
 			virtual void Visit(FLevelModel& Item) override
 			{
@@ -155,14 +155,14 @@ void FWorldTileCollectionModel::UnloadLevels(const FLevelModelList& InLevelList)
 			}
 		};
 
-		FHierachyCollector HierachyCollector(LevelsToUnload);
+		FHierarchyCollector HierarchyCollector(LevelsToUnload);
 
 		for (TSharedPtr<FLevelModel> LevelModel : InLevelList)
 		{
 			ULevel* Level = LevelModel->GetLevelObject();
 			if (Level && Level->GetOutermost()->IsDirty())
 			{
-				LevelModel->Accept(HierachyCollector);
+				LevelModel->Accept(HierarchyCollector);
 			}
 		}
 	}
@@ -1590,8 +1590,7 @@ static void SetupLandscapeImportLayers(const FTiledLandscapeImportSettings& InIm
 		}
 
 		LayerImportInfo.LayerInfo = GetLandscapeLayerInfoObject(LayerImportInfo.LayerName, ContentPath);
-		LayerImportInfo.LayerInfo->Modify();
-		LayerImportInfo.LayerInfo->bNoWeightBlend = LayerSettings.bNoBlendWeight;
+		LayerImportInfo.LayerInfo->SetBlendMethod(LayerSettings.bNoBlendWeight ? ELandscapeTargetLayerBlendMethod::None : ELandscapeTargetLayerBlendMethod::FinalWeightBlending, /*bInModify = */true);
 
 		OutLayerInfo.Add(LayerImportInfo);
 	}
diff --git a/Engine/Source/Runtime/Core/Public/UObject/FortniteMainBranchObjectVersions.inl b/Engine/Source/Runtime/Core/Public/UObject/FortniteMainBranchObjectVersions.inl
--- a/Engine/Source/Runtime/Core/Public/UObject/FortniteMainBranchObjectVersions.inl
+++ b/Engine/Source/Runtime/Core/Public/UObject/FortniteMainBranchObjectVersions.inl
@@ -420,7 +420,7 @@ FORTNITE_MAIN_VERSION(AnimNextCombineGraphContexts, 135)
 // Combine parameter blocks and graphs
 FORTNITE_MAIN_VERSION(AnimNextCombineParameterBlocksAndGraphs, 136)
 
-// Move workspaces to a seperate plugin
+// Move workspaces to a separate plugin
 FORTNITE_MAIN_VERSION(AnimNextMoveWorkspaces, 137)
 
 // Level Instance Property overrides
@@ -684,5 +684,9 @@ FORTNITE_MAIN_VERSION(MigrateLandscapeNonEditLayerToEditLayer, 223)
 
 // FDynamicMeshAttributeSet has Morph Targets.
 FORTNITE_MAIN_VERSION(DynamicMeshAttributesMorphTargets, 224)
+
+// Introduce landscape advanced weight blending 
+FORTNITE_MAIN_VERSION(LandscapeAdvancedWeightBlending, 225)
+
 // -----<new versions can be added above this line>-------------------------------------------------
 #undef FORTNITE_MAIN_VERSION
\ No newline at end of file
diff --git a/Engine/Source/Runtime/Landscape/Classes/Landscape.h b/Engine/Source/Runtime/Landscape/Classes/Landscape.h
--- a/Engine/Source/Runtime/Landscape/Classes/Landscape.h
+++ b/Engine/Source/Runtime/Landscape/Classes/Landscape.h
@@ -577,6 +577,8 @@ PRAGMA_ENABLE_DEPRECATION_WARNINGS
 	bool GetIsPendingEditLayerConversion() const;
 	void SetIsPendingEditLayerConversion(bool bInPendingEditLayerConversion);
 
+	UE_INTERNAL LANDSCAPE_API ELandscapeEditLayersMergeMode GetEditLayersMergeMode() const;
+
 private:
 	FLandscapeLayer* GetLayerInternal(int32 InLayerIndex);
 	void OnLayerCreatedInternal(ULandscapeEditLayerBase* EditLayer);
@@ -584,7 +586,6 @@ private:
 
 	bool SupportsEditLayersLocalMerge();
 
-	ELandscapeEditLayersMergeMode GetEditLayersMergeMode() const;
 	bool CreateLayersRenderingResource(bool bInForceLayerResourceReset);
 	void PrepareEditLayersLocalMergeResources();
 	void UpdateLayersContent(bool bInWaitForStreaming = false, bool bInSkipMonitorLandscapeEdModeChanges = false, bool bIntermediateRender = false, bool bFlushRender = false);
diff --git a/Engine/Source/Runtime/Landscape/Classes/LandscapeEditLayer.h b/Engine/Source/Runtime/Landscape/Classes/LandscapeEditLayer.h
--- a/Engine/Source/Runtime/Landscape/Classes/LandscapeEditLayer.h
+++ b/Engine/Source/Runtime/Landscape/Classes/LandscapeEditLayer.h
@@ -407,7 +407,7 @@ public:
 	//~ Begin ILandscapeEditLayerRenderer implementation
 	LANDSCAPE_EDIT_LAYERS_BATCHED_MERGE_EXPERIMENTAL
 	LANDSCAPE_API virtual void GetRendererStateInfo(const UE::Landscape::EditLayers::FMergeContext* InMergeContext,
-		UE::Landscape::EditLayers::FEditLayerTargetTypeState& OutSupportedTargetTypeState, UE::Landscape::EditLayers::FEditLayerTargetTypeState& OutEnabledTargetTypeState, TArray<TBitArray<>>& OutTargetLayerGroups) const override;
+		UE::Landscape::EditLayers::FEditLayerTargetTypeState& OutSupportedTargetTypeState, UE::Landscape::EditLayers::FEditLayerTargetTypeState& OutEnabledTargetTypeState, TArray<UE::Landscape::EditLayers::FTargetLayerGroup>& OutTargetLayerGroups) const override;
 	LANDSCAPE_EDIT_LAYERS_BATCHED_MERGE_EXPERIMENTAL
 	virtual UE::Landscape::EditLayers::ERenderFlags GetRenderFlags(const UE::Landscape::EditLayers::FMergeContext* InMergeContext) const override;
 	LANDSCAPE_EDIT_LAYERS_BATCHED_MERGE_EXPERIMENTAL
diff --git a/Engine/Source/Runtime/Landscape/Classes/LandscapeLayerInfoObject.h b/Engine/Source/Runtime/Landscape/Classes/LandscapeLayerInfoObject.h
--- a/Engine/Source/Runtime/Landscape/Classes/LandscapeLayerInfoObject.h
+++ b/Engine/Source/Runtime/Landscape/Classes/LandscapeLayerInfoObject.h
@@ -3,6 +3,7 @@
 #pragma once
 
 #include "CoreMinimal.h"
+#include "LandscapeEditTypes.h"
 #include "UObject/ObjectMacros.h"
 #include "UObject/Object.h"
 #include "LandscapeLayerInfoObject.generated.h"
@@ -11,6 +12,9 @@ class UPhysicalMaterial;
 class UTexture2D;
 struct FPropertyChangedEvent;
 
+
+// ----------------------------------------------------------------------------------
+
 UENUM()
 enum class ESplineModulationColorMask : uint8
 {
@@ -20,6 +24,9 @@ enum class ESplineModulationColorMask : uint8
 	Alpha
 };
 
+
+// ----------------------------------------------------------------------------------
+
 UCLASS(MinimalAPI, BlueprintType)
 class ULandscapeLayerInfoObject : public UObject
 {
@@ -32,37 +39,60 @@ class ULandscapeLayerInfoObject : public UObject
 	TObjectPtr<UPhysicalMaterial> PhysMaterial;
 
 	UPROPERTY(EditAnywhere, Category = LandscapeLayerInfoObject, Meta = (ClampMin = "0", ClampMax = "1", Tooltip = "Defines how much 'resistance' areas painted with this layer will offer to the Erosion tool. A hardness of 0 means the layer is fully affected by erosion, while 1 means fully unaffected."))
-	float Hardness;
+	float Hardness = 0.0f;
 
 #if WITH_EDITORONLY_DATA
 	UPROPERTY(EditAnywhere, Category = LandscapeLayerInfoObject, Meta = (ClampMin = "0", ClampMax = "1", Tooltip = "The minimum weight that needs to be painted for that layer to be picked up as the dominant physical layer."))
-	float MinimumCollisionRelevanceWeight;
+	float MinimumCollisionRelevanceWeight = 0.0f;
 
-	UPROPERTY(EditAnywhere, Category = LandscapeLayerInfoObject, Meta = (Tooltip = "Prevents this layer to be weight-blended with others."))
-	uint32 bNoWeightBlend:1;
+	UE_DEPRECATED(5.7, "bNoWeightBlend has been replaced by BlendMethod (false is ELandscapeTargetLayerBlendMethod::FinalWeightBlending, true is ELandscapeTargetLayerBlendMethod::None)")
+	UPROPERTY()
+	uint32 bNoWeightBlend_DEPRECATED:1;
 
 	UPROPERTY(EditAnywhere, Category = SplineFalloffModulation, Meta = (DisplayName="Texture", Tooltip = "Texture to modulate the Splines Falloff Layer Alpha"))
 	TObjectPtr<UTexture2D> SplineFalloffModulationTexture;
 
 	UPROPERTY(EditAnywhere, Category = SplineFalloffModulation, Meta = (DisplayName = "Color Mask", Tooltip = "Defines which channel of the Spline Falloff Modulation Texture to use."))
-	ESplineModulationColorMask SplineFalloffModulationColorMask;
+	ESplineModulationColorMask SplineFalloffModulationColorMask = ESplineModulationColorMask::Red;
 
 	UPROPERTY(EditAnywhere, Category = SplineFalloffModulation, Meta = (DisplayName = "Tiling", ClampMin = "0.01", Tooltip = "Defines the tiling to use when sampling the Spline Falloff Modulation Texture."))
-	float SplineFalloffModulationTiling;
+	float SplineFalloffModulationTiling = 1.0f;
 
 	UPROPERTY(EditAnywhere, Category = SplineFalloffModulation, Meta = (DisplayName = "Bias", ClampMin = "0", Tooltip = "Defines the offset to use when sampling the Spline Falloff Modulation Texture."))
-	float SplineFalloffModulationBias;
+	float SplineFalloffModulationBias = 0.5f;
 
 	UPROPERTY(EditAnywhere, Category = SplineFalloffModulation, Meta = (DisplayName = "Scale", ClampMin = "0", Tooltip = "Allows to scale the value sampled from the Spline Falloff Modulation Texture."))
-	float SplineFalloffModulationScale;
+	float SplineFalloffModulationScale = 1.0f;
 				
 	UPROPERTY(NonTransactional, Transient)
-	bool IsReferencedFromLoadedData;
+	bool IsReferencedFromLoadedData = false;
 #endif // WITH_EDITORONLY_DATA
 
 	UPROPERTY(EditAnywhere, Category = LandscapeLayerInfoObject, Meta = (Tooltip = "The color to use for layer usage debug"))
 	FLinearColor LayerUsageDebugColor;
 
+private:
+	UPROPERTY(EditAnywhere, Category = LandscapeLayerInfoObject, Meta = (Tooltip = "Allows this layer's final weight to be adjusted against others."))
+	ELandscapeTargetLayerBlendMethod BlendMethod = ELandscapeTargetLayerBlendMethod::None;
+
+	UPROPERTY(EditAnywhere, Category = LandscapeLayerInfoObject, Meta = (EditCondition = "BlendMethod == ELandscapeTargetLayerBlendMethod::PremultipliedAlphaBlending", Tooltip = "Only available for Advanced Weight Blending. Allows target layers from the same group only to have their weight adjusted against one another."))
+	FName BlendGroup;
+
+public:
+	inline ELandscapeTargetLayerBlendMethod GetBlendMethod() const
+	{
+		return BlendMethod;
+	}
+	LANDSCAPE_API void SetBlendMethod(ELandscapeTargetLayerBlendMethod InBlendMethod, bool bInModify);
+
+	inline FName GetBlendGroup() const
+	{
+		return BlendGroup;
+	}
+	LANDSCAPE_API void SetBlendGroup(const FName& InBlendGroup, bool bInModify);
+
+	virtual void Serialize(FArchive& Ar) override;
+
 #if WITH_EDITOR
 	//~ Begin UObject Interface
 	virtual void PostEditChangeProperty(FPropertyChangedEvent& PropertyChangedEvent) override;
diff --git a/Engine/Source/Runtime/Landscape/Classes/LandscapeProxy.h b/Engine/Source/Runtime/Landscape/Classes/LandscapeProxy.h
--- a/Engine/Source/Runtime/Landscape/Classes/LandscapeProxy.h
+++ b/Engine/Source/Runtime/Landscape/Classes/LandscapeProxy.h
@@ -228,6 +228,9 @@ enum class ELandscapeLayerDisplayMode : uint8
 
 	/** User specific sorting display mode */
 	UserSpecific,
+
+	/** By blend type sorting display mode */
+	ByBlendMethod,
 };
 
 UENUM()
diff --git a/Engine/Source/Runtime/Landscape/Private/Landscape.cpp b/Engine/Source/Runtime/Landscape/Private/Landscape.cpp
--- a/Engine/Source/Runtime/Landscape/Private/Landscape.cpp
+++ b/Engine/Source/Runtime/Landscape/Private/Landscape.cpp
@@ -1681,10 +1681,6 @@ ALandscapeProxy::ALandscapeProxy(const FObjectInitializer& ObjectInitializer)
 
 		VisibilityLayer = ConstructorStatics.LandscapeVisibilityLayerInfoFinder.Get();
 		check(VisibilityLayer);
-#if WITH_EDITORONLY_DATA
-		// This layer should be no weight blending
-		VisibilityLayer->bNoWeightBlend = true;
-#endif
 		VisibilityLayer->LayerName = UMaterialExpressionLandscapeVisibilityMask::ParameterName;
 		VisibilityLayer->LayerUsageDebugColor = FLinearColor(0, 0, 0, 0);
 		VisibilityLayer->AddToRoot();
@@ -3185,6 +3181,8 @@ public:
 
 	static void ModifyCompilationEnvironment(const FGlobalShaderPermutationParameters& InParameters, FShaderCompilerEnvironment& OutEnvironment)
 	{
+		UE::Landscape::ModifyShaderCompilerEnvironmentForDebug(OutEnvironment);
+
 		OutEnvironment.SetDefine(TEXT("MERGE_TEXTURE"), 1);
 	}
 
@@ -3250,6 +3248,8 @@ public:
 
 	static void ModifyCompilationEnvironment(const FGlobalShaderPermutationParameters& InParameters, FShaderCompilerEnvironment& OutEnvironment)
 	{
+		UE::Landscape::ModifyShaderCompilerEnvironmentForDebug(OutEnvironment);
+
 		OutEnvironment.SetDefine(TEXT("RESAMPLE_MERGED_TEXTURE"), 1);
 	}
 
@@ -3295,6 +3295,8 @@ public:
 
 	static void ModifyCompilationEnvironment(const FGlobalShaderPermutationParameters& InParameters, FShaderCompilerEnvironment& OutEnvironment)
 	{
+		UE::Landscape::ModifyShaderCompilerEnvironmentForDebug(OutEnvironment);
+
 		OutEnvironment.SetDefine(TEXT("PACK_RGBA_CHANNELS"), 1);
 	}
 
diff --git a/Engine/Source/Runtime/Landscape/Private/LandscapeBlueprintBrushBase.cpp b/Engine/Source/Runtime/Landscape/Private/LandscapeBlueprintBrushBase.cpp
--- a/Engine/Source/Runtime/Landscape/Private/LandscapeBlueprintBrushBase.cpp
+++ b/Engine/Source/Runtime/Landscape/Private/LandscapeBlueprintBrushBase.cpp
@@ -66,6 +66,8 @@ public:
 
 	static void ModifyCompilationEnvironment(const FGlobalShaderPermutationParameters& InParameters, FShaderCompilerEnvironment& OutEnvironment)
 	{
+		UE::Landscape::ModifyShaderCompilerEnvironmentForDebug(OutEnvironment);
+
 		OutEnvironment.SetDefine(TEXT("RESOLVE_LAYER_DATA"), 1);
 	}
 
@@ -293,7 +295,8 @@ void ALandscapeBlueprintBrushBase::SetUsePowerOfTwoRenderTarget(bool bInUsePower
 #if WITH_EDITOR
 
 void ALandscapeBlueprintBrushBase::GetRendererStateInfo(const UE::Landscape::EditLayers::FMergeContext* InMergeContext,
-	UE::Landscape::EditLayers::FEditLayerTargetTypeState& OutSupportedTargetTypeState, UE::Landscape::EditLayers::FEditLayerTargetTypeState& OutEnabledTargetTypeState, TArray<TBitArray<>>& OutTargetLayerGroups) const
+	UE::Landscape::EditLayers::FEditLayerTargetTypeState& OutSupportedTargetTypeState, UE::Landscape::EditLayers::FEditLayerTargetTypeState& OutEnabledTargetTypeState, 
+	TArray<UE::Landscape::EditLayers::FTargetLayerGroup>& OutTargetLayerGroups) const
 {
 	// What can the brush do?
 	if (CanAffectHeightmap())
@@ -348,7 +351,7 @@ TArray<UE::Landscape::EditLayers::FEditLayerRenderItem> ALandscapeBlueprintBrush
 
 	FEditLayerTargetTypeState SupportedTargetTypeState(InMergeContext);
 	FEditLayerTargetTypeState EnabledTargetTypeState(InMergeContext);
-	TArray<TBitArray<>> DummyTargetLayerGroups;
+	TArray<FTargetLayerGroup> DummyTargetLayerGroups;
 	GetRendererStateInfo(InMergeContext, SupportedTargetTypeState, EnabledTargetTypeState, DummyTargetLayerGroups);
 
 	// By default, for landscape BP brushes, we use FInputWorldArea::EType::Infinite, to indicate they can only reliably work when applied globally on the entire landscape
diff --git a/Engine/Source/Runtime/Landscape/Private/LandscapeEdit.cpp b/Engine/Source/Runtime/Landscape/Private/LandscapeEdit.cpp
--- a/Engine/Source/Runtime/Landscape/Private/LandscapeEdit.cpp
+++ b/Engine/Source/Runtime/Landscape/Private/LandscapeEdit.cpp
@@ -539,7 +539,8 @@ bool ULandscapeComponent::ValidateCombinationMaterial(UMaterialInstanceConstant*
 			{
 				return ((TerrainLayerWeightParameter.LayerName == Allocation.LayerInfo->LayerName)
 					&& (TerrainLayerWeightParameter.WeightmapIndex == Allocation.WeightmapTextureIndex)
-					&& (TerrainLayerWeightParameter.bWeightBasedBlend == !Allocation.LayerInfo->bNoWeightBlend));
+					// TODO [jonathan.bard] : TerrainLayerWeightParameter.bWeightBasedBlend is unused now, remove it
+					&& (TerrainLayerWeightParameter.bWeightBasedBlend == (Allocation.LayerInfo->GetBlendMethod() == ELandscapeTargetLayerBlendMethod::FinalWeightBlending)));
 			}))
 		{
 			UE_LOG(LogLandscape, Display, TEXT("Material instance %s in landscape component %s doesn't match the expected shader combination: missing layer %s"),
@@ -570,7 +571,8 @@ FString ULandscapeComponent::GetLayerAllocationKey(const TArray<FWeightmapLayerA
 	{
 		if (Allocations[LayerIdx].GetLayerName() != NAME_None) // ignore layers that have no name
 		{
-			const bool bNoWeightBlend = Allocations[LayerIdx].LayerInfo && Allocations[LayerIdx].LayerInfo->bNoWeightBlend;
+			// TODO [jonathan.bard] : remove from layer allocation key : this doesn't make any difference in the shader anymore since mobile landscape mesh was removed (see bAtLeastOneWeightBasedBlend in UE4.27)
+			const bool bNoWeightBlend = Allocations[LayerIdx].LayerInfo && (Allocations[LayerIdx].LayerInfo->GetBlendMethod() != ELandscapeTargetLayerBlendMethod::FinalWeightBlending);
 			LayerStrings.Add(FString::Printf(TEXT("_%s_%s%d"), *Allocations[LayerIdx].GetLayerName().ToString(), bNoWeightBlend ? TEXT("n") : TEXT("w"), Allocations[LayerIdx].WeightmapTextureIndex));
 		}
 	}
@@ -652,7 +654,8 @@ UMaterialInstanceConstant* ULandscapeComponent::GetCombinationMaterial(FMaterial
 					FName LayerName = Allocation.GetLayerName();
 					if (LayerName != NAME_None) // ignore allocations that have no name
 					{
-						StaticParameters.EditorOnly.TerrainLayerWeightParameters.Add(FStaticTerrainLayerWeightParameter(LayerName, Allocation.WeightmapTextureIndex, !Allocation.LayerInfo->bNoWeightBlend));
+						// TODO [jonathan.bard] : remove from layer allocation key : this doesn't make any difference in the shader anymore since mobile landscape mesh was removed (see bAtLeastOneWeightBasedBlend in UE4.27)
+						StaticParameters.EditorOnly.TerrainLayerWeightParameters.Add(FStaticTerrainLayerWeightParameter(LayerName, Allocation.WeightmapTextureIndex, (Allocation.LayerInfo->GetBlendMethod() == ELandscapeTargetLayerBlendMethod::FinalWeightBlending)));
 					}
 				}
 			}
@@ -3412,7 +3415,8 @@ void ALandscapeProxy::Import(const FGuid& InGuid, int32 InMinX, int32 InMinY, in
 			WeightValues.Empty(EditingAlphaLayerData.Num());
 			WeightValues.AddZeroed(EditingAlphaLayerData.Num());
 			ComponentWeightmapLayerAllocations.Empty(EditingAlphaLayerData.Num());
-
+			// TODO [jonathan.bard] : remove all this once we get rid of non-edit layers, weight balancing at the data level is not reliable and weight blending with edit layers active 
+			//  is providing the same service in a non-destructive way
 			TArray<bool, TInlineAllocator<16>> IsNoBlendArray;
 			IsNoBlendArray.Empty(EditingAlphaLayerData.Num());
 			IsNoBlendArray.AddZeroed(EditingAlphaLayerData.Num());
@@ -3422,7 +3426,7 @@ void ALandscapeProxy::Import(const FGuid& InGuid, int32 InMinX, int32 InMinY, in
 				// Lookup the original layer name
 				WeightValues[WeightLayerIndex] = EditingAlphaLayerData[WeightLayerIndex].AlphaValues;
 				new(ComponentWeightmapLayerAllocations) FWeightmapLayerAllocationInfo(ImportLayerInfos[EditingAlphaLayerData[WeightLayerIndex].LayerIndex].LayerInfo);
-				IsNoBlendArray[WeightLayerIndex] = ImportLayerInfos[EditingAlphaLayerData[WeightLayerIndex].LayerIndex].LayerInfo->bNoWeightBlend;
+				IsNoBlendArray[WeightLayerIndex] = (ImportLayerInfos[EditingAlphaLayerData[WeightLayerIndex].LayerIndex].LayerInfo->GetBlendMethod() != ELandscapeTargetLayerBlendMethod::FinalWeightBlending);
 			}
 
 			// Discard the temporary alpha data
@@ -5436,123 +5440,8 @@ void ALandscape::PostDuplicate(bool bDuplicateForPIE)
 }
 #endif	//WITH_EDITOR
 
-ULandscapeLayerInfoObject::ULandscapeLayerInfoObject(const FObjectInitializer& ObjectInitializer)
-	: Super(ObjectInitializer)
-#if WITH_EDITORONLY_DATA
-	, IsReferencedFromLoadedData(false)
-#endif // WITH_EDITORONLY_DATA
-{
-	Hardness = 0.0f;
-#if WITH_EDITORONLY_DATA
-	MinimumCollisionRelevanceWeight = 0.0f;
-	bNoWeightBlend = false;
-	SplineFalloffModulationTexture = nullptr;
-	SplineFalloffModulationColorMask = ESplineModulationColorMask::Red;
-	SplineFalloffModulationTiling = 1.0f;
-	SplineFalloffModulationBias = 0.5;
-	SplineFalloffModulationScale = 1.0f;
-#endif // WITH_EDITORONLY_DATA
-
-	// Assign initial LayerUsageDebugColor
-	if (!IsTemplate())
-	{
-		LayerUsageDebugColor = GenerateLayerUsageDebugColor();
-	}
-}
-
-FLinearColor ULandscapeLayerInfoObject::GenerateLayerUsageDebugColor() const
-{
-		uint8 Hash[20];
-		FString PathNameString = GetPathName();
-		FSHA1::HashBuffer(*PathNameString, PathNameString.Len() * sizeof(PathNameString[0]), Hash);
-
-	return FLinearColor(float(Hash[0]) / 255.f, float(Hash[1]) / 255.f, float(Hash[2]) / 255.f, 1.f);
-}
 
 #if WITH_EDITOR
-void ULandscapeLayerInfoObject::PostEditChangeProperty(FPropertyChangedEvent& PropertyChangedEvent)
-{
-	Super::PostEditChangeProperty(PropertyChangedEvent);
-
-	const FName MemberPropertyName = PropertyChangedEvent.MemberProperty ? PropertyChangedEvent.MemberProperty->GetFName() : NAME_None;
-
-	if ((MemberPropertyName == GET_MEMBER_NAME_CHECKED(ULandscapeLayerInfoObject, PhysMaterial))
-		|| (MemberPropertyName == GET_MEMBER_NAME_CHECKED(ULandscapeLayerInfoObject, MinimumCollisionRelevanceWeight)))
-	{
-		for (TObjectIterator<ALandscapeProxy> It(/*AdditionalExclusionFlags = */RF_ClassDefaultObject, /*bIncludeDerivedClasses = */true, /*InInternalExclusionFlags = */EInternalObjectFlags::Garbage); It; ++It)
-		{
-			ALandscapeProxy* Proxy = *It;
-			if (Proxy->GetWorld() && !Proxy->GetWorld()->IsPlayInEditor())
-			{
-				ULandscapeInfo* Info = Proxy->GetLandscapeInfo();
-				if (Info)
-				{
-					for (int32 i = 0; i < Info->Layers.Num(); ++i)
-					{
-						if (Info->Layers[i].LayerInfoObj == this)
-						{
-							Proxy->ChangedPhysMaterial();
-							break;
-						}
-					}
-				}
-			}
-		}
-	}
-	else if (MemberPropertyName == GET_MEMBER_NAME_CHECKED(ULandscapeLayerInfoObject, LayerUsageDebugColor))
-	{
-		LayerUsageDebugColor.A = 1.0f;
-		for (TObjectIterator<ALandscapeProxy> It(/*AdditionalExclusionFlags = */RF_ClassDefaultObject, /*bIncludeDerivedClasses = */true, /*InInternalExclusionFlags = */EInternalObjectFlags::Garbage); It; ++It)
-		{
-			ALandscapeProxy* Proxy = *It;
-			if (Proxy->GetWorld() && !Proxy->GetWorld()->IsPlayInEditor())
-			{
-				Proxy->MarkComponentsRenderStateDirty();
-			}
-		}
-	}
-	else if ((MemberPropertyName == GET_MEMBER_NAME_CHECKED(ULandscapeLayerInfoObject, SplineFalloffModulationTexture))
-		|| (MemberPropertyName == GET_MEMBER_NAME_CHECKED(ULandscapeLayerInfoObject, SplineFalloffModulationColorMask))
-		|| (MemberPropertyName == GET_MEMBER_NAME_CHECKED(ULandscapeLayerInfoObject, SplineFalloffModulationBias))
-		|| (MemberPropertyName == GET_MEMBER_NAME_CHECKED(ULandscapeLayerInfoObject, SplineFalloffModulationScale))
-		|| (MemberPropertyName == GET_MEMBER_NAME_CHECKED(ULandscapeLayerInfoObject, SplineFalloffModulationTiling)))
-	{
-		for (TObjectIterator<ULandscapeInfo> It(/*AdditionalExclusionFlags = */RF_ClassDefaultObject, /*bIncludeDerivedClasses = */true, /*InInternalExclusionFlags = */EInternalObjectFlags::Garbage); It; ++It)
-		{
-			if(ALandscape* Landscape = It->LandscapeActor.Get())
-			{
-				Landscape->OnLayerInfoSplineFalloffModulationChanged(this);
-			}
-		}
-	}
-}
-
-void ULandscapeLayerInfoObject::PostEditUndo()
-{
-	Super::PostEditUndo();
-
-	// Force the update of spline data for the (potentially) affected landscapes in case of an undo
-	for (TObjectIterator<ULandscapeInfo> It(/*AdditionalExclusionFlags = */RF_ClassDefaultObject, /*bIncludeDerivedClasses = */true, /*InInternalExclusionFlags = */EInternalObjectFlags::Garbage); It; ++It)
-	{
-		if (ALandscape* Landscape = It->LandscapeActor.Get())
-		{
-			Landscape->OnLayerInfoSplineFalloffModulationChanged(this);
-		}
-	}
-}
-
-void ULandscapeLayerInfoObject::PostLoad()
-{
-	Super::PostLoad();
-	if (GIsEditor)
-	{
-		if (!HasAnyFlags(RF_Standalone))
-		{
-			SetFlags(RF_Standalone);
-		}
-		Hardness = FMath::Clamp<float>(Hardness, 0.0f, 1.0f);
-	}
-}
 
 void ALandscapeProxy::RemoveXYOffsets()
 {
@@ -7912,9 +7801,10 @@ void ULandscapeComponent::GenerateMobileWeightmapLayerAllocations()
 		if (LhsLayerInfo == ALandscapeProxy::VisibilityLayer && RhsLayerInfo != ALandscapeProxy::VisibilityLayer) return true;
 		if (RhsLayerInfo == ALandscapeProxy::VisibilityLayer && LhsLayerInfo != ALandscapeProxy::VisibilityLayer) return false;
 
+		// TODO [jonathan.bard] : this doesn't make any difference in the shader anymore since mobile landscape mesh was removed (see bAtLeastOneWeightBasedBlend in UE4.27)
 		// Sort non-weight blended layers to the front so if we have exactly 3 layers, the 3rd is definitely weight-based.
-		if (LhsLayerInfo->bNoWeightBlend && !RhsLayerInfo->bNoWeightBlend) return true;
-		if (RhsLayerInfo->bNoWeightBlend && !LhsLayerInfo->bNoWeightBlend) return false;
+		if ((LhsLayerInfo->GetBlendMethod() != ELandscapeTargetLayerBlendMethod::FinalWeightBlending) && (RhsLayerInfo->GetBlendMethod() == ELandscapeTargetLayerBlendMethod::FinalWeightBlending)) return true;
+		if ((RhsLayerInfo->GetBlendMethod() != ELandscapeTargetLayerBlendMethod::FinalWeightBlending) && (LhsLayerInfo->GetBlendMethod() == ELandscapeTargetLayerBlendMethod::FinalWeightBlending)) return false;
 
 		return false; // equal, preserve order
 	}));
diff --git a/Engine/Source/Runtime/Landscape/Private/LandscapeEditInterface.cpp b/Engine/Source/Runtime/Landscape/Private/LandscapeEditInterface.cpp
--- a/Engine/Source/Runtime/Landscape/Private/LandscapeEditInterface.cpp
+++ b/Engine/Source/Runtime/Landscape/Private/LandscapeEditInterface.cpp
@@ -1641,9 +1641,9 @@ void ULandscapeComponent::DeleteLayerInternal(ULandscapeLayerInfoObject* LayerIn
 
 	// See if the deleted layer is a NoWeightBlend layer - if not, we don't have to worry about normalization
 	// If the layer doesn't exist, assume it is a blended layer (so renormalization will run whether needed or not)
-	bool bDeleteLayerIsNoWeightBlend = (LayerInfo && LayerInfo->bNoWeightBlend);
-
-	if (!bDeleteLayerIsNoWeightBlend)
+	// TODO [jonathan.bard] : remove all this once we get rid of non-edit layers, weight balancing at the data level is not reliable and weight blending with edit layers active 
+	//  is providing the same service in a non-destructive way
+	if ((LayerInfo == nullptr) || (LayerInfo->GetBlendMethod() == ELandscapeTargetLayerBlendMethod::FinalWeightBlending))
 	{
 		// Lock data for all the weightmaps
 		TArray<FLandscapeTextureDataInfo*> TexDataInfos;
@@ -1661,8 +1661,9 @@ void ULandscapeComponent::DeleteLayerInternal(ULandscapeLayerInfoObject* LayerIn
 			const FWeightmapLayerAllocationInfo& Allocation = ComponentWeightmapLayerAllocations[LayerIdx];
 			LayerDataPtrs.Add((uint8*)TexDataInfos[Allocation.WeightmapTextureIndex]->GetMipData(0) + ChannelOffsets[Allocation.WeightmapTextureChannel]);
 
-			// Find the layer info and record if it is a bNoWeightBlend layer.
-			LayerNoWeightBlends.Add(Allocation.LayerInfo && Allocation.LayerInfo->bNoWeightBlend);
+			// Find the layer info and record if it is a "no weight blend" layer.
+			const bool bPerformLegacyWeightBalancing = (Allocation.LayerInfo == nullptr) || (Allocation.LayerInfo->GetBlendMethod() == ELandscapeTargetLayerBlendMethod::FinalWeightBlending);
+			LayerNoWeightBlends.Add(!bPerformLegacyWeightBalancing);
 		}
 
 		// Find the texture data corresponding to this vertex
@@ -1787,8 +1788,10 @@ void ULandscapeComponent::FillLayer(ULandscapeLayerInfoObject* LayerInfo, FLands
 	Component->Modify(LandscapeEdit.GetShouldDirtyPackage());
 	Proxy->Modify(LandscapeEdit.GetShouldDirtyPackage());
 
-	const bool bFillLayerIsNoWeightBlend = LayerInfo->bNoWeightBlend;
-	bool bClearOtherWeightBlendLayers = !bFillLayerIsNoWeightBlend;
+	// TODO [jonathan.bard] : remove all this once we get rid of non-edit layers, weight balancing at the data level is not reliable and weight blending with edit layers active 
+	//  is providing the same service in a non-destructive way
+	const bool bPerformLegacyWeightBalancing = (LayerInfo->GetBlendMethod() == ELandscapeTargetLayerBlendMethod::FinalWeightBlending);
+	bool bClearOtherWeightBlendLayers = bPerformLegacyWeightBalancing;
 
 	TArray<FWeightmapLayerAllocationInfo>& ComponentWeightmapLayerAllocations = Component->GetWeightmapLayerAllocations(EditLayerGuid);
 	TArray<TObjectPtr<UTexture2D>>& ComponentWeightmapTextures = Component->GetWeightmapTextures(EditLayerGuid);
@@ -1799,10 +1802,13 @@ void ULandscapeComponent::FillLayer(ULandscapeLayerInfoObject* LayerInfo, FLands
 		[LayerInfo](const FWeightmapLayerAllocationInfo& Allocation) { return Allocation.LayerInfo == LayerInfo; });
 
 	// if the layer isn't used on this component yet but is a weight-blend layer, then simply steal the allocation of another weight-blend layer!
-	if (FillLayerIdx == INDEX_NONE && !bFillLayerIsNoWeightBlend)
+	if (FillLayerIdx == INDEX_NONE && bPerformLegacyWeightBalancing)
 	{
 		FillLayerIdx = ComponentWeightmapLayerAllocations.IndexOfByPredicate(
-			[](const FWeightmapLayerAllocationInfo& Allocation) { return !Allocation.LayerInfo || !Allocation.LayerInfo->bNoWeightBlend; });
+			[](const FWeightmapLayerAllocationInfo& Allocation) 
+			{ 
+				return (Allocation.LayerInfo == nullptr) || (Allocation.LayerInfo->GetBlendMethod() == ELandscapeTargetLayerBlendMethod::FinalWeightBlending);
+			});
 
 		if (FillLayerIdx != INDEX_NONE)
 		{
@@ -1857,7 +1863,7 @@ void ULandscapeComponent::FillLayer(ULandscapeLayerInfoObject* LayerInfo, FLands
 				continue;
 			}
 			FWeightmapLayerAllocationInfo& Allocation = ComponentWeightmapLayerAllocations[LayerIdx];
-			if ((Allocation.LayerInfo == nullptr) || Allocation.LayerInfo->bNoWeightBlend)
+			if ((Allocation.LayerInfo == nullptr) || (Allocation.LayerInfo->GetBlendMethod() != ELandscapeTargetLayerBlendMethod::FinalWeightBlending))
 			{
 				continue;
 			}
@@ -2064,9 +2070,11 @@ void ULandscapeComponent::ReplaceLayerInternal(ULandscapeLayerInfoObject* FromLa
 	}
 
 	// See if the deleted layer is a NoWeightBlend layer - if not, we don't have to worry about normalization
-	const bool bFromLayerIsNoWeightBlend = (FromLayerInfo && FromLayerInfo->bNoWeightBlend);
-	const bool bToLayerIsNoWeightBlend = (ToLayerInfo && ToLayerInfo->bNoWeightBlend);
+	const bool bFromLayerIsNoWeightBlend = (FromLayerInfo && (FromLayerInfo->GetBlendMethod() != ELandscapeTargetLayerBlendMethod::FinalWeightBlending));
+	const bool bToLayerIsNoWeightBlend = (ToLayerInfo && (ToLayerInfo->GetBlendMethod() != ELandscapeTargetLayerBlendMethod::FinalWeightBlending));
 
+	// TODO [jonathan.bard] : remove all this once we get rid of non-edit layers, weight balancing at the data level is not reliable and weight blending with edit layers active 
+	//  is providing the same service in a non-destructive way
 	const bool bRequireNormalization = (bFromLayerIsNoWeightBlend != bToLayerIsNoWeightBlend);
 
 	if (bMerging)
@@ -2475,7 +2483,7 @@ inline TMap<const ULandscapeLayerInfoObject*, uint32> FLandscapeEditDataInterfac
 	for (int32 LayerIdx = 0; LayerIdx < ComponentWeightmapLayerAllocations.Num(); LayerIdx++)
 	{
 		const FWeightmapLayerAllocationInfo& Allocation = ComponentWeightmapLayerAllocations[LayerIdx];
-		if (Allocation.LayerInfo == nullptr || Allocation.LayerInfo->bNoWeightBlend)
+		if ((Allocation.LayerInfo == nullptr) || (Allocation.LayerInfo->GetBlendMethod() != ELandscapeTargetLayerBlendMethod::FinalWeightBlending))
 		{
 			continue;
 		}
@@ -2591,7 +2599,8 @@ void FLandscapeEditDataInterface::SetAlphaData(ULandscapeLayerInfoObject* const
 
 	check(LayerInfo != nullptr);
 	check(Data != nullptr);
-	if (LayerInfo->bNoWeightBlend)
+
+	if ((LayerInfo == nullptr) || (LayerInfo->GetBlendMethod() != ELandscapeTargetLayerBlendMethod::FinalWeightBlending))
 	{
 		bWeightAdjust = false;
 	}
@@ -2744,7 +2753,7 @@ void FLandscapeEditDataInterface::SetAlphaData(ULandscapeLayerInfoObject* const
 				if (Allocation.LayerInfo != nullptr) // only take into account valid layer
 				{
 					LayerDataPtrs[LayerIdx] = (uint8*)TexDataInfos[Allocation.WeightmapTextureIndex]->GetMipData(0) + ChannelOffsets[Allocation.WeightmapTextureChannel];
-					LayerNoWeightBlends[LayerIdx] = Allocation.LayerInfo->bNoWeightBlend;
+					LayerNoWeightBlends[LayerIdx] = (Allocation.LayerInfo->GetBlendMethod() != ELandscapeTargetLayerBlendMethod::FinalWeightBlending);
 					LayerEditDataAllZero[LayerIdx] = true;
 				}
 				else
@@ -2819,7 +2828,7 @@ void FLandscapeEditDataInterface::SetAlphaData(ULandscapeLayerInfoObject* const
 							if (bWeightAdjust)
 							{
 								// Normalize all layers including the painted one
-								// gmartin: this isn't used. TODO: Remove
+								// TODO [jonathan.bard] : this isn't used. TODO: Remove
 								if (bTotalWeightAdjust)
 								{
 									int32 MaxLayerIdx = -1;
diff --git a/Engine/Source/Runtime/Landscape/Private/LandscapeEditLayerMergeContext.cpp b/Engine/Source/Runtime/Landscape/Private/LandscapeEditLayerMergeContext.cpp
--- a/Engine/Source/Runtime/Landscape/Private/LandscapeEditLayerMergeContext.cpp
+++ b/Engine/Source/Runtime/Landscape/Private/LandscapeEditLayerMergeContext.cpp
@@ -6,6 +6,7 @@
 #include "Landscape.h"
 #include "LandscapeInfo.h"
 #include "LandscapeLayerInfoObject.h"
+#include "LandscapeEditLayerTargetLayerGroup.h"
 #include "Materials/MaterialExpressionLandscapeVisibilityMask.h"
 
 namespace UE::Landscape::EditLayers
@@ -57,6 +58,9 @@ FMergeContext::FMergeContext(ALandscape* InLandscape, bool bInIsHeightmapMerge,
 		}
 		check(VisibilityTargetLayerIndex != INDEX_NONE);
 
+		// Now that we have all valid target layers ready, analyze them to prepare the target layer groups as they are declared within the layer infos : 
+		FinalizeTargetLayerGroupsPerBlendingMethod();
+
 		VisibilityTargetLayerMask = TBitArray<>(false, AllTargetLayerNames.Num());
 		VisibilityTargetLayerMask[VisibilityTargetLayerIndex] = true;
 	}
@@ -219,6 +223,85 @@ TBitArray<> FMergeContext::BuildTargetLayerBitIndices(bool bInBitValue) const
 	return TBitArray<>(bInBitValue, AllTargetLayerNames.Num());
 }
 
+TArray<FTargetLayerGroup> FMergeContext::BuildGenericBlendTargetLayerGroups(const TBitArray<>& InTargetLayerBitIndices) const
+{
+	// Only ELandscapeTargetLayerBlendMethod::PremultipliedAlphaBlending is applied per renderer, so we only need to retrieve the target layer groups
+	//  for the specified target layers and trim them of all other target layers
+	const TArray<FTargetLayerGroup>& SourceTargetLayerGroups = TargetLayerGroupsPerBlendingMethod[static_cast<int32>(ELandscapeTargetLayerBlendMethod::PremultipliedAlphaBlending)];
+
+	// First, intersect our layer groups with the specified target layers 
+	TArray<FTargetLayerGroup> FilteredTargetLayerGroups;
+	FilteredTargetLayerGroups.Reserve(SourceTargetLayerGroups.Num());
+	Algo::Transform(SourceTargetLayerGroups, FilteredTargetLayerGroups, [&InTargetLayerBitIndices](const FTargetLayerGroup& InTargetLayerGroup)
+		{
+			return FTargetLayerGroup(InTargetLayerGroup.GetName(), TBitArray<>::BitwiseAND(InTargetLayerGroup.GetWeightmapTargetLayerBitIndices(), InTargetLayerBitIndices, EBitwiseOperatorFlags::MinSize));
+		});
+	
+	// Then remove all those that are now empty : 
+	FilteredTargetLayerGroups.RemoveAllSwap([](const FTargetLayerGroup& InTargetLayerGroup) { return !InTargetLayerGroup.GetWeightmapTargetLayerBitIndices().Contains(true); });
+
+	return FilteredTargetLayerGroups;
+}
+
+void FMergeContext::FinalizeTargetLayerGroupsPerBlendingMethod()
+{
+	for (int32 Index = 0; Index < static_cast<int32>(ELandscapeTargetLayerBlendMethod::Count); ++Index)
+	{
+		TargetLayerGroupsPerBlendingMethod[Index] = GatherTargetLayerGroupsForBlendMethod(static_cast<ELandscapeTargetLayerBlendMethod>(Index));
+	}
+}
+
+TArray<FTargetLayerGroup> FMergeContext::GatherTargetLayerGroupsForBlendMethod(ELandscapeTargetLayerBlendMethod InBlendMethod) const
+{
+	TArray<FTargetLayerGroup> TargetLayerGroups;
+	switch (InBlendMethod)
+	{
+	case ELandscapeTargetLayerBlendMethod::None:
+	{
+		// Returning an empty layer group list for the no weight-blending case is fine because it will be treated as N independent layer groups (each with a single target layer)
+	}
+	break;
+	case ELandscapeTargetLayerBlendMethod::FinalWeightBlending:
+	{
+		TBitArray<> WeightBlendedWeightmapLayerBitIndices = BuildTargetLayerBitIndices(/*bInBitValue = */false);
+		ForEachValidTargetLayer(
+			[&WeightBlendedWeightmapLayerBitIndices](int32 InTargetLayerIndex, const FName& InTargetLayerName, ULandscapeLayerInfoObject* InWeightmapLayerInfo)
+			{
+				check(InWeightmapLayerInfo != nullptr);
+				WeightBlendedWeightmapLayerBitIndices[InTargetLayerIndex] = (InWeightmapLayerInfo->GetBlendMethod() == ELandscapeTargetLayerBlendMethod::FinalWeightBlending);
+				return true;
+			});
+		static const FName FinalWeightBlendingLayerGroupName(TEXT("FinalWeightBlendingLayerGroup"));
+		TargetLayerGroups.Emplace(FinalWeightBlendingLayerGroupName, WeightBlendedWeightmapLayerBitIndices);
+	}
+	break;
+	case ELandscapeTargetLayerBlendMethod::PremultipliedAlphaBlending:
+	{
+		// Find all ELandscapeTargetLayerBlendMethod::PremultipliedAlphaBlending layers and segregate them by blend group :
+		TMap<FName, TBitArray<>> TargetLayerGroupsMap;
+		const TBitArray<> DefaultTargetLayerBitIndices = BuildTargetLayerBitIndices(/*bInBitValue = */false);
+		ForEachValidTargetLayer(
+			[&TargetLayerGroupsMap, &DefaultTargetLayerBitIndices](int32 InTargetLayerIndex, const FName& InTargetLayerName, ULandscapeLayerInfoObject* InWeightmapLayerInfo)
+			{
+				check(InWeightmapLayerInfo != nullptr);
+				if (InWeightmapLayerInfo->GetBlendMethod() == ELandscapeTargetLayerBlendMethod::PremultipliedAlphaBlending)
+				{
+					TBitArray<>& TargetLayerGroup = TargetLayerGroupsMap.FindOrAdd(InWeightmapLayerInfo->GetBlendGroup(), DefaultTargetLayerBitIndices);
+					TargetLayerGroup[InTargetLayerIndex] = true;
+				}
+				return true;
+			});
+
+		Algo::Transform(TargetLayerGroupsMap, TargetLayerGroups, [](const TPair<FName, TBitArray<>>& InTargetLayerGroupPair) { return FTargetLayerGroup(InTargetLayerGroupPair.Key, InTargetLayerGroupPair.Value); });
+	}
+	break;
+	default:
+		check(false);
+	}
+
+	return TargetLayerGroups;
+}
+
 #endif // WITH_EDITOR
 
 } // namespace UE::Landscape::EditLayers
diff --git a/Engine/Source/Runtime/Landscape/Private/LandscapeEditLayerMergeRenderContext.cpp b/Engine/Source/Runtime/Landscape/Private/LandscapeEditLayerMergeRenderContext.cpp
--- a/Engine/Source/Runtime/Landscape/Private/LandscapeEditLayerMergeRenderContext.cpp
+++ b/Engine/Source/Runtime/Landscape/Private/LandscapeEditLayerMergeRenderContext.cpp
@@ -58,6 +58,8 @@ class FRenderComponentIdPS : public FGlobalShader
 
 	static void ModifyCompilationEnvironment(const FGlobalShaderPermutationParameters& Parameters, FShaderCompilerEnvironment& OutEnvironment)
 	{
+		UE::Landscape::ModifyShaderCompilerEnvironmentForDebug(OutEnvironment);
+
 		OutEnvironment.SetDefine(TEXT("RENDER_COMPONENTID"), 1);
 	}
 };
@@ -92,6 +94,8 @@ class FCopyQuadsPS : public FGlobalShader
 
 	static void ModifyCompilationEnvironment(const FGlobalShaderPermutationParameters& Parameters, FShaderCompilerEnvironment& OutEnvironment)
 	{
+		UE::Landscape::ModifyShaderCompilerEnvironmentForDebug(OutEnvironment);
+
 		OutEnvironment.SetDefine(TEXT("COPY_QUADS"), 1);
 	}
 };
@@ -129,6 +133,8 @@ public:
 
 	static void ModifyCompilationEnvironment(const FGlobalShaderPermutationParameters& InParameters, FShaderCompilerEnvironment& OutEnvironment)
 	{
+		UE::Landscape::ModifyShaderCompilerEnvironmentForDebug(OutEnvironment);
+
 		OutEnvironment.SetDefine(TEXT("MERGE_EDIT_LAYER"), 1);
 	}
 
@@ -162,8 +168,10 @@ class FLandscapeEditLayersWeightmapsMergeEditLayerPS : public FGlobalShader
 public:
 	BEGIN_SHADER_PARAMETER_STRUCT(FParameters, )
 		SHADER_PARAMETER(uint32, InTargetLayerIndex)
+		SHADER_PARAMETER(uint32, InNumTargetLayers)
 		SHADER_PARAMETER(uint32, InEditLayerTargetLayerBlendMode)
 		SHADER_PARAMETER(float, InEditLayerAlpha)
+		SHADER_PARAMETER_RDG_BUFFER_SRV(StructuredBuffer<UE::Landscape::FMergeEditLayerTargetLayerInfo>, InMergeEditLayerTargetLayerInfos)
 		SHADER_PARAMETER_RDG_TEXTURE_SRV(Texture2DArray<float4>, InCurrentEditLayerWeightmaps)
 		SHADER_PARAMETER_RDG_TEXTURE_SRV(Texture2DArray<float4>, InPreviousEditLayersWeightmaps)
 		RENDER_TARGET_BINDING_SLOTS()
@@ -176,6 +184,8 @@ public:
 
 	static void ModifyCompilationEnvironment(const FGlobalShaderPermutationParameters& InParameters, FShaderCompilerEnvironment& OutEnvironment)
 	{
+		UE::Landscape::ModifyShaderCompilerEnvironmentForDebug(OutEnvironment);
+
 		OutEnvironment.SetDefine(TEXT("MERGE_EDIT_LAYER"), 1);
 	}
 
@@ -799,11 +809,49 @@ void FMergeRenderContext::GenericBlendLayer(const FBlendParams& InBlendParams, F
 	}
 	else
 	{
+		const TArray<FTargetLayerGroup> TargetLayerGroups(BuildGenericBlendTargetLayerGroups(ConvertTargetLayerNamesToBitIndices(RenderParams.TargetLayerGroupLayerNames)));
+
+		// Prepare layer info for the target layers that will need it for merging :
+		TArray<FMergeEditLayerTargetLayerInfo> WeightmapTargetLayerInfos;
+		WeightmapTargetLayerInfos.Reserve(RenderParams.TargetLayerGroupLayerInfos.Num());
+		Algo::Transform(RenderParams.TargetLayerGroupLayerInfos, WeightmapTargetLayerInfos, [this, &TargetLayerGroups, RenderBatch](ULandscapeLayerInfoObject* InLayerInfo)
+			{
+				check(InLayerInfo != nullptr); // There should only be valid layer infos at this point
+				const int32 TargetLayerIndex = GetTargetLayerIndexForLayerInfoChecked(InLayerInfo);
+				// Skip weight-blended layers that are not rendered in the batch, because their texture slice might contain garbage :
+				const bool bRenderedInBatch = RenderBatch->TargetLayerBitIndices[TargetLayerIndex];
+
+				FMergeEditLayerTargetLayerInfo WeightmapTargetLayerInfo;
+				WeightmapTargetLayerInfo.BlendGroupIndex = INDEX_NONE;
+				if (!bRenderedInBatch)
+				{
+					WeightmapTargetLayerInfo.Flags |= EWeightmapTargetLayerFlags::Skip;
+				}
+				if (InLayerInfo->GetBlendMethod() == ELandscapeTargetLayerBlendMethod::PremultipliedAlphaBlending)
+				{
+					WeightmapTargetLayerInfo.Flags |= EWeightmapTargetLayerFlags::IsPremultipliedAlphaWeightBlended;
+					// Find the target layer group that corresponds to this target layer. There must be one
+					const int32 TargetLayerGroupIndex = TargetLayerGroups.IndexOfByPredicate([TargetLayerIndex](const FTargetLayerGroup& InTargetLayerGroup)
+						{ 
+							return InTargetLayerGroup.GetWeightmapTargetLayerBitIndices()[TargetLayerIndex];
+						});
+					checkf(TargetLayerGroupIndex != INDEX_NONE, TEXT("Couldn't find a target layer group for layer %s. All premultiplied alpha blending must belong to a group (and one only)"), *InLayerInfo->LayerName.ToString());
+					WeightmapTargetLayerInfo.BlendGroupIndex = TargetLayerGroupIndex;
+				}
+				if (InLayerInfo == ALandscapeProxy::VisibilityLayer)
+				{
+					WeightmapTargetLayerInfo.Flags |= EWeightmapTargetLayerFlags::IsVisibilityLayer;
+				}
+				
+				return WeightmapTargetLayerInfo;
+			});
+
 		ForEachTargetLayer(RenderBatch->TargetLayerBitIndices,
-			[WriteRT, CurrentLayerReadRT, PreviousLayersReadRT, &RenderBatch, &RenderParams, &InBlendParams, &RDGBuilderRecorder](int32 InTargetLayerIndex, const FName& InTargetLayerName, ULandscapeLayerInfoObject* InWeightmapLayerInfo)
+			[WriteRT, CurrentLayerReadRT, PreviousLayersReadRT, &WeightmapTargetLayerInfos, &RenderBatch, &RenderParams, &InBlendParams, &RDGBuilderRecorder](int32 InTargetLayerIndex, const FName& InTargetLayerName, ULandscapeLayerInfoObject* InWeightmapLayerInfo)
 			{
-				int32 TargetLayerIndexInGroup = RenderParams.TargetLayerGroupLayerNames.Find(InTargetLayerName);
+				const int32 TargetLayerIndexInGroup = RenderParams.TargetLayerGroupLayerNames.Find(InTargetLayerName);
 				check(TargetLayerIndexInGroup != INDEX_NONE);
+				const int32 NumTargetLayersInGroup = RenderParams.TargetLayerGroupLayerNames.Num();
 
 				// By default, use passthrough mode so that each layer gets at least copied into the WriteRT :
 				const FWeightmapBlendParams* TargetLayerBlendParams = &FWeightmapBlendParams::GetDefaultPassthroughBlendParams();
@@ -821,6 +869,8 @@ void FMergeRenderContext::GenericBlendLayer(const FBlendParams& InBlendParams, F
 				auto RDGCommand =
 					[ TargetLayerBlendParams = *TargetLayerBlendParams
 					, TargetLayerIndexInGroup
+					, NumTargetLayersInGroup
+					, WeightmapTargetLayerInfos
 					, InTargetLayerName
 					, OutputResource = WriteRT->GetRenderTarget2DArray()->GetResource()
 					, OutputResourceName = WriteRT->GetDebugName()
@@ -835,11 +885,17 @@ void FMergeRenderContext::GenericBlendLayer(const FBlendParams& InBlendParams, F
 						FRDGTextureRef PreviousEditLayersTextureRef = GraphBuilder.RegisterExternalTexture(CreateRenderTarget(PreviousEditLayersResource->TextureRHI, TEXT("PreviousEditLayersTexture")));
 						FRDGTextureSRVRef PreviousEditLayersTextureSRVRef = GraphBuilder.CreateSRV(FRDGTextureSRVDesc::Create(PreviousEditLayersTextureRef));
 
+						// TODO [jonathan.bard] : we could find a way to create this once for all target layers instead of re-creating it for each RDGCommand
+						FRDGBufferRef TargetLayerInfosBuffer = CreateStructuredBuffer(GraphBuilder, TEXT("LandscapeMergeEditLayerTargetLayerInfosBuffer"), WeightmapTargetLayerInfos);
+						FRDGBufferSRVRef TargetLayerInfosBufferSRV = GraphBuilder.CreateSRV(FRDGBufferSRVDesc(TargetLayerInfosBuffer));
+
 						FLandscapeEditLayersWeightmapsMergeEditLayerPS::FParameters* PSParams = GraphBuilder.AllocParameters<FLandscapeEditLayersWeightmapsMergeEditLayerPS::FParameters>();
 						PSParams->RenderTargets[0] = FRenderTargetBinding(OutputTextureRef, ERenderTargetLoadAction::ENoAction, /*InMipIndex = */0, /*InArraySlice = */TargetLayerIndexInGroup);
 						PSParams->InTargetLayerIndex = TargetLayerIndexInGroup;
+						PSParams->InNumTargetLayers = NumTargetLayersInGroup;
 						PSParams->InEditLayerTargetLayerBlendMode = static_cast<uint32>(TargetLayerBlendParams.BlendMode);
 						PSParams->InEditLayerAlpha = TargetLayerBlendParams.Alpha;
+						PSParams->InMergeEditLayerTargetLayerInfos = TargetLayerInfosBufferSRV;
 						PSParams->InCurrentEditLayerWeightmaps = CurrentEditLayerTextureSRVRef;
 						PSParams->InPreviousEditLayersWeightmaps = PreviousEditLayersTextureSRVRef;
 
diff --git a/Engine/Source/Runtime/Landscape/Private/LandscapeEditLayerRendererPrivate.cpp b/Engine/Source/Runtime/Landscape/Private/LandscapeEditLayerRendererPrivate.cpp
--- a/Engine/Source/Runtime/Landscape/Private/LandscapeEditLayerRendererPrivate.cpp
+++ b/Engine/Source/Runtime/Landscape/Private/LandscapeEditLayerRendererPrivate.cpp
@@ -7,6 +7,7 @@
 #include "Algo/Transform.h"
 #include "LandscapeEditResourcesSubsystem.h"
 #include "LandscapeEditLayerMergeRenderContext.h"
+#include "LandscapeEditLayerTargetLayerGroup.h"
 #include "LandscapeLayerInfoObject.h"
 #include "LandscapeInfo.h"
 #include "RHIAccess.h"
@@ -19,7 +20,8 @@
 // ----------------------------------------------------------------------------------
 
 void ULandscapeDefaultEditLayerRenderer::GetRendererStateInfo(const UE::Landscape::EditLayers::FMergeContext* InMergeContext,
-	UE::Landscape::EditLayers::FEditLayerTargetTypeState& OutSupportedTargetTypeState, UE::Landscape::EditLayers::FEditLayerTargetTypeState& OutEnabledTargetTypeState, TArray<TBitArray<>>& OutTargetLayerGroups) const
+	UE::Landscape::EditLayers::FEditLayerTargetTypeState& OutSupportedTargetTypeState, UE::Landscape::EditLayers::FEditLayerTargetTypeState& OutEnabledTargetTypeState, 
+	TArray<UE::Landscape::EditLayers::FTargetLayerGroup>& OutTargetLayerGroups) const
 {
 	using namespace UE::Landscape::EditLayers;
 
@@ -71,7 +73,8 @@ bool ULandscapeDefaultEditLayerRenderer::RenderLayer(UE::Landscape::EditLayers::
 // ----------------------------------------------------------------------------------
 
 void ULandscapeHeightmapNormalsEditLayerRenderer::GetRendererStateInfo(const UE::Landscape::EditLayers::FMergeContext* InMergeContext,
-	UE::Landscape::EditLayers::FEditLayerTargetTypeState& OutSupportedTargetTypeState, UE::Landscape::EditLayers::FEditLayerTargetTypeState& OutEnabledTargetTypeState, TArray<TBitArray<>>& OutTargetLayerGroups) const
+	UE::Landscape::EditLayers::FEditLayerTargetTypeState& OutSupportedTargetTypeState, UE::Landscape::EditLayers::FEditLayerTargetTypeState& OutEnabledTargetTypeState, 
+	TArray<UE::Landscape::EditLayers::FTargetLayerGroup>& OutTargetLayerGroups) const
 {
 	// Only relevant for heightmaps :
 	OutSupportedTargetTypeState.SetTargetTypeMask(ELandscapeToolTargetTypeFlags::Heightmap);
@@ -105,11 +108,13 @@ FString ULandscapeHeightmapNormalsEditLayerRenderer::GetEditLayerRendererDebugNa
 // ----------------------------------------------------------------------------------
 
 void ULandscapeWeightmapWeightBlendedLayersRenderer::GetRendererStateInfo(const UE::Landscape::EditLayers::FMergeContext* InMergeContext,
-	UE::Landscape::EditLayers::FEditLayerTargetTypeState& OutSupportedTargetTypeState, UE::Landscape::EditLayers::FEditLayerTargetTypeState& OutEnabledTargetTypeState, TArray<TBitArray<>>& OutTargetLayerGroups) const
+	UE::Landscape::EditLayers::FEditLayerTargetTypeState& OutSupportedTargetTypeState, UE::Landscape::EditLayers::FEditLayerTargetTypeState& OutEnabledTargetTypeState, 
+	TArray<UE::Landscape::EditLayers::FTargetLayerGroup>& OutTargetLayerGroups) const
 {
 	using namespace UE::Landscape::EditLayers;
 
 	TBitArray<> WeightBlendedWeightmapLayerBitIndices = GatherWeightBlendedWeightmapLayerBitIndices(InMergeContext);
+
 	// Only relevant for weightmaps :
 	FEditLayerTargetTypeState OutputTargetTypeState(InMergeContext, ELandscapeToolTargetTypeFlags::Weightmap, WeightBlendedWeightmapLayerBitIndices);
 	OutEnabledTargetTypeState = OutSupportedTargetTypeState = OutputTargetTypeState;
@@ -117,7 +122,8 @@ void ULandscapeWeightmapWeightBlendedLayersRenderer::GetRendererStateInfo(const
 	// Now fill in the target layer groups : one group for all weight-blended layers : 
 	if (WeightBlendedWeightmapLayerBitIndices.Contains(true))
 	{
-		OutTargetLayerGroups.Add(WeightBlendedWeightmapLayerBitIndices);
+		static const FName FinalWeightBlendedTargetLayerGroupName(TEXT("FinalWeightBlendedLayerGroup"));
+		OutTargetLayerGroups.Add(FTargetLayerGroup(FinalWeightBlendedTargetLayerGroupName, WeightBlendedWeightmapLayerBitIndices));
 	}
 }
 
@@ -148,15 +154,16 @@ FString ULandscapeWeightmapWeightBlendedLayersRenderer::GetEditLayerRendererDebu
 
 TBitArray<> ULandscapeWeightmapWeightBlendedLayersRenderer::GatherWeightBlendedWeightmapLayerBitIndices(const UE::Landscape::EditLayers::FMergeContext* InMergeContext) const
 {
-	TBitArray<> WeightBlendedWeightmapLayerBitIndices = InMergeContext->GetValidTargetLayerBitIndices();
-	InMergeContext->ForEachValidTargetLayer(
-		[&WeightBlendedWeightmapLayerBitIndices, InMergeContext] (int32 InTargetLayerIndex, const FName& InTargetLayerName, ULandscapeLayerInfoObject* InWeightmapLayerInfo)
-		{
-			check(InWeightmapLayerInfo != nullptr);
-			WeightBlendedWeightmapLayerBitIndices[InTargetLayerIndex] = !InWeightmapLayerInfo->bNoWeightBlend;
-			return true;
-		});
-	return WeightBlendedWeightmapLayerBitIndices;
+	using namespace UE::Landscape::EditLayers;
+
+	// Retrieve only the target layers that use ELandscapeTargetLayerBlendMethod::FinalWeightBlending
+	TArray<FTargetLayerGroup> TargetLayerGroups = InMergeContext->GetTargetLayerGroupsPerBlendingMethod()[static_cast<uint8>(ELandscapeTargetLayerBlendMethod::FinalWeightBlending)];
+	// There should be 0 if no target layer uses ELandscapeTargetLayerBlendMethod::FinalWeightBlending, or a single one since it doesn't have the notion of "blend group" :
+	check(TargetLayerGroups.Num() <= 1);
+
+	return TargetLayerGroups.IsEmpty()
+		? InMergeContext->BuildTargetLayerBitIndices(/*bInBitValue = */false)
+		: TargetLayerGroups[0].GetWeightmapTargetLayerBitIndices();
 }
 
 #endif // WITH_EDITOR
diff --git a/Engine/Source/Runtime/Landscape/Private/LandscapeEditLayerRendererPrivate.h b/Engine/Source/Runtime/Landscape/Private/LandscapeEditLayerRendererPrivate.h
--- a/Engine/Source/Runtime/Landscape/Private/LandscapeEditLayerRendererPrivate.h
+++ b/Engine/Source/Runtime/Landscape/Private/LandscapeEditLayerRendererPrivate.h
@@ -28,7 +28,8 @@ class ULandscapeDefaultEditLayerRenderer : public UObject
 public:
 	//~ Begin ILandscapeEditLayerRenderer implementation
 	virtual void GetRendererStateInfo(const UE::Landscape::EditLayers::FMergeContext* InMergeContext,
-		UE::Landscape::EditLayers::FEditLayerTargetTypeState& OutSupportedTargetTypeState, UE::Landscape::EditLayers::FEditLayerTargetTypeState& OutEnabledTargetTypeState, TArray<TBitArray<>>& OutTargetLayerGroups) const override;
+		UE::Landscape::EditLayers::FEditLayerTargetTypeState& OutSupportedTargetTypeState, UE::Landscape::EditLayers::FEditLayerTargetTypeState& OutEnabledTargetTypeState,
+		TArray<UE::Landscape::EditLayers::FTargetLayerGroup>& OutTargetLayerGroups) const override;
 	virtual TArray<UE::Landscape::EditLayers::FEditLayerRenderItem> GetRenderItems(const UE::Landscape::EditLayers::FMergeContext* InMergeContext) const override;
 	virtual FString GetEditLayerRendererDebugName() const override;
 	virtual UE::Landscape::EditLayers::ERenderFlags GetRenderFlags(const UE::Landscape::EditLayers::FMergeContext* InMergeContext) const override { return UE::Landscape::EditLayers::ERenderFlags::RenderMode_Recorded; }
@@ -54,7 +55,8 @@ class ULandscapeHeightmapNormalsEditLayerRenderer : public UObject
 public:
 	//~ Begin ILandscapeEditLayerRenderer implementation
 	virtual void GetRendererStateInfo(const UE::Landscape::EditLayers::FMergeContext* InMergeContext,
-		UE::Landscape::EditLayers::FEditLayerTargetTypeState& OutSupportedTargetTypeState, UE::Landscape::EditLayers::FEditLayerTargetTypeState& OutEnabledTargetTypeState, TArray<TBitArray<>>& OutTargetLayerGroups) const override;
+		UE::Landscape::EditLayers::FEditLayerTargetTypeState& OutSupportedTargetTypeState, UE::Landscape::EditLayers::FEditLayerTargetTypeState& OutEnabledTargetTypeState, 
+		TArray<UE::Landscape::EditLayers::FTargetLayerGroup>& OutTargetLayerGroups) const override;
 	virtual TArray<UE::Landscape::EditLayers::FEditLayerRenderItem> GetRenderItems(const UE::Landscape::EditLayers::FMergeContext* InMergeContext) const override;
 	virtual FString GetEditLayerRendererDebugName() const override;
 	virtual UE::Landscape::EditLayers::ERenderFlags GetRenderFlags(const UE::Landscape::EditLayers::FMergeContext* InMergeContext) const override { return UE::Landscape::EditLayers::ERenderFlags::RenderMode_Recorded; }
@@ -78,7 +80,8 @@ class ULandscapeWeightmapWeightBlendedLayersRenderer : public UObject
 public:
 	//~ Begin ILandscapeEditLayerRenderer implementation
 	virtual void GetRendererStateInfo(const UE::Landscape::EditLayers::FMergeContext* InMergeContext,
-		UE::Landscape::EditLayers::FEditLayerTargetTypeState& OutSupportedTargetTypeState, UE::Landscape::EditLayers::FEditLayerTargetTypeState& OutEnabledTargetTypeState, TArray<TBitArray<>>& OutTargetLayerGroups) const override;
+		UE::Landscape::EditLayers::FEditLayerTargetTypeState& OutSupportedTargetTypeState, UE::Landscape::EditLayers::FEditLayerTargetTypeState& OutEnabledTargetTypeState, 
+		TArray<UE::Landscape::EditLayers::FTargetLayerGroup>& OutTargetLayerGroups) const override;
 	virtual TArray<UE::Landscape::EditLayers::FEditLayerRenderItem> GetRenderItems(const UE::Landscape::EditLayers::FMergeContext* InMergeContext) const override;
 	virtual FString GetEditLayerRendererDebugName() const override;
 	virtual UE::Landscape::EditLayers::ERenderFlags GetRenderFlags(const UE::Landscape::EditLayers::FMergeContext* InMergeContext) const override { return UE::Landscape::EditLayers::ERenderFlags::RenderMode_Recorded; }
diff --git a/Engine/Source/Runtime/Landscape/Private/LandscapeEditLayerRendererState.cpp b/Engine/Source/Runtime/Landscape/Private/LandscapeEditLayerRendererState.cpp
--- a/Engine/Source/Runtime/Landscape/Private/LandscapeEditLayerRendererState.cpp
+++ b/Engine/Source/Runtime/Landscape/Private/LandscapeEditLayerRendererState.cpp
@@ -41,14 +41,15 @@ FEditLayerRendererState::FEditLayerRendererState(const FMergeContext* InMergeCon
 	InMergeContext->ForEachTargetLayer(SupportedTargetTypeState.GetActiveWeightmapBitIndices(),
 		[InMergeContext, InRenderer, this](int32 InTargetLayerIndex, const FName& InTargetLayerName, ULandscapeLayerInfoObject* InWeightmapLayerInfo)
 		{
-			const int32 TargetLayerGroupCount = Algo::CountIf(TargetLayerGroups, [InTargetLayerIndex](const TBitArray<>& InTargetLayerGroup) { return InTargetLayerGroup[InTargetLayerIndex]; });
+			const int32 TargetLayerGroupCount = Algo::CountIf(TargetLayerGroups, [InTargetLayerIndex](const FTargetLayerGroup& InTargetLayerGroup) { return InTargetLayerGroup.GetWeightmapTargetLayerBitIndices()[InTargetLayerIndex]; });
 			checkf(TargetLayerGroupCount < 2, TEXT("Target layer %s belongs to more than 1 target layer group in edit layer renderer %s. This is forbidden: in the end, it must belong to 1 and 1 only."),
 				*InTargetLayerName.ToString(), *InRenderer->GetEditLayerRendererDebugName());
 			if (TargetLayerGroupCount == 0)
 			{
-				TBitArray<> TargetLayerGroup = InMergeContext->BuildTargetLayerBitIndices(false);
-				TargetLayerGroup[InTargetLayerIndex] = true;
-				TargetLayerGroups.Add(TargetLayerGroup);
+				static const FName SoloTargetLayerGroupName(TEXT("SoloTargetLayerGroup"));
+				TBitArray<> TargetLayerGroupBitIndices = InMergeContext->BuildTargetLayerBitIndices(false);
+				TargetLayerGroupBitIndices[InTargetLayerIndex] = true;
+				TargetLayerGroups.Add(FTargetLayerGroup(SoloTargetLayerGroupName, TargetLayerGroupBitIndices));
 			}
 			return true;
 		});
diff --git a/Engine/Source/Runtime/Landscape/Private/LandscapeEditLayers.cpp b/Engine/Source/Runtime/Landscape/Private/LandscapeEditLayers.cpp
--- a/Engine/Source/Runtime/Landscape/Private/LandscapeEditLayers.cpp
+++ b/Engine/Source/Runtime/Landscape/Private/LandscapeEditLayers.cpp
@@ -330,7 +330,6 @@ TAutoConsoleVariable<int32> CVarLandscapeEditLayersClearBeforeEachWriteToScratch
 	TEXT("Debug to help with RenderDoc debugging : clear each time we're about to write on a scratch render target (since those are reused and can be used to write RTs of different resolutions"),
 	ECVF_RenderThreadSafe);
 
-
 static void LandscapeDumpSelectiveLayerRender(const TArray<FString>& Args);
 
 FAutoConsoleCommand CmdLandscapeDumpSelectiveLayerRender(
@@ -527,7 +526,6 @@ struct FLandscapeEditLayerWeightmapPaintLayerInfo
 
 namespace UE::Landscape::Private
 {
-
 #if WITH_EDITOR
 	static FFileHelper::EColorChannel GetWeightmapColorChannel(const FWeightmapLayerAllocationInfo& AllocInfo)
 	{
@@ -706,6 +704,7 @@ public:
 
 	static void ModifyCompilationEnvironment(const FGlobalShaderPermutationParameters& Parameters, FShaderCompilerEnvironment& OutEnvironment)
 	{
+		UE::Landscape::ModifyShaderCompilerEnvironmentForDebug(OutEnvironment);
 	}
 
 	FLandscapeLayersVS(const ShaderMetaType::CompiledShaderInitializerType& Initializer)
@@ -776,6 +775,7 @@ public:
 
 	static void ModifyCompilationEnvironment(const FGlobalShaderPermutationParameters& Parameters, FShaderCompilerEnvironment& OutEnvironment)
 	{
+		UE::Landscape::ModifyShaderCompilerEnvironmentForDebug(OutEnvironment);
 	}
 
 	FLandscapeLayersHeightmapPS(const ShaderMetaType::CompiledShaderInitializerType& Initializer)
@@ -841,6 +841,7 @@ public:
 
 	static void ModifyCompilationEnvironment(const FGlobalShaderPermutationParameters& Parameters, FShaderCompilerEnvironment& OutEnvironment)
 	{
+		UE::Landscape::ModifyShaderCompilerEnvironmentForDebug(OutEnvironment);
 	}
 
 	FLandscapeLayersHeightmapMipsPS(const ShaderMetaType::CompiledShaderInitializerType& Initializer)
@@ -914,6 +915,7 @@ public:
 
 	static void ModifyCompilationEnvironment(const FGlobalShaderPermutationParameters& Parameters, FShaderCompilerEnvironment& OutEnvironment)
 	{
+		UE::Landscape::ModifyShaderCompilerEnvironmentForDebug(OutEnvironment);
 	}
 
 	FLandscapeLayersWeightmapPS(const ShaderMetaType::CompiledShaderInitializerType& Initializer)
@@ -968,6 +970,7 @@ public:
 
 	static void ModifyCompilationEnvironment(const FGlobalShaderPermutationParameters& Parameters, FShaderCompilerEnvironment& OutEnvironment)
 	{
+		UE::Landscape::ModifyShaderCompilerEnvironmentForDebug(OutEnvironment);
 	}
 
 	FLandscapeLayersWeightmapMipsPS(const ShaderMetaType::CompiledShaderInitializerType& Initializer)
@@ -1032,6 +1035,7 @@ public:
 
 	static void ModifyCompilationEnvironment(const FGlobalShaderPermutationParameters& Parameters, FShaderCompilerEnvironment& OutEnvironment)
 	{
+		UE::Landscape::ModifyShaderCompilerEnvironmentForDebug(OutEnvironment);
 	}
 
 	FLandscapeCopyTexturePS(const ShaderMetaType::CompiledShaderInitializerType& Initializer)
@@ -1159,6 +1163,7 @@ public:
 	{
 		OutEnvironment.SetDefine(TEXT("THREADGROUP_SIZEX"), GLandscapeLayerWeightmapThreadGroupSizeX);
 		OutEnvironment.SetDefine(TEXT("THREADGROUP_SIZEY"), GLandscapeLayerWeightmapThreadGroupSizeY);
+		UE::Landscape::ModifyShaderCompilerEnvironmentForDebug(OutEnvironment);
 	}
 
 	FLandscapeLayerWeightmapExtractMaterialLayersCS(const ShaderMetaType::CompiledShaderInitializerType& Initializer)
@@ -1352,6 +1357,7 @@ public:
 	{
 		OutEnvironment.SetDefine(TEXT("THREADGROUP_SIZEX"), GLandscapeLayerWeightmapThreadGroupSizeX);
 		OutEnvironment.SetDefine(TEXT("THREADGROUP_SIZEY"), GLandscapeLayerWeightmapThreadGroupSizeY);
+		UE::Landscape::ModifyShaderCompilerEnvironmentForDebug(OutEnvironment);
 	}
 
 	FLandscapeLayerWeightmapPackMaterialLayersCS(const ShaderMetaType::CompiledShaderInitializerType& Initializer)
@@ -1451,6 +1457,7 @@ public:
 	static void ModifyCompilationEnvironment(const FGlobalShaderPermutationParameters& InParameters, FShaderCompilerEnvironment& OutEnvironment)
 	{
 		OutEnvironment.SetDefine(TEXT("MERGE_EDIT_LAYERS"), 1);
+		UE::Landscape::ModifyShaderCompilerEnvironmentForDebug(OutEnvironment);
 	}
 
 	static void MergeEditLayers(FRDGBuilder& GraphBuilder, FParameters* InParameters, const FIntPoint& InTextureSize)
@@ -1492,6 +1499,7 @@ public:
 	static void ModifyCompilationEnvironment(const FGlobalShaderPermutationParameters& Parameters, FShaderCompilerEnvironment& OutEnvironment)
 	{
 		OutEnvironment.SetDefine(TEXT("STITCH_HEIGHTMAP"), 1);
+		UE::Landscape::ModifyShaderCompilerEnvironmentForDebug(OutEnvironment);
 	}
 
 	static void StitchHeightmap(FRDGBuilder& GraphBuilder, FParameters* InParameters)
@@ -1535,6 +1543,7 @@ public:
 	static void ModifyCompilationEnvironment(const FGlobalShaderPermutationParameters& Parameters, FShaderCompilerEnvironment& OutEnvironment)
 	{
 		OutEnvironment.SetDefine(TEXT("FINALIZE_HEIGHTMAP"), 1);
+		UE::Landscape::ModifyShaderCompilerEnvironmentForDebug(OutEnvironment);
 	}
 
 	static void FinalizeHeightmap(FRDGBuilder& GraphBuilder, FParameters* InParameters)
@@ -1575,6 +1584,7 @@ public:
 	static void ModifyCompilationEnvironment(const FGlobalShaderPermutationParameters& Parameters, FShaderCompilerEnvironment& OutEnvironment)
 	{
 		OutEnvironment.SetDefine(TEXT("GENERATE_MIPS"), 1);
+		UE::Landscape::ModifyShaderCompilerEnvironmentForDebug(OutEnvironment);
 	}
 
 	static void GenerateSingleMip(FRDGBuilder& GraphBuilder, FParameters* InParameters)
@@ -1621,6 +1631,7 @@ public:
 	static void ModifyCompilationEnvironment(const FGlobalShaderPermutationParameters& InParameters, FShaderCompilerEnvironment& OutEnvironment)
 	{
 		OutEnvironment.SetDefine(TEXT("MERGE_EDIT_LAYERS"), 1);
+		UE::Landscape::ModifyShaderCompilerEnvironmentForDebug(OutEnvironment);
 	}
 
 	static void MergeEditLayers(FRDGBuilder& GraphBuilder, FParameters* InParameters)
@@ -1667,6 +1678,7 @@ public:
 	static void ModifyCompilationEnvironment(const FGlobalShaderPermutationParameters& Parameters, FShaderCompilerEnvironment& OutEnvironment)
 	{
 		OutEnvironment.SetDefine(TEXT("FINALIZE_WEIGHTMAP"), 1);
+		UE::Landscape::ModifyShaderCompilerEnvironmentForDebug(OutEnvironment);
 	}
 
 	static void FinalizeWeightmap(FRDGBuilder& GraphBuilder, FParameters* InParameters)
@@ -1709,6 +1721,7 @@ public:
 	static void ModifyCompilationEnvironment(const FGlobalShaderPermutationParameters& Parameters, FShaderCompilerEnvironment& OutEnvironment)
 	{
 		OutEnvironment.SetDefine(TEXT("GENERATE_MIPS"), 1);
+		UE::Landscape::ModifyShaderCompilerEnvironmentForDebug(OutEnvironment);
 	}
 
 	static void GenerateSingleMip(FRDGBuilder& GraphBuilder, FParameters* InParameters)
@@ -1760,6 +1773,7 @@ public:
 	{
 		OutEnvironment.SetDefine(TEXT("GENERATE_NORMALS"), 1);
 		OutEnvironment.CompilerFlags.Remove(CFLAG_HLSL2021 | CFLAG_PrecompileWithDXC);
+		UE::Landscape::ModifyShaderCompilerEnvironmentForDebug(OutEnvironment);
 	}
 
 	static void GenerateNormalsPS(FRDGEventName&& InRDGEventName, FRDGBuilder& GraphBuilder, FParameters* InParameters, const FIntPoint& InTextureSize)
@@ -1799,6 +1813,7 @@ public:
 	static void ModifyCompilationEnvironment(const FGlobalShaderPermutationParameters& InParameters, FShaderCompilerEnvironment& OutEnvironment)
 	{
 		OutEnvironment.SetDefine(TEXT("GENERATE_MIPS"), 1);
+		UE::Landscape::ModifyShaderCompilerEnvironmentForDebug(OutEnvironment);
 	}
 
 	static void GenerateMipsPS(FRDGBuilder& GraphBuilder, FParameters* InParameters, const FIntPoint& InTextureSize)
@@ -1822,16 +1837,16 @@ IMPLEMENT_GLOBAL_SHADER(FLandscapeEditLayersHeightmapsGenerateMipsPS, "/Engine/P
 // ----------------------------------------------------------------------------------
 // /Engine/Private/Landscape/LandscapeEditLayersWeightmaps.usf shaders :
 
-class FLandscapeEditLayersWeightmapsPerformLegacyWeightBlendingPS : public FGlobalShader
+class FLandscapeEditLayersWeightmapsPerformFinalWeightBlendingPS : public FGlobalShader
 {
-	DECLARE_GLOBAL_SHADER(FLandscapeEditLayersWeightmapsPerformLegacyWeightBlendingPS);
-	SHADER_USE_PARAMETER_STRUCT(FLandscapeEditLayersWeightmapsPerformLegacyWeightBlendingPS, FGlobalShader);
+	DECLARE_GLOBAL_SHADER(FLandscapeEditLayersWeightmapsPerformFinalWeightBlendingPS);
+	SHADER_USE_PARAMETER_STRUCT(FLandscapeEditLayersWeightmapsPerformFinalWeightBlendingPS, FGlobalShader);
 
 public:
 	BEGIN_SHADER_PARAMETER_STRUCT(FParameters, )
 		SHADER_PARAMETER(uint32, InTargetLayerIndex)
 		SHADER_PARAMETER(uint32, InNumTargetLayers)
-		SHADER_PARAMETER_RDG_BUFFER_SRV(StructuredBuffer<UE::Landscape::FWeightmapTargetLayerInfo>, InTargetLayerInfos)
+		SHADER_PARAMETER_RDG_BUFFER_SRV(StructuredBuffer<UE::Landscape::FFinalWeightBlendingTargetLayerInfo>, InFinalWeightBlendingTargetLayerInfos)
 		SHADER_PARAMETER_RDG_TEXTURE_SRV(Texture2DArray<float4>, InCurrentEditLayerWeightmaps)
 		RENDER_TARGET_BINDING_SLOTS()
 	END_SHADER_PARAMETER_STRUCT()
@@ -1843,13 +1858,14 @@ public:
 
 	static void ModifyCompilationEnvironment(const FGlobalShaderPermutationParameters& InParameters, FShaderCompilerEnvironment& OutEnvironment)
 	{
-		OutEnvironment.SetDefine(TEXT("PERFORM_LEGACY_WEIGHT_BLENDING"), 1);
+		OutEnvironment.SetDefine(TEXT("PERFORM_FINAL_WEIGHT_BLENDING"), 1);
+		UE::Landscape::ModifyShaderCompilerEnvironmentForDebug(OutEnvironment);
 	}
 
-	static void PerformLegacyWeightBlendingPS(FRDGEventName&& InRDGEventName, FRDGBuilder& GraphBuilder, FParameters* InParameters, const FIntPoint& InTextureSize)
+	static void PerformFinalWeightBlendingPS(FRDGEventName&& InRDGEventName, FRDGBuilder& GraphBuilder, FParameters* InParameters, const FIntPoint& InTextureSize)
 	{
 		FGlobalShaderMap* ShaderMap = GetGlobalShaderMap(GMaxRHIFeatureLevel);
-		TShaderMapRef<FLandscapeEditLayersWeightmapsPerformLegacyWeightBlendingPS> PixelShader(ShaderMap);
+		TShaderMapRef<FLandscapeEditLayersWeightmapsPerformFinalWeightBlendingPS> PixelShader(ShaderMap);
 
 		FPixelShaderUtils::AddFullscreenPass(
 			GraphBuilder,
@@ -1862,7 +1878,7 @@ public:
 	}
 };
 
-IMPLEMENT_GLOBAL_SHADER(FLandscapeEditLayersWeightmapsPerformLegacyWeightBlendingPS, "/Engine/Private/Landscape/LandscapeEditLayersWeightmaps.usf", "PerformLegacyWeightBlendingPS", SF_Pixel);
+IMPLEMENT_GLOBAL_SHADER(FLandscapeEditLayersWeightmapsPerformFinalWeightBlendingPS, "/Engine/Private/Landscape/LandscapeEditLayersWeightmaps.usf", "PerformFinalWeightBlendingPS", SF_Pixel);
 
 class FLandscapeEditLayersWeightmapsPackWeightmapPS : public FGlobalShader
 {
@@ -1888,6 +1904,7 @@ public:
 	static void ModifyCompilationEnvironment(const FGlobalShaderPermutationParameters& InParameters, FShaderCompilerEnvironment& OutEnvironment)
 	{
 		OutEnvironment.SetDefine(TEXT("PACK_WEIGHTMAP"), 1);
+		UE::Landscape::ModifyShaderCompilerEnvironmentForDebug(OutEnvironment);
 	}
 
 	static void PackWeightmapPS(FRDGBuilder& GraphBuilder, FParameters* InParameters, const FIntRect& InTextureRect)
@@ -1927,6 +1944,7 @@ public:
 	static void ModifyCompilationEnvironment(const FGlobalShaderPermutationParameters& InParameters, FShaderCompilerEnvironment& OutEnvironment)
 	{
 		OutEnvironment.SetDefine(TEXT("GENERATE_MIPS"), 1);
+		UE::Landscape::ModifyShaderCompilerEnvironmentForDebug(OutEnvironment);
 	}
 
 	static void GenerateMipsPS(FRDGBuilder& GraphBuilder, FParameters* InParameters, const FIntPoint& InTextureSize)
@@ -1975,6 +1993,7 @@ public:
 	static void ModifyCompilationEnvironment(const FGlobalShaderPermutationParameters& Parameters, FShaderCompilerEnvironment& OutEnvironment)
 	{
 		OutEnvironment.SetDefine(TEXT("COPY_QUADS_MULTISOURCE"), 1);
+		UE::Landscape::ModifyShaderCompilerEnvironmentForDebug(OutEnvironment);
 	}
 };
 
@@ -2004,6 +2023,7 @@ public:
 	static void ModifyCompilationEnvironment(const FGlobalShaderPermutationParameters& Parameters, FShaderCompilerEnvironment& OutEnvironment)
 	{
 		OutEnvironment.SetDefine(TEXT("COPY_QUADS_MULTISOURCE"), 1);
+		UE::Landscape::ModifyShaderCompilerEnvironmentForDebug(OutEnvironment);
 	}
 };
 
@@ -2512,10 +2532,7 @@ ELandscapeEditLayersMergeMode ALandscape::GetEditLayersMergeMode() const
 
 	if (DesiredMode == ELandscapeEditLayersMergeMode::LocalMerge)
 	{
-		const ALandscape* Landscape = GetLandscapeActor();
-		check(Landscape);
-
-		for (const FLandscapeLayer& Layer : Landscape->LandscapeEditLayers)
+		for (const FLandscapeLayer& Layer : LandscapeEditLayers)
 		{
 			// No BP brush is supported for local merge of edit layers : 
 			if (Layer.EditLayer->IsVisible() && Algo::AnyOf(Layer.Brushes, [](const FLandscapeLayerBrush& Brush) { return (Brush.GetBrush() != nullptr) && Brush.GetBrush()->IsVisible(); }))
@@ -5625,7 +5642,8 @@ TArray<UE::Landscape::EditLayers::FEditLayerRendererState> FLandscapeLayerBrush:
 }
 
 void ULandscapeEditLayerPersistent::GetRendererStateInfo(const UE::Landscape::EditLayers::FMergeContext* InMergeContext,
-	UE::Landscape::EditLayers::FEditLayerTargetTypeState& OutSupportedTargetTypeState, UE::Landscape::EditLayers::FEditLayerTargetTypeState& OutEnabledTargetTypeState, TArray<TBitArray<>>& OutTargetLayerGroups) const
+	UE::Landscape::EditLayers::FEditLayerTargetTypeState& OutSupportedTargetTypeState, UE::Landscape::EditLayers::FEditLayerTargetTypeState& OutEnabledTargetTypeState,
+	TArray<UE::Landscape::EditLayers::FTargetLayerGroup>& OutTargetLayerGroups) const
 {
 	using namespace UE::Landscape::EditLayers;
 
@@ -5635,6 +5653,9 @@ void ULandscapeEditLayerPersistent::GetRendererStateInfo(const UE::Landscape::Ed
 
 	// Compute the default state of each target type : 
 	OutEnabledTargetTypeState = FEditLayerTargetTypeState(InMergeContext, GetEnabledTargetTypeMask(), OutSupportedTargetTypeState.GetActiveWeightmapBitIndices());
+
+	// Build the target layer groups for this renderer. Since we use the generic blend, we can use BuildGenericBlendTargetLayerGroups
+	OutTargetLayerGroups = InMergeContext->BuildGenericBlendTargetLayerGroups(OutSupportedTargetTypeState.GetActiveWeightmapBitIndices());
 }
 
 UE::Landscape::EditLayers::ERenderFlags ULandscapeEditLayerPersistent::GetRenderFlags(const UE::Landscape::EditLayers::FMergeContext* InMergeContext) const
@@ -6669,7 +6690,7 @@ bool ULandscapeWeightmapWeightBlendedLayersRenderer::RenderLayer(UE::Landscape::
 	WriteRT->TransitionTo(ERHIAccess::RTV, RDGBuilderRecorder);
 	ReadRT->TransitionTo(ERHIAccess::SRVMask, RDGBuilderRecorder);
 
-	TArray<FWeightmapTargetLayerInfo> WeightmapTargetLayerInfos;
+	TArray<FFinalWeightBlendingTargetLayerInfo> WeightmapTargetLayerInfos;
 	WeightmapTargetLayerInfos.Reserve(RenderParams.TargetLayerGroupLayerInfos.Num());
 	Algo::Transform(RenderParams.TargetLayerGroupLayerInfos, WeightmapTargetLayerInfos, [RenderBatch, MergeRenderContext = RenderParams.MergeRenderContext](ULandscapeLayerInfoObject* InLayerInfo)
 		{
@@ -6679,12 +6700,12 @@ bool ULandscapeWeightmapWeightBlendedLayersRenderer::RenderLayer(UE::Landscape::
 			// Skip weight-blended layers that are not rendered in the batch, because their texture slice might contain garbage :
 			const bool bRenderedInBatch = RenderBatch->TargetLayerBitIndices[TargetLayerIndex];
 
-			FWeightmapTargetLayerInfo WeightmapTargetLayerInfo;
+			FFinalWeightBlendingTargetLayerInfo WeightmapTargetLayerInfo;
 			if (!bRenderedInBatch)
 			{
 				WeightmapTargetLayerInfo.Flags |= EWeightmapTargetLayerFlags::Skip;
 			}
-			if (!InLayerInfo->bNoWeightBlend)
+			if (InLayerInfo->GetBlendMethod() == ELandscapeTargetLayerBlendMethod::FinalWeightBlending)
 			{
 				WeightmapTargetLayerInfo.Flags |= EWeightmapTargetLayerFlags::IsWeightBlended;
 			}
@@ -6707,7 +6728,7 @@ bool ULandscapeWeightmapWeightBlendedLayersRenderer::RenderLayer(UE::Landscape::
 			FRDGTextureRef CurrentEditLayerTextureRef = GraphBuilder.RegisterExternalTexture(CreateRenderTarget(CurrentEditLayerResource->TextureRHI, TEXT("CurrentEditLayerTexture")));
 			FRDGTextureSRVRef CurrentEditLayerTextureSRVRef = GraphBuilder.CreateSRV(FRDGTextureSRVDesc::Create(CurrentEditLayerTextureRef));
 
-			FRDGBufferRef TargetLayerInfosBuffer = CreateStructuredBuffer(GraphBuilder, TEXT("LandscapeTargetLayerInfosBuffer"), WeightmapTargetLayerInfos);
+			FRDGBufferRef TargetLayerInfosBuffer = CreateStructuredBuffer(GraphBuilder, TEXT("LandscapeFinalWeightBlendingTargetLayerInfosBuffer"), WeightmapTargetLayerInfos);
 			FRDGBufferSRVRef TargetLayerInfosBufferSRV = GraphBuilder.CreateSRV(FRDGBufferSRVDesc(TargetLayerInfosBuffer));
 
 			const int32 NumTargetLayers = TargetLayerNames.Num();
@@ -6715,17 +6736,17 @@ bool ULandscapeWeightmapWeightBlendedLayersRenderer::RenderLayer(UE::Landscape::
 			{
 				if (!EnumHasAnyFlags(WeightmapTargetLayerInfos[TargetLayerIndex].Flags, EWeightmapTargetLayerFlags::Skip))
 				{ 
-					FLandscapeEditLayersWeightmapsPerformLegacyWeightBlendingPS::FParameters* PSParams = GraphBuilder.AllocParameters<FLandscapeEditLayersWeightmapsPerformLegacyWeightBlendingPS::FParameters>();
-					PSParams->RenderTargets[0] = FRenderTargetBinding(OutputTextureRef, ERenderTargetLoadAction::ENoAction, /*InMipIndex = */0, /*InArraySlice = */TargetLayerIndex);
-					PSParams->InTargetLayerIndex = TargetLayerIndex;
-					PSParams->InNumTargetLayers = NumTargetLayers;
-					PSParams->InTargetLayerInfos = TargetLayerInfosBufferSRV;
-					PSParams->InCurrentEditLayerWeightmaps = CurrentEditLayerTextureSRVRef;
-
-					FLandscapeEditLayersWeightmapsPerformLegacyWeightBlendingPS::PerformLegacyWeightBlendingPS(
-						RDG_EVENT_NAME("LegacyWeightBlending(%s) -> %s", *TargetLayerNames[TargetLayerIndex].ToString(), *OutputResourceName),
-						GraphBuilder, PSParams, EffectiveTextureSize);
-				}
+				    FLandscapeEditLayersWeightmapsPerformFinalWeightBlendingPS::FParameters* PSParams = GraphBuilder.AllocParameters<FLandscapeEditLayersWeightmapsPerformFinalWeightBlendingPS::FParameters>();
+				    PSParams->RenderTargets[0] = FRenderTargetBinding(OutputTextureRef, ERenderTargetLoadAction::ENoAction, /*InMipIndex = */0, /*InArraySlice = */TargetLayerIndex);
+				    PSParams->InTargetLayerIndex = TargetLayerIndex;
+				    PSParams->InNumTargetLayers = NumTargetLayers;
+				    PSParams->InFinalWeightBlendingTargetLayerInfos = TargetLayerInfosBufferSRV;
+				    PSParams->InCurrentEditLayerWeightmaps = CurrentEditLayerTextureSRVRef;
+    
+				    FLandscapeEditLayersWeightmapsPerformFinalWeightBlendingPS::PerformFinalWeightBlendingPS(
+					    RDG_EVENT_NAME("FinalWeightBlending(%s) -> %s", *TargetLayerNames[TargetLayerIndex].ToString(), *OutputResourceName),
+					    GraphBuilder, PSParams, EffectiveTextureSize);
+			    }
 			}
 		};
 
@@ -6797,7 +6818,7 @@ UE::Landscape::EditLayers::FMergeRenderContext ALandscape::PrepareEditLayersMerg
 		// Then append all the provided renderer states
 		CandidateEditLayerRendererStates.Append(InMergeRenderParams.EditLayerRendererStates);
 
-		// Legacy weight-blending requires an additional renderer at the end of the stack, to weight-blend the weightmaps needing it :
+		// Final weight-blending requires an additional renderer at the end of the stack, to weight-blend the weightmaps needing it :
 		if (!MergeRenderContext.IsHeightmapMerge())
 		{
 			ULandscapeWeightmapWeightBlendedLayersRenderer* WeightmapWeightBlendedLayersRenderer = ULandscapeWeightmapWeightBlendedLayersRenderer::StaticClass()->GetDefaultObject<ULandscapeWeightmapWeightBlendedLayersRenderer>();
@@ -6827,7 +6848,6 @@ UE::Landscape::EditLayers::FMergeRenderContext ALandscape::PrepareEditLayersMerg
 		FinalEditLayerRendererStates = CandidateEditLayerRendererStates;
 
 		// Only one group in the case of heightmap: 
-		const FName TargetLayerName = FName("Height");
 		FinalTargetLayerGroups = { TBitArray<>(true, 1) };
 	}
 	else
@@ -6848,10 +6868,10 @@ UE::Landscape::EditLayers::FMergeRenderContext ALandscape::PrepareEditLayersMerg
 
 			// List of all supported target layer groups for this renderer :
 			TArray<TBitArray<>> RelevantTargetLayerGroupsForThisRenderer;
-			for (const TBitArray<>& TargetLayerGroup : EditLayerRendererState.GetTargetLayerGroups())
+			for (const FTargetLayerGroup& TargetLayerGroup : EditLayerRendererState.GetTargetLayerGroups())
 			{
 				// Only retain the target layers that have passed the previous test (valid and enabled):
-				TBitArray<> FilteredTargetLayerGroupBitIndices = TargetLayerGroup;
+				TBitArray<> FilteredTargetLayerGroupBitIndices = TargetLayerGroup.GetWeightmapTargetLayerBitIndices();
 				FilteredTargetLayerGroupBitIndices.CombineWithBitwiseAND(RendererStateEnabledTargetLayerBitIndices, EBitwiseOperatorFlags::MinSize);
 				if (FilteredTargetLayerGroupBitIndices.Find(true) != INDEX_NONE)
 				{
@@ -10162,7 +10182,7 @@ void ALandscape::PrepareLayersWeightmapsLocalMergeRenderThreadData(const FUpdate
 							if (PaintLayerInfoIndex == nullptr)
 							{
 								FLandscapeEditLayerWeightmapPaintLayerInfo NewPaintLayerInfo;
-								if (!ComponentEditLayerAllocation.LayerInfo->bNoWeightBlend)
+								if (ComponentEditLayerAllocation.LayerInfo->GetBlendMethod() == ELandscapeTargetLayerBlendMethod::FinalWeightBlending)
 								{
 									NewPaintLayerInfo.Flags = ELandscapeEditLayerWeightmapPaintLayerFlags::IsWeightBlended;
 								}
@@ -10825,7 +10845,7 @@ int32 ALandscape::PerformLayersWeightmapsGlobalMerge(FUpdateLayersContentContext
 						for (int32 LayerInfoSettingsIndex = 0; LayerInfoSettingsIndex < Info->Layers.Num(); ++LayerInfoSettingsIndex)
 						{
 							const FLandscapeInfoLayerSettings& LayerInfo = Info->Layers[LayerInfoSettingsIndex];
-							WeightmapLayerWeightBlend[LayerInfoSettingsIndex + 1] = LayerInfo.LayerInfoObj != nullptr ? (LayerInfo.LayerInfoObj->bNoWeightBlend ? 0.0f : 1.0f) : 1.0f;
+							WeightmapLayerWeightBlend[LayerInfoSettingsIndex + 1] = (LayerInfo.LayerInfoObj != nullptr) ? ((LayerInfo.LayerInfoObj->GetBlendMethod() == ELandscapeTargetLayerBlendMethod::FinalWeightBlending) ? 1.0f : 0.0f) : 1.0f;
 						}
 
 						WeightmapLayerWeightBlend[0] = 0.0f; // Blend of Visibility 
@@ -11209,7 +11229,7 @@ int32 ALandscape::PerformLayersWeightmapsBatchedMerge(FUpdateLayersContentContex
 					const FWeightmapResolveInfo& WeightmapResolveInfo = It.Value;
 					FRDGTextureRef DestinationTextureRef = GraphBuilder.RegisterExternalTexture(CreateRenderTarget(WeightmapResolveInfo.TextureResource->TextureRHI, TEXT("DestinationTexture")));
 
-					RDG_EVENT_SCOPE(GraphBuilder, "%sFinalize %s", WeightmapResolveInfo.bIsFinalResolve ? TEXT("") : TEXT("(partially)"), *WeightmapResolveInfo.TextureResourceDebugName);
+					RDG_EVENT_SCOPE(GraphBuilder, "%sFinalize %s", WeightmapResolveInfo.bIsFinalResolve ? TEXT("") : TEXT("(partially) "), *WeightmapResolveInfo.TextureResourceDebugName);
 					{
 						RDG_EVENT_SCOPE(GraphBuilder, "Pack %d channels %s", FMath::CountBits(WeightmapResolveInfo.ChannelMask), WeightmapResolveInfo.bIsAdditiveResolve ? TEXT("(additive)") : TEXT(""));
 
diff --git a/Engine/Source/Runtime/Landscape/Private/LandscapeLayerInfoObject.cpp b/Engine/Source/Runtime/Landscape/Private/LandscapeLayerInfoObject.cpp
new file mode 100644
--- /dev/null
+++ b/Engine/Source/Runtime/Landscape/Private/LandscapeLayerInfoObject.cpp
@@ -0,0 +1,181 @@
+// Copyright Epic Games, Inc. All Rights Reserved.
+
+#include "LandscapeLayerInfoObject.h"
+
+#include "Landscape.h"
+#include "LandscapeInfo.h"
+#include "LandscapeProxy.h"
+#include "LandscapeSettings.h"
+#include "UObject/UObjectIterator.h"
+#include "UObject/FortniteMainBranchObjectVersion.h"
+#include "Engine/World.h"
+
+
+ULandscapeLayerInfoObject::ULandscapeLayerInfoObject(const FObjectInitializer& ObjectInitializer)
+	: Super(ObjectInitializer)
+{
+	if (!IsTemplate())
+	{
+		// Assign initial random LayerUsageDebugColor
+		LayerUsageDebugColor = GenerateLayerUsageDebugColor();
+
+		// The default blend method comes from the settings :
+		BlendMethod = GetDefault<ULandscapeSettings>()->GetTargetLayerDefaultBlendMethod();
+	}
+}
+
+FLinearColor ULandscapeLayerInfoObject::GenerateLayerUsageDebugColor() const
+{
+	uint8 Hash[20];
+	FString PathNameString = GetPathName();
+	FSHA1::HashBuffer(*PathNameString, PathNameString.Len() * sizeof(PathNameString[0]), Hash);
+
+	return FLinearColor(float(Hash[0]) / 255.f, float(Hash[1]) / 255.f, float(Hash[2]) / 255.f, 1.f);
+}
+
+void ULandscapeLayerInfoObject::SetBlendMethod(ELandscapeTargetLayerBlendMethod InBlendMethod, bool bInModify)
+{
+	const bool bHasValueChanged = (InBlendMethod != BlendMethod);
+	if (bHasValueChanged)
+	{
+		if (bInModify)
+		{
+			Modify();
+		}
+		BlendMethod = InBlendMethod;
+	}
+
+	// TODO [jonathan.bard] : send event for data change
+}
+
+void ULandscapeLayerInfoObject::SetBlendGroup(const FName& InBlendGroup, bool bInModify)
+{
+	const bool bHasValueChanged = (InBlendGroup != BlendGroup);
+	if (bHasValueChanged)
+	{
+		if (bInModify)
+		{
+			Modify();
+		}
+		BlendGroup = InBlendGroup;
+	}
+
+	// TODO [jonathan.bard] : send event for data change
+}
+
+void ULandscapeLayerInfoObject::Serialize(FArchive& Ar)
+{
+	Super::Serialize(Ar);
+
+	Ar.UsingCustomVersion(FFortniteMainBranchObjectVersion::GUID);
+}
+
+#if WITH_EDITOR
+
+void ULandscapeLayerInfoObject::PostEditChangeProperty(FPropertyChangedEvent& PropertyChangedEvent)
+{
+	Super::PostEditChangeProperty(PropertyChangedEvent);
+
+	const FName MemberPropertyName = PropertyChangedEvent.MemberProperty ? PropertyChangedEvent.MemberProperty->GetFName() : NAME_None;
+
+	// TODO [jonathan.bard] : automate the data changes via a data change event, so that we don't have to perform ad-hoc fixups like this 
+
+	if ((MemberPropertyName == GET_MEMBER_NAME_CHECKED(ULandscapeLayerInfoObject, PhysMaterial))
+		|| (MemberPropertyName == GET_MEMBER_NAME_CHECKED(ULandscapeLayerInfoObject, MinimumCollisionRelevanceWeight)))
+	{
+		for (TObjectIterator<ALandscapeProxy> It(/*AdditionalExclusionFlags = */RF_ClassDefaultObject, /*bIncludeDerivedClasses = */true, /*InInternalExclusionFlags = */EInternalObjectFlags::Garbage); It; ++It)
+		{
+			ALandscapeProxy* Proxy = *It;
+			if (Proxy->GetWorld() && !Proxy->GetWorld()->IsPlayInEditor())
+			{
+				ULandscapeInfo* Info = Proxy->GetLandscapeInfo();
+				if (Info)
+				{
+					for (int32 i = 0; i < Info->Layers.Num(); ++i)
+					{
+						if (Info->Layers[i].LayerInfoObj == this)
+						{
+							Proxy->ChangedPhysMaterial();
+							break;
+						}
+					}
+				}
+			}
+		}
+	}
+	else if (MemberPropertyName == GET_MEMBER_NAME_CHECKED(ULandscapeLayerInfoObject, LayerUsageDebugColor))
+	{
+		LayerUsageDebugColor.A = 1.0f;
+		for (TObjectIterator<ALandscapeProxy> It(/*AdditionalExclusionFlags = */RF_ClassDefaultObject, /*bIncludeDerivedClasses = */true, /*InInternalExclusionFlags = */EInternalObjectFlags::Garbage); It; ++It)
+		{
+			ALandscapeProxy* Proxy = *It;
+			if (Proxy->GetWorld() && !Proxy->GetWorld()->IsPlayInEditor())
+			{
+				Proxy->MarkComponentsRenderStateDirty();
+			}
+		}
+	}
+	else if ((MemberPropertyName == GET_MEMBER_NAME_CHECKED(ULandscapeLayerInfoObject, SplineFalloffModulationTexture))
+		|| (MemberPropertyName == GET_MEMBER_NAME_CHECKED(ULandscapeLayerInfoObject, SplineFalloffModulationColorMask))
+		|| (MemberPropertyName == GET_MEMBER_NAME_CHECKED(ULandscapeLayerInfoObject, SplineFalloffModulationBias))
+		|| (MemberPropertyName == GET_MEMBER_NAME_CHECKED(ULandscapeLayerInfoObject, SplineFalloffModulationScale))
+		|| (MemberPropertyName == GET_MEMBER_NAME_CHECKED(ULandscapeLayerInfoObject, SplineFalloffModulationTiling)))
+	{
+		for (TObjectIterator<ULandscapeInfo> It(/*AdditionalExclusionFlags = */RF_ClassDefaultObject, /*bIncludeDerivedClasses = */true, /*InInternalExclusionFlags = */EInternalObjectFlags::Garbage); It; ++It)
+		{
+			if (ALandscape* Landscape = It->LandscapeActor.Get())
+			{
+				Landscape->OnLayerInfoSplineFalloffModulationChanged(this);
+			}
+		}
+	}
+	else if (MemberPropertyName == GET_MEMBER_NAME_CHECKED(ULandscapeLayerInfoObject, BlendMethod))
+	{
+		SetBlendMethod(BlendMethod, /*bInModify = */true);
+	}
+	else if (MemberPropertyName == GET_MEMBER_NAME_CHECKED(ULandscapeLayerInfoObject, BlendGroup))
+	{
+		SetBlendGroup(BlendGroup, /*bInModify = */true);
+	}
+}
+
+void ULandscapeLayerInfoObject::PostEditUndo()
+{
+	Super::PostEditUndo();
+
+	// TODO [jonathan.bard] : send event for data change and remove this object iterator nonsense : 
+
+	// Force the update of spline data for the (potentially) affected landscapes in case of an undo
+	for (TObjectIterator<ULandscapeInfo> It(/*AdditionalExclusionFlags = */RF_ClassDefaultObject, /*bIncludeDerivedClasses = */true, /*InInternalExclusionFlags = */EInternalObjectFlags::Garbage); It; ++It)
+	{
+		if (ALandscape* Landscape = It->LandscapeActor.Get())
+		{
+			Landscape->OnLayerInfoSplineFalloffModulationChanged(this);
+		}
+	}
+}
+
+void ULandscapeLayerInfoObject::PostLoad()
+{
+	Super::PostLoad();
+
+#if WITH_EDITOR
+	if (GetLinkerCustomVersion(FFortniteMainBranchObjectVersion::GUID) < FFortniteMainBranchObjectVersion::LandscapeAdvancedWeightBlending)
+	{
+PRAGMA_DISABLE_DEPRECATION_WARNINGS
+		BlendMethod = bNoWeightBlend_DEPRECATED ? ELandscapeTargetLayerBlendMethod::None : ELandscapeTargetLayerBlendMethod::FinalWeightBlending;
+PRAGMA_ENABLE_DEPRECATION_WARNINGS
+	}
+	
+	if (GIsEditor)
+	{
+		if (!HasAnyFlags(RF_Standalone))
+		{
+			SetFlags(RF_Standalone);
+		}
+		Hardness = FMath::Clamp<float>(Hardness, 0.0f, 1.0f);
+	}
+#endif // WITH_EDITOR
+}
+
+#endif // WITH_EDITOR
\ No newline at end of file
diff --git a/Engine/Source/Runtime/Landscape/Private/LandscapeProxy.cpp b/Engine/Source/Runtime/Landscape/Private/LandscapeProxy.cpp
--- a/Engine/Source/Runtime/Landscape/Private/LandscapeProxy.cpp
+++ b/Engine/Source/Runtime/Landscape/Private/LandscapeProxy.cpp
@@ -17,7 +17,7 @@
 
 // ----------------------------------------------------------------------------------
 
-LANDSCAPE_API FLandscapeImportLayerInfo::FLandscapeImportLayerInfo(const FLandscapeInfoLayerSettings& InLayerSettings)
+FLandscapeImportLayerInfo::FLandscapeImportLayerInfo(const FLandscapeInfoLayerSettings& InLayerSettings)
 	: LayerName(InLayerSettings.GetLayerName())
 	, LayerInfo(InLayerSettings.LayerInfoObj)
 	, SourceFilePath(InLayerSettings.GetTargetLayerSettings().ReimportLayerFilePath)
diff --git a/Engine/Source/Runtime/Landscape/Private/LandscapeSplineRaster.cpp b/Engine/Source/Runtime/Landscape/Private/LandscapeSplineRaster.cpp
--- a/Engine/Source/Runtime/Landscape/Private/LandscapeSplineRaster.cpp
+++ b/Engine/Source/Runtime/Landscape/Private/LandscapeSplineRaster.cpp
@@ -448,7 +448,8 @@ void RasterizeControlPointAlpha(int32& MinX, int32& MinY, int32& MaxX, int32& Ma
 		Rasterizer.DrawTriangle(FalloffRight0, Left1, FalloffLeft1, FalloffRight0Pos, Left1Pos, FalloffLeft1Pos, false);
 	}
 
-	LandscapeEdit.SetAlphaData(LayerInfo, MinX, MinY, MaxX, MaxY, Data.GetData(), 0, ELandscapeLayerPaintingRestriction::None, !LayerInfo->bNoWeightBlend, false);
+	LandscapeEdit.SetAlphaData(LayerInfo, MinX, MinY, MaxX, MaxY, Data.GetData(), 0, ELandscapeLayerPaintingRestriction::None, 
+		/*bWeightAdjust = */(LayerInfo->GetBlendMethod() == ELandscapeTargetLayerBlendMethod::FinalWeightBlending), /*bTotalWeightAdjust = */false);
 
 	LandscapeEdit.GetComponentsInRegion(MinX, MinY, MaxX, MaxY, &ModifiedComponents);
 }
@@ -571,7 +572,8 @@ void RasterizeSegmentAlpha(int32& MinX, int32& MinY, int32& MaxX, int32& MaxY, F
 		Rasterizer.DrawTriangle(FalloffRight0, Right1, FalloffRight1, FalloffRight0Pos, Right1Pos, FalloffRight1Pos, false);
 	}
 
-	LandscapeEdit.SetAlphaData(LayerInfo, MinX, MinY, MaxX, MaxY, Data.GetData(), 0, ELandscapeLayerPaintingRestriction::None, !LayerInfo->bNoWeightBlend, false);
+	LandscapeEdit.SetAlphaData(LayerInfo, MinX, MinY, MaxX, MaxY, Data.GetData(), 0, ELandscapeLayerPaintingRestriction::None, 
+		/*bWeightAdjust = */(LayerInfo->GetBlendMethod() == ELandscapeTargetLayerBlendMethod::FinalWeightBlending), /*bTotalWeightAdjust = */false);
 
 	LandscapeEdit.GetComponentsInRegion(MinX, MinY, MaxX, MaxY, &ModifiedComponents);
 }
diff --git a/Engine/Source/Runtime/Landscape/Private/LandscapeUtils.cpp b/Engine/Source/Runtime/Landscape/Private/LandscapeUtils.cpp
--- a/Engine/Source/Runtime/Landscape/Private/LandscapeUtils.cpp
+++ b/Engine/Source/Runtime/Landscape/Private/LandscapeUtils.cpp
@@ -10,6 +10,8 @@
 #include "Algo/Transform.h"
 #include "RenderGraphBuilder.h"
 #include "TextureResource.h"
+#include "ShaderCore.h"
+#include "ShaderCompilerCore.h"
 
 #if WITH_EDITOR
 #include "Algo/AllOf.h"
@@ -34,6 +36,12 @@ extern const size_t ChannelOffsets[4];
 extern UNREALED_API UEditorEngine* GEditor;
 #endif // WITH_EDITOR
 
+static TAutoConsoleVariable<bool> CVarLandscapeDebugShaders(
+	TEXT("landscape.DebugShaders"),
+	false,
+	TEXT("De-optimizes landscape shaders for easier debugging (use in conjunction with r.Shaders.Symbols=1 to dump shader symbols)"));
+
+
 namespace UE::Landscape
 {
 
@@ -45,6 +53,16 @@ bool DoesPlatformSupportEditLayers(EShaderPlatform InShaderPlatform)
 		&& !IsMobilePlatform(InShaderPlatform);
 }
 
+void ModifyShaderCompilerEnvironmentForDebug(FShaderCompilerEnvironment& OutEnvironment)
+{
+#if !UE_BUILD_SHIPPING 
+	if (CVarLandscapeDebugShaders.GetValueOnAnyThread())
+	{
+		OutEnvironment.CompilerFlags.Add(CFLAG_Debug);
+	}
+#endif // !UE_BUILD_SHIPPING 
+}
+
 ELandscapeToolTargetTypeFlags GetLandscapeToolTargetTypeAsFlags(ELandscapeToolTargetType InTargetType)
 {
 	uint8 TargetTypeValue = static_cast<uint8>(InTargetType);
diff --git a/Engine/Source/Runtime/Landscape/Public/LandscapeBlueprintBrushBase.h b/Engine/Source/Runtime/Landscape/Public/LandscapeBlueprintBrushBase.h
--- a/Engine/Source/Runtime/Landscape/Public/LandscapeBlueprintBrushBase.h
+++ b/Engine/Source/Runtime/Landscape/Public/LandscapeBlueprintBrushBase.h
@@ -110,7 +110,8 @@ public:
 	//~ Begin ILandscapeEditLayerRenderer implementation
 	LANDSCAPE_EDIT_LAYERS_BATCHED_MERGE_EXPERIMENTAL
 	LANDSCAPE_API virtual void GetRendererStateInfo(const UE::Landscape::EditLayers::FMergeContext* InMergeContext,
-		UE::Landscape::EditLayers::FEditLayerTargetTypeState& OutSupportedTargetTypeState, UE::Landscape::EditLayers::FEditLayerTargetTypeState& OutEnabledTargetTypeState, TArray<TBitArray<>>& OutTargetLayerGroups) const override;
+		UE::Landscape::EditLayers::FEditLayerTargetTypeState& OutSupportedTargetTypeState, UE::Landscape::EditLayers::FEditLayerTargetTypeState& OutEnabledTargetTypeState, 
+		TArray<UE::Landscape::EditLayers::FTargetLayerGroup>& OutTargetLayerGroups) const override;
 	LANDSCAPE_EDIT_LAYERS_BATCHED_MERGE_EXPERIMENTAL
 	virtual UE::Landscape::EditLayers::ERenderFlags GetRenderFlags(const UE::Landscape::EditLayers::FMergeContext* InMergeContext) const override { return UE::Landscape::EditLayers::ERenderFlags::RenderMode_Immediate; }
 	LANDSCAPE_EDIT_LAYERS_BATCHED_MERGE_EXPERIMENTAL
diff --git a/Engine/Source/Runtime/Landscape/Public/LandscapeEdit.h b/Engine/Source/Runtime/Landscape/Public/LandscapeEdit.h
--- a/Engine/Source/Runtime/Landscape/Public/LandscapeEdit.h
+++ b/Engine/Source/Runtime/Landscape/Public/LandscapeEdit.h
@@ -230,6 +230,7 @@ struct FLandscapeEditDataInterface : public FLandscapeTextureDataInterface
 	LANDSCAPE_API void GetWeightDataFast(ULandscapeLayerInfoObject* LayerInfo, const int32 X1, const int32 Y1, const int32 X2, const int32 Y2, TMap<FIntPoint, uint8>& SparseData);
 	LANDSCAPE_API void GetWeightDataFast(ULandscapeLayerInfoObject* LayerInfo, const int32 X1, const int32 Y1, const int32 X2, const int32 Y2, TMap<FIntPoint, TArray<uint8>>& SparseData);
 	// Updates weightmap for LayerInfo, optionally adjusting all other weightmaps.
+	// TODO [jonathan.bard] : Remove bWeightAdjust (and bTotalWeightAdjust, which is unused) as soon as we remove support for non-edit layers
 	LANDSCAPE_API void SetAlphaData(ULandscapeLayerInfoObject* LayerInfo, const int32 X1, const int32 Y1, const int32 X2, const int32 Y2, const uint8* Data, int32 Stride, ELandscapeLayerPaintingRestriction PaintingRestriction = ELandscapeLayerPaintingRestriction::None, bool bWeightAdjust = true, bool bTotalWeightAdjust = false);
 	// Updates weightmaps for all layers. Data points to packed data for all layers in the landscape info
 	LANDSCAPE_API void SetAlphaData(const TSet<ULandscapeLayerInfoObject*>& DirtyLayerInfos, const int32 X1, const int32 Y1, const int32 X2, const int32 Y2, const uint8* Data, int32 Stride, ELandscapeLayerPaintingRestriction PaintingRestriction = ELandscapeLayerPaintingRestriction::None);
@@ -334,9 +335,11 @@ private:
 	                          int32 ComponentIndexY, int32 SubIndexY, int32 SubY);
 
 	// counts the total influence of each weight-blended layer on this component
+	// TODO [jonathan.bard] : Remove when we get rid of non-edit layers (this is only used by ChooseReplacementLayer, which is only used by SetAlphaData, when bWeightAdjust == true)
 	inline TMap<const ULandscapeLayerInfoObject*, uint32> CountWeightBlendedLayerInfluence(int32 ComponentIndexX, int32 ComponentIndexY, TOptional<TArrayView<const uint8* const>> LayerDataPtrs);
 
 	// chooses a replacement layer to use when erasing from 100% influence on a texel
+	// TODO [jonathan.bard] : Remove when we get rid of non-edit layers (this is only used by SetAlphaData, when bWeightAdjust == true)
 	const ULandscapeLayerInfoObject* ChooseReplacementLayer(const ULandscapeLayerInfoObject* LayerInfo, int32 ComponentIndexX, int32 SubIndexX, int32 SubX, int32 ComponentIndexY, int32 SubIndexY, int32 SubY, TMap<FIntPoint, TMap<const ULandscapeLayerInfoObject*, uint32>>& LayerInfluenceCache, TArrayView<const uint8* const> LayerDataPtrs);
 };
 
@@ -498,18 +501,18 @@ struct FAlphamapAccessor
 		: LandscapeInfo(InLandscapeInfo)
 		, LandscapeEdit(InLandscapeInfo)
 		, LayerInfo(InLayerInfo)
-		, bBlendWeight(true)
+		, bPerformLegacyWeightBalancing(true)
 	{
 		// should be no Layer change during FAlphamapAccessor lifetime...
 		if (InLandscapeInfo && InLayerInfo)
 		{
 			if (LayerInfo == ALandscapeProxy::VisibilityLayer)
 			{
-				bBlendWeight = false;
+				bPerformLegacyWeightBalancing = false;
 			}
 			else
 			{
-				bBlendWeight = !LayerInfo->bNoWeightBlend;
+				bPerformLegacyWeightBalancing = (LayerInfo->GetBlendMethod() == ELandscapeTargetLayerBlendMethod::FinalWeightBlending);
 			}
 		}
 	}
@@ -551,7 +554,7 @@ struct FAlphamapAccessor
 				LandscapeComponent->RequestWeightmapUpdate();
 			}
 			
-			LandscapeEdit.SetAlphaData(LayerInfo, X1, Y1, X2, Y2, Data, 0, PaintingRestriction, bBlendWeight, bUseTotalNormalize);
+			LandscapeEdit.SetAlphaData(LayerInfo, X1, Y1, X2, Y2, Data, 0, PaintingRestriction, /*bWeightAdjust = */bPerformLegacyWeightBalancing, /*bTotalWeightAdjust = */bUseTotalNormalize);
 			//LayerInfo->IsReferencedFromLoadedData = true;
 		}
 	}
@@ -576,7 +579,7 @@ private:
 	FLandscapeEditDataInterface LandscapeEdit;
 	TSet<ULandscapeComponent*> ModifiedComponents;
 	ULandscapeLayerInfoObject* LayerInfo;
-	bool bBlendWeight;
+	bool bPerformLegacyWeightBalancing;
 };
 
 #endif
diff --git a/Engine/Source/Runtime/Landscape/Public/LandscapeEditLayerMergeContext.h b/Engine/Source/Runtime/Landscape/Public/LandscapeEditLayerMergeContext.h
--- a/Engine/Source/Runtime/Landscape/Public/LandscapeEditLayerMergeContext.h
+++ b/Engine/Source/Runtime/Landscape/Public/LandscapeEditLayerMergeContext.h
@@ -13,6 +13,10 @@ class ALandscape;
 class ULandscapeInfo;
 class ULandscapeLayerInfoObject;
 
+namespace UE::Landscape::EditLayers
+{
+	class FTargetLayerGroup;
+} // namespace UE::Landscape::EditLayers
 
 // ----------------------------------------------------------------------------------
 
@@ -29,6 +33,8 @@ class FMergeContext
 public:
 	friend ALandscape;
 
+	using FTargetLayerGroupsPerBlendingMethod = TStaticArray<TArray<FTargetLayerGroup>, static_cast<uint8>(ELandscapeTargetLayerBlendMethod::Count)>;
+
 	FMergeContext(ALandscape* InLandscape, bool bInIsHeightmapMerge, bool bInSkipProceduralRenderers);
 	virtual ~FMergeContext() = default;
 	FMergeContext(const FMergeContext& Other) = default;
@@ -208,6 +214,20 @@ public:
 	LANDSCAPE_EDIT_LAYERS_BATCHED_MERGE_EXPERIMENTAL
 	LANDSCAPE_API TBitArray<> BuildTargetLayerBitIndices(bool bInBitValue) const;
 
+	LANDSCAPE_EDIT_LAYERS_BATCHED_MERGE_EXPERIMENTAL
+	LANDSCAPE_API const FTargetLayerGroupsPerBlendingMethod& GetTargetLayerGroupsPerBlendingMethod() const { return TargetLayerGroupsPerBlendingMethod; }
+
+	/** 
+	 * Return the target layer groups for the renderers performing the generic blend. This concerns only the blend methods that are applied per renderer
+	 * @param InTargetLayerBitIndices bit mask of the target layers that are concerned by this blend
+	 */
+	LANDSCAPE_EDIT_LAYERS_BATCHED_MERGE_EXPERIMENTAL
+	LANDSCAPE_API TArray<FTargetLayerGroup> BuildGenericBlendTargetLayerGroups(const TBitArray<>& InTargetLayerBitIndices) const;
+
+private:
+	void FinalizeTargetLayerGroupsPerBlendingMethod();
+	TArray<FTargetLayerGroup> GatherTargetLayerGroupsForBlendMethod(ELandscapeTargetLayerBlendMethod InBlendMethod) const;
+
 protected:
 	// COMMENT [jonathan.bard] : purposefully don't use ELandscapeToolTargetType here as ELandscapeToolTargetType::Weightmap and ELandscapeToolTargetType::Visibility have to be processed together (because of weightmap packing, 
 	//  which means a visibility weightmap could be another channel of a texture which contains weightmap up to 3 other weightmaps, so we have to resolve the 4 channels altogether). 
@@ -241,6 +261,9 @@ protected:
 	/** List of valid target layers. If a target layer name is present here, it's because it has a valid landscape layer info object. Each bit in that bit array corresponds to an entry in AllTargetLayerNames */
 	TBitArray<> ValidTargetLayerBitIndices;
 
+	/** Contains the list of target layer groups for each weight-blending method as defined by the landscape layer info objects */
+	FTargetLayerGroupsPerBlendingMethod TargetLayerGroupsPerBlendingMethod;
+
 	// Index of the visibility layer in AllTargetLayerNames (always valid for weightmap merges)
 	int32 VisibilityTargetLayerIndex = INDEX_NONE;
 };
diff --git a/Engine/Source/Runtime/Landscape/Public/LandscapeEditLayerRenderer.h b/Engine/Source/Runtime/Landscape/Public/LandscapeEditLayerRenderer.h
--- a/Engine/Source/Runtime/Landscape/Public/LandscapeEditLayerRenderer.h
+++ b/Engine/Source/Runtime/Landscape/Public/LandscapeEditLayerRenderer.h
@@ -31,6 +31,7 @@ namespace UE::Landscape::EditLayers
 	struct FComponentMergeRenderInfo;
 	struct FRenderParams;
 	class FMergeRenderContext;
+	class FTargetLayerGroup;
 } // namespace UE::Landscape::EditLayers
 
 
@@ -286,7 +287,8 @@ public:
 	 */
 	LANDSCAPE_EDIT_LAYERS_BATCHED_MERGE_EXPERIMENTAL
 	virtual void GetRendererStateInfo(const UE::Landscape::EditLayers::FMergeContext* InMergeContext,
-		UE::Landscape::EditLayers::FEditLayerTargetTypeState& OutSupportedTargetTypeState, UE::Landscape::EditLayers::FEditLayerTargetTypeState& OutEnabledTargetTypeState, TArray<TBitArray<>>& OutTargetLayerGroups) const
+		UE::Landscape::EditLayers::FEditLayerTargetTypeState& OutSupportedTargetTypeState, UE::Landscape::EditLayers::FEditLayerTargetTypeState& OutEnabledTargetTypeState, 
+		TArray<UE::Landscape::EditLayers::FTargetLayerGroup>& OutTargetLayerGroups) const
 		PURE_VIRTUAL(ILandscapeEditLayerRenderer::GetRendererStateInfo, );
 
 	/**
diff --git a/Engine/Source/Runtime/Landscape/Public/LandscapeEditLayerRendererState.h b/Engine/Source/Runtime/Landscape/Public/LandscapeEditLayerRendererState.h
--- a/Engine/Source/Runtime/Landscape/Public/LandscapeEditLayerRendererState.h
+++ b/Engine/Source/Runtime/Landscape/Public/LandscapeEditLayerRendererState.h
@@ -5,6 +5,7 @@
 #include "CoreMinimal.h"
 #include "LandscapeEditTypes.h"
 #include "LandscapeEditLayerTargetTypeState.h"
+#include "LandscapeEditLayerTargetLayerGroup.h"
 #include "UObject/ScriptInterface.h"
 
 
@@ -155,7 +156,7 @@ public:
 	/** Returns the target layer groups associated with this renderer. A target layer group is a set of target layers (weightmaps) that depend on one another in order to produce the output target layers.
 	 This allows to implement "horizontal blending", where weightmaps can be blended with one another at each step of the landscape edit layers merge algorithm */
 	LANDSCAPE_EDIT_LAYERS_BATCHED_MERGE_EXPERIMENTAL
-	TArray<TBitArray<>> GetTargetLayerGroups() const { return TargetLayerGroups; }
+	TArray<FTargetLayerGroup> GetTargetLayerGroups() const { return TargetLayerGroups; }
 
 private:
 	void UpdateActiveTargetTypeState();
@@ -184,7 +185,7 @@ PRAGMA_ENABLE_EXPERIMENTAL_WARNINGS
 	/* List of groups of target layers that this renderer requires to be rendered together. All target layers listed in SupportedTargetTypeState must belong to 1 (and 1 only) target layer group.
 	 *  Each target layer group is a bit array for which each bit corresponds to an entry in FMergeContext's AllTargetLayerNames
 	 */
-	TArray<TBitArray<>> TargetLayerGroups;
+	TArray<FTargetLayerGroup> TargetLayerGroups;
 };
 
 #endif // WITH_EDITOR
diff --git a/Engine/Source/Runtime/Landscape/Public/LandscapeEditLayerTargetLayerGroup.h b/Engine/Source/Runtime/Landscape/Public/LandscapeEditLayerTargetLayerGroup.h
new file mode 100644
--- /dev/null
+++ b/Engine/Source/Runtime/Landscape/Public/LandscapeEditLayerTargetLayerGroup.h
@@ -0,0 +1,51 @@
+// Copyright Epic Games, Inc. All Rights Reserved.
+
+#pragma once
+
+#include "CoreMinimal.h"
+
+
+// ----------------------------------------------------------------------------------
+
+namespace UE::Landscape::EditLayers
+{
+
+#if WITH_EDITOR
+
+/** 
+ * FTargetLayerGroup describes a group of target layers that need to be rendered together, typically for weight-blending reasons
+ */
+class FTargetLayerGroup
+{
+public:
+	inline FTargetLayerGroup(const FName& InName, const TBitArray<>& InWeightmapTargetLayerBitIndices)
+		: Name(InName)
+		, WeightmapTargetLayerBitIndices(InWeightmapTargetLayerBitIndices)
+	{}
+
+	inline FName GetName() const 
+	{ 
+		return Name;
+	}
+
+	inline const TBitArray<>& GetWeightmapTargetLayerBitIndices() const 
+	{
+		return WeightmapTargetLayerBitIndices;
+	}
+
+private:
+	/** 
+	 * Identifier for this group of target layers
+	 */
+	FName Name;
+
+	/**
+	 * List of weightmaps that belong to this group
+	 *  Each bit in that bit array corresponds to an entry in FMergeContext's AllTargetLayerNames
+	 */
+	TBitArray<> WeightmapTargetLayerBitIndices;
+};
+
+#endif // WITH_EDITOR
+
+} //namespace UE::Landscape::EditLayers
diff --git a/Engine/Source/Runtime/Landscape/Public/LandscapeEditLayerTypes.h b/Engine/Source/Runtime/Landscape/Public/LandscapeEditLayerTypes.h
--- a/Engine/Source/Runtime/Landscape/Public/LandscapeEditLayerTypes.h
+++ b/Engine/Source/Runtime/Landscape/Public/LandscapeEditLayerTypes.h
@@ -73,6 +73,7 @@ enum class EWeightmapTargetLayerFlags : uint32
 	IsVisibilityLayer = (1 << 0), // This target layer is the visibility layer
 	IsWeightBlended = (1 << 1), // Blend the target layer's value with all the other target layers weights
 	Skip = (1 << 2), // This layer should be skipped from blending because it has not been rendered in this batch
+	IsPremultipliedAlphaWeightBlended = (1 << 3), // Blend the target layer's weight with other target layers' weights of the same target layer blend group, using the premultiplied alpha blend formula
 
 	None = 0
 };
@@ -81,13 +82,23 @@ ENUM_CLASS_FLAGS(EWeightmapTargetLayerFlags);
 
 // ----------------------------------------------------------------------------------
 
-// Must match FWeightmapTargetLayerInfo in LandscapeEditLayersWeightmaps.usf
-struct FWeightmapTargetLayerInfo
+// Must match FFinalWeightBlendingTargetLayerInfo in LandscapeEditLayersWeightmaps.usf
+struct FFinalWeightBlendingTargetLayerInfo
 {
 	EWeightmapTargetLayerFlags Flags = EWeightmapTargetLayerFlags::None; // Additional info about this target layer
 };
 
 
+// ----------------------------------------------------------------------------------
+
+// Must match FMergeEditLayerTargetLayerInfo in LandscapeEditLayersWeightmaps.usf
+struct FMergeEditLayerTargetLayerInfo
+{
+	EWeightmapTargetLayerFlags Flags = EWeightmapTargetLayerFlags::None; // Additional info about this target layer
+	int32 BlendGroupIndex = INDEX_NONE; // Defines the target layer blend group that this target layer belongs to in this blend operation. -1 if not applicable
+};
+
+
 // ----------------------------------------------------------------------------------
 
 // Defines how heightmaps should be blended (see GenericBlendLayer)
diff --git a/Engine/Source/Runtime/Landscape/Public/LandscapeEditTypes.h b/Engine/Source/Runtime/Landscape/Public/LandscapeEditTypes.h
--- a/Engine/Source/Runtime/Landscape/Public/LandscapeEditTypes.h
+++ b/Engine/Source/Runtime/Landscape/Public/LandscapeEditTypes.h
@@ -34,6 +34,14 @@ enum class ELandscapeToolTargetTypeFlags : uint8
 };
 ENUM_CLASS_FLAGS(ELandscapeToolTargetTypeFlags);
 
+UENUM()
+enum class ELandscapeTargetLayerBlendMethod : uint8
+{
+	None = 0 UMETA(DisplayName = "No Weight Blending", Tooltip = "The target layer's weight is unaffected by other target layers."),
+	FinalWeightBlending UMETA(DisplayName = "Weight Blending (Legacy)", Tooltip = "The target layer's weight is adjusted against the weights of the other (legacy) weight-blended target layers such that their sum is 0. Doesn't work well when combined with edit layers, as this (horizontal) blend is only applied at the end of the edit layers (vertical) blending."),
+	PremultipliedAlphaBlending UMETA(DisplayName = "Advanced Weight Blending", Tooltip = "The target layer's weight is adjusted against the weights of the other advanced weight-blended target layers such that their sum is 0. Takes into account edit layers : the blending is applied at each edit layer, such that the target layer's weight will always be weight-blended but will also retain its relative weight against other edit layers underneath."),
+	Count UMETA(Hidden), // Only the elements above actually count as proper target types
+};
 
 namespace UE::Landscape
 {
diff --git a/Engine/Source/Runtime/Landscape/Public/LandscapeSettings.h b/Engine/Source/Runtime/Landscape/Public/LandscapeSettings.h
--- a/Engine/Source/Runtime/Landscape/Public/LandscapeSettings.h
+++ b/Engine/Source/Runtime/Landscape/Public/LandscapeSettings.h
@@ -1,15 +1,16 @@
 // Copyright Epic Games, Inc. All Rights Reserved.
- 
+
 #pragma once
- 
+
 #include "CoreMinimal.h"
 #include "Engine/DeveloperSettings.h"
 #include "Delegates/Delegate.h"
+#include "LandscapeEditTypes.h"
 #include "LandscapeSettings.generated.h"
- 
+
 class UMaterialInterface;
 class ULandscapeLayerInfoObject;
- 
+
 UENUM()
 enum class ELandscapeDirtyingMode : uint8
 {
@@ -28,119 +29,125 @@ enum class ELandscapeDirtyingMode : uint8
 	 This mode is recommended for team collaboration as it provides the best features of the other two modes while ensuring that modified landscape actors are still saved and properly submitted to source control. */
 	InLandscapeModeAndUserTriggeredChanges
 };
- 
+
 UCLASS(config = Engine, defaultconfig, meta = (DisplayName = "Landscape"), MinimalAPI)
 class ULandscapeSettings : public UDeveloperSettings
 {
 	GENERATED_BODY()
- 
+
 public:
 	/** Returns true if landscape resolution should be constrained. */
 	bool IsLandscapeResolutionRestricted() const { return InRestrictiveMode(); }
- 
+
 	/** Returns true if blueprint landscape tools usage is allowed */
 	bool AreBlueprintToolsAllowed() const { return !InRestrictiveMode(); }
 	
 	/** Returns the current landscape resolution limit. */
 	int32 GetTotalResolutionLimit() const { return SideResolutionLimit * SideResolutionLimit; }
- 
+
 	bool InRestrictiveMode() const { return bRestrictiveMode; }
 	void SetRestrictiveMode(bool bEnabled) { bRestrictiveMode = bEnabled; }
- 
+
 	int32 GetSideResolutionLimit() const { return SideResolutionLimit; }
- 
+
 	float GetBrushSizeUIMax() const { return BrushSizeUIMax; }
 	float GetBrushSizeClampMax() const { return BrushSizeClampMax; }
- 
+
 	int32 GetHLODMaxTextureSize() const { return HLODMaxTextureSize; }
- 
+
 	float GetSplineIconWorldZOffset() const { return SplineIconWorldZOffset; }
 	float GetSplineIconScale() const { return SplineIconScale;  }
- 
+
 	/** Returns the default landscape material that should be used when creating a new landscape. */
 	TSoftObjectPtr<UMaterialInterface> GetDefaultLandscapeMaterial() const { return DefaultLandscapeMaterial; }
- 
+
 	/** Returns the default landscape layer info object that will be assigned to unset layers when creating a new landscape. */
 	TSoftObjectPtr<ULandscapeLayerInfoObject> GetDefaultLayerInfoObject() const { return DefaultLayerInfoObject; }
- 
+
 	bool GetShouldUpdateEditLayersDuringInteractiveChanges() const { return bShouldUpdateEditLayersDuringInteractiveChanges; }
  
 	bool ShouldDisplayTargetLayerThumbnails() const { return bDisplayTargetLayerThumbnails; }
- 
+
 	bool ShouldDisableTemporalAntiAliasingInLandscapeMode() const { return bDisableTemporalAntiAliasingInLandscapeMode; }
- 
+
+	ELandscapeTargetLayerBlendMethod GetTargetLayerDefaultBlendMethod() const { return TargetLayerDefaultBlendMethod; }
+
 #if WITH_EDITOR
 	//~ Begin UObject Interface.
 	virtual void PostEditChangeProperty(FPropertyChangedEvent& PropertyChangedEvent) override;
 	virtual void PreEditUndo() override;
 	virtual void PostEditUndo() override;
 #endif // WITH_EDITOR
- 
+
 public:
 	UPROPERTY(config, EditAnywhere, Category = "Edit Layers", meta=(UIMin = "1", UIMax = "32", ClampMin = "1", ClampMax = "32", ToolTip = "This option controls the maximum editing layers that can be added to a Landscape"))
 	int32 MaxNumberOfLayers = 8;
- 
-	UPROPERTY(EditAnywhere, config, Category = "Layers", meta = (ToolTip = 
+
+	UPROPERTY(EditAnywhere, config, Category = "Edit Layers", meta = (ToolTip = 
 		"When true, automatic edit layer creation pops up a dialog where the new layer can be reordered relative to other layers."))
 	bool bShowDialogForAutomaticLayerCreation = true;
- 
+
 	UPROPERTY(config, EditAnywhere, Category = "Configuration", meta=(ToolTip = "Maximum Dimension of Landscape in Components"))
 	int32 MaxComponents = 256;
- 
+
 	UPROPERTY(config, EditAnywhere, Category = "Configuration", meta = (UIMin = "1", UIMax = "1024", ClampMin = "1", ClampMax = "1024", ToolTip = "Maximum Size of Import Image Cache in MB"))
 	uint32 MaxImageImportCacheSizeMegaBytes = 256;
- 
+
 	UPROPERTY(config, EditAnywhere, Category = "Configuration", meta = (UIMin = "0.0", UIMax = "10.0", ClampMin = "0.0", ClampMax = "10.0", ToolTip = "Exponent for the Paint Tool Strength"))
 	float PaintStrengthGamma = 2.2f;
- 
+
 	UPROPERTY(config, EditAnywhere, Category = "Configuration", meta = (ToolTip = "Disable Painting Startup Slowdown"))
 	bool bDisablePaintingStartupSlowdown = true;
 	
 	/** Defines when the engine requires the landscape actors to be resaved */
 	UPROPERTY(Config, Category = "Configuration", EditAnywhere)
 	ELandscapeDirtyingMode LandscapeDirtyingMode = ELandscapeDirtyingMode::InLandscapeModeAndUserTriggeredChanges;
- 
+
 protected:
 	UPROPERTY(config)
 	int32 SideResolutionLimit = 2048;
- 
+
 	UPROPERTY(EditAnywhere, config, Category = "Materials", meta = (ToolTip = "Default Landscape Material will be prefilled when creating a new landscape."))
 	TSoftObjectPtr<UMaterialInterface> DefaultLandscapeMaterial;
 
 	UPROPERTY(EditAnywhere, config, Category = "Target Layers", meta = (ToolTip = "Default Layer Info Object"))
 	TSoftObjectPtr<ULandscapeLayerInfoObject> DefaultLayerInfoObject;
- 
+
 	UPROPERTY(EditAnywhere, config, Category = "Configuration", meta = (ToolTip = "Maximum size that can be set via the slider for the landscape sculpt/paint brushes"))
 	float BrushSizeUIMax = 8192;
- 
+
 	UPROPERTY(EditAnywhere, config, Category = "Configuration", meta = (ToolTip = "Maximum size that can be set manually for the landscape sculpt/paint brushes"))
 	float BrushSizeClampMax = 65536;
- 
+
 	UPROPERTY(EditAnywhere, config, Category = "HLOD", meta = (DisplayName = "HLOD Max Texture Size", ClampMin = "64", ClampMax = "8192", ToolTip = "Maximum size of the textures generated for landscape HLODs"))
 	int32 HLODMaxTextureSize = 1024;
- 
+
 	/** For landscape layers-affecting changes, allows the landscape to be updated when performing an interactive change (e.g. when changing an edit layer's alpha). Set to false if the performance when editing gets too bad (the landscape will be properly updated when the dragging operation is done). */
 	UPROPERTY(EditAnywhere, config, Category = "Edit Layers")
 	bool bShouldUpdateEditLayersDuringInteractiveChanges = false;
  
 	UPROPERTY(transient)
 	bool bRestrictiveMode = false;
- 
+
 	/** Offset in Z for the landscape spline icon in world-space. */
 	UPROPERTY(EditAnywhere, config, Category = "Spline", meta = (ClampMin = "0", UIMax = "4096"))
 	float SplineIconWorldZOffset = 100.0f;
- 
+
 	UPROPERTY(EditAnywhere, config, Category = "Spline", meta = (ClampMin = "32", UIMax = "2048"))
 	float SplineIconScale = 125.0f;
- 
-	UPROPERTY(EditAnywhere, config, Category = "Configuration", meta = (ToolTip = 
+
+	UPROPERTY(EditAnywhere, config, Category = "Target Layers", meta = (ToolTip = 
 		"When true, each target layer will have a representative thumbnail in landscape mode. However, the thumbnail system can be slow. Plus, certain landscape materials (e.g. using RVT) are not able to properly render thumbnails. Thus, setting this to false will skip needlessly rendering landscape layer thumbnails, which can improve the editing experience"))
 	bool bDisplayTargetLayerThumbnails = true;
- 
+
 	UPROPERTY(EditAnywhere, config, Category = "Configuration", meta = (ToolTip = 
 		"When true, temporal anti-aliasing will be inactive while in landscape mode. This avoids the ghosting effect on the landscape brush but can lead to aliasing or shimmering on other parts of the image."))
 	bool bDisableTemporalAntiAliasingInLandscapeMode = true;
- 
+
+	/** Target layer blend method to use for newly created Landscape Layer Info assets. Changing this setting will not affect existing Landscape Layer Info assets. This is only used when DefaultLayerInfoObject isn't set. */
+	UPROPERTY(Config, Category = "Target Layers", EditAnywhere)
+	ELandscapeTargetLayerBlendMethod TargetLayerDefaultBlendMethod = ELandscapeTargetLayerBlendMethod::None;
+
 private:
 	TOptional<bool> DisplayTargetLayerThumbnailsBeforeUndo;
 };
diff --git a/Engine/Source/Runtime/Landscape/Public/LandscapeUtils.h b/Engine/Source/Runtime/Landscape/Public/LandscapeUtils.h
--- a/Engine/Source/Runtime/Landscape/Public/LandscapeUtils.h
+++ b/Engine/Source/Runtime/Landscape/Public/LandscapeUtils.h
@@ -27,6 +27,7 @@ class FTextureResource;
 class FMaterialUpdateContext;
 enum class ELandscapeToolTargetType : uint8;
 enum class ELandscapeToolTargetTypeFlags : uint8;
+struct FShaderCompilerEnvironment;
 
 namespace UE::Landscape
 {
@@ -38,11 +39,13 @@ enum class EBuildFlags : uint8;
 */
 LANDSCAPE_API bool DoesPlatformSupportEditLayers(EShaderPlatform InShaderPlatform);
 
+/** Provides an opportunity to modify the shader compiler environment (for landscape debugging purposes) */
+LANDSCAPE_API void ModifyShaderCompilerEnvironmentForDebug(FShaderCompilerEnvironment& OutEnvironment);
+
 LANDSCAPE_API ELandscapeToolTargetTypeFlags GetLandscapeToolTargetTypeAsFlags(ELandscapeToolTargetType InTargetType);
 LANDSCAPE_API ELandscapeToolTargetType GetLandscapeToolTargetTypeSingleFlagAsType(ELandscapeToolTargetTypeFlags InSingleFlag);
 LANDSCAPE_API FString GetLandscapeToolTargetTypeFlagsAsString(ELandscapeToolTargetTypeFlags InTargetTypeFlags);
 
-
 // ----------------------------------------------------------------------------------
 
 /** This struct is usually meant to be allocated on the game thread (where there's no FRDGBuilder, which is Render Thread-only) and allows to queue successive operations (lambdas) onto a single render command,
