Added "is on (parallel) game thread" assert to FLandscapeComponentDataInterface to avoid another crash in the future where this is used as a cache on a worker thread while landscape components might manipulate their inner structs on the game thread.
Ideally, this whole class should be replaced by one that could copy all the necessary data on the game thread and then be used as a temporary cache from which to read, but this involves a risky refactor.

diff --git a/Engine/Plugins/PCG/Source/PCG/Private/Grid/PCGLandscapeCache.cpp b/Engine/Plugins/PCG/Source/PCG/Private/Grid/PCGLandscapeCache.cpp
--- a/Engine/Plugins/PCG/Source/PCG/Private/Grid/PCGLandscapeCache.cpp
+++ b/Engine/Plugins/PCG/Source/PCG/Private/Grid/PCGLandscapeCache.cpp
@@ -108,6 +108,7 @@ FPCGLandscapeCacheEntry* FPCGLandscapeCacheEntry::CreateCacheEntry(ULandscapeInf
 
 		Result->PositionsAndNormals.Reserve(2 * NumVertices);
 
+		TRACE_CPUPROFILER_EVENT_SCOPE(GetWorldPositionTangents);
 		for (int32 Index = 0; Index < NumVertices; ++Index)
 		{
 			CDI.GetWorldPositionTangents(Index, WorldPos, WorldTangentX, WorldTangentY, WorldTangentZ);
@@ -127,6 +128,7 @@ FPCGLandscapeCacheEntry* FPCGLandscapeCacheEntry::CreateCacheEntry(ULandscapeInf
 				continue;
 			}
 
+			TRACE_CPUPROFILER_EVENT_SCOPE(GetWeightmapTextureData);
 			if (CDI.GetWeightmapTextureData(LayerInfo, LayerCache, /*bUseEditingLayer=*/false, /*bRemoveSubsectionDuplicates=*/true))
 			{
 				Result->LayerDataNames.Add(Layer.LayerName);
diff --git a/Engine/Source/Runtime/Landscape/Private/LandscapeDataAccess.cpp b/Engine/Source/Runtime/Landscape/Private/LandscapeDataAccess.cpp
--- a/Engine/Source/Runtime/Landscape/Private/LandscapeDataAccess.cpp
+++ b/Engine/Source/Runtime/Landscape/Private/LandscapeDataAccess.cpp
@@ -6,10 +6,12 @@
 #if WITH_EDITOR
 
 
-FLandscapeComponentDataInterfaceBase::FLandscapeComponentDataInterfaceBase(ULandscapeComponent* InComponent, int32 InMipLevel, bool InWorkOnEditingLayer): 
-	MipLevel(InMipLevel)
+FLandscapeComponentDataInterfaceBase::FLandscapeComponentDataInterfaceBase(ULandscapeComponent* InComponent, int32 InMipLevel, bool bInWorkOnEditingLayer)
+	: MipLevel(InMipLevel)
 {
-	UTexture2D* HeightMapTexture = InComponent->GetHeightmap(InWorkOnEditingLayer);
+	check(IsInParallelGameThread() || IsInGameThread());
+
+	UTexture2D* HeightMapTexture = InComponent->GetHeightmap(bInWorkOnEditingLayer);
 	HeightmapStride = HeightMapTexture->Source.GetSizeX() >> InMipLevel;
 	HeightmapComponentOffsetX = FMath::RoundToInt((HeightMapTexture->Source.GetSizeX() >> InMipLevel) * InComponent->HeightmapScaleBias.Z);
 	HeightmapComponentOffsetY = FMath::RoundToInt((HeightMapTexture->Source.GetSizeY() >> InMipLevel) * InComponent->HeightmapScaleBias.W);
@@ -21,14 +23,16 @@ FLandscapeComponentDataInterfaceBase::FLandscapeComponentDataInterfaceBase(ULand
 	ComponentNumSubsections = InComponent->NumSubsections;
 }
 
-LANDSCAPE_API FLandscapeComponentDataInterface::FLandscapeComponentDataInterface(ULandscapeComponent* InComponent, int32 InMipLevel, bool InWorkOnEditingLayer) :
-	FLandscapeComponentDataInterfaceBase(InComponent, InMipLevel, InWorkOnEditingLayer),
-	Component(InComponent),
-	bWorkOnEditingLayer(InWorkOnEditingLayer),
-	HeightMipData(NULL),
-	XYOffsetMipData(NULL)
+FLandscapeComponentDataInterface::FLandscapeComponentDataInterface(ULandscapeComponent* InComponent, int32 InMipLevel, bool bInWorkOnEditingLayer)
+	: FLandscapeComponentDataInterfaceBase(InComponent, InMipLevel, bInWorkOnEditingLayer)
+	, Component(InComponent)
+	, bWorkOnEditingLayer(bInWorkOnEditingLayer)
+	, HeightMipData(nullptr)
+	, XYOffsetMipData(nullptr)
 {
-	UTexture2D* HeightMapTexture = Component->GetHeightmap(bWorkOnEditingLayer);
+	check(IsInParallelGameThread() || IsInGameThread());
+
+	UTexture2D* HeightMapTexture = InComponent->GetHeightmap(bWorkOnEditingLayer);
 	if (MipLevel < HeightMapTexture->Source.GetNumMips())
 	{
 		HeightMipData = (FColor*)DataInterface.LockMip(HeightMapTexture, MipLevel);
@@ -39,8 +43,10 @@ LANDSCAPE_API FLandscapeComponentDataInterface::FLandscapeComponentDataInterface
 	}
 }
 
-LANDSCAPE_API FLandscapeComponentDataInterface::~FLandscapeComponentDataInterface()
+FLandscapeComponentDataInterface::~FLandscapeComponentDataInterface()
 {
+	check(IsInParallelGameThread() || IsInGameThread());
+
 	if (HeightMipData)
 	{
 		UTexture2D* HeightMapTexture = Component->GetHeightmap(bWorkOnEditingLayer);
@@ -52,8 +58,9 @@ LANDSCAPE_API FLandscapeComponentDataInterface::~FLandscapeComponentDataInterfac
 	}
 }
 
-LANDSCAPE_API void FLandscapeComponentDataInterface::GetHeightmapTextureData(TArray<FColor>& OutData, bool bOkToFail)
+void FLandscapeComponentDataInterface::GetHeightmapTextureData(TArray<FColor>& OutData, bool bOkToFail)
 {
+	check(IsInParallelGameThread() || IsInGameThread());
 	if (bOkToFail && !HeightMipData)
 	{
 		OutData.Empty();
@@ -79,8 +86,9 @@ LANDSCAPE_API void FLandscapeComponentDataInterface::GetHeightmapTextureData(TAr
 	}
 }
 
-LANDSCAPE_API bool FLandscapeComponentDataInterface::GetWeightmapTextureData(ULandscapeLayerInfoObject* InLayerInfo, TArray<uint8>& OutData, bool bInUseEditingWeightmap, bool bInRemoveSubsectionDuplicates)
+bool FLandscapeComponentDataInterface::GetWeightmapTextureData(ULandscapeLayerInfoObject* InLayerInfo, TArray<uint8>& OutData, bool bInUseEditingWeightmap, bool bInRemoveSubsectionDuplicates)
 {
+	check(IsInParallelGameThread() || IsInGameThread());
 	int32 LayerIdx = INDEX_NONE;
 	const TArray<FWeightmapLayerAllocationInfo>& ComponentWeightmapLayerAllocations = Component->GetWeightmapLayerAllocations(bInUseEditingWeightmap);
 	const TArray<UTexture2D*>& ComponentWeightmapTextures = Component->GetWeightmapTextures(bInUseEditingWeightmap);
@@ -132,9 +140,10 @@ LANDSCAPE_API bool FLandscapeComponentDataInterface::GetWeightmapTextureData(ULa
 	return true;
 }
 
-LANDSCAPE_API FColor* FLandscapeComponentDataInterface::GetXYOffsetData(int32 LocalX, int32 LocalY) const
+FColor* FLandscapeComponentDataInterface::GetXYOffsetData(int32 LocalX, int32 LocalY) const
 {
 #if LANDSCAPE_VALIDATE_DATA_ACCESS
+	checkSlow(IsInParallelGameThread() || IsInGameThread());
 	check(Component);
 	check(LocalX >= 0 && LocalY >= 0 && LocalX < ComponentSizeQuads + 1 && LocalY < ComponentSizeQuads + 1);
 #endif
@@ -149,7 +158,7 @@ LANDSCAPE_API FColor* FLandscapeComponentDataInterface::GetXYOffsetData(int32 Lo
 	return &XYOffsetMipData[SubX + SubNumX*SubsectionSizeVerts + (SubY + SubNumY*SubsectionSizeVerts)*WeightmapSize];
 }
 
-LANDSCAPE_API FVector FLandscapeComponentDataInterface::GetLocalVertex(int32 LocalX, int32 LocalY) const
+FVector FLandscapeComponentDataInterface::GetLocalVertex(int32 LocalX, int32 LocalY) const
 {
 	const float ScaleFactor = (float)ComponentSizeQuads / (float)(ComponentSizeVerts - 1);
 	float XOffset, YOffset;
@@ -157,12 +166,12 @@ LANDSCAPE_API FVector FLandscapeComponentDataInterface::GetLocalVertex(int32 Loc
 	return FVector(LocalX * ScaleFactor + XOffset, LocalY * ScaleFactor + YOffset, LandscapeDataAccess::GetLocalHeight(GetHeight(LocalX, LocalY)));
 }
 
-LANDSCAPE_API FVector FLandscapeComponentDataInterface::GetWorldVertex(int32 LocalX, int32 LocalY) const
+FVector FLandscapeComponentDataInterface::GetWorldVertex(int32 LocalX, int32 LocalY) const
 {
 	return Component->GetComponentTransform().TransformPosition(GetLocalVertex(LocalX, LocalY));
 }
 
-LANDSCAPE_API void FLandscapeComponentDataInterface::GetWorldTangentVectors(int32 LocalX, int32 LocalY, FVector& WorldTangentX, FVector& WorldTangentY, FVector& WorldTangentZ) const
+void FLandscapeComponentDataInterface::GetWorldTangentVectors(int32 LocalX, int32 LocalY, FVector& WorldTangentX, FVector& WorldTangentY, FVector& WorldTangentZ) const
 {
 	FColor* Data = GetHeightData(LocalX, LocalY);
 	WorldTangentZ = LandscapeDataAccess::UnpackNormal(*Data);
@@ -174,7 +183,7 @@ LANDSCAPE_API void FLandscapeComponentDataInterface::GetWorldTangentVectors(int3
 	WorldTangentZ = Component->GetComponentTransform().TransformVectorNoScale(WorldTangentZ);
 }
 
-LANDSCAPE_API void FLandscapeComponentDataInterface::GetWorldPositionTangents(int32 LocalX, int32 LocalY, FVector& WorldPos, FVector& WorldTangentX, FVector& WorldTangentY, FVector& WorldTangentZ) const
+void FLandscapeComponentDataInterface::GetWorldPositionTangents(int32 LocalX, int32 LocalY, FVector& WorldPos, FVector& WorldTangentX, FVector& WorldTangentY, FVector& WorldTangentZ) const
 {
 	FColor* Data = GetHeightData(LocalX, LocalY);
 	WorldTangentZ = LandscapeDataAccess::UnpackNormal(*Data);
@@ -194,11 +203,13 @@ LANDSCAPE_API void FLandscapeComponentDataInterface::GetWorldPositionTangents(in
 
 int32 FLandscapeComponentDataInterface::GetHeightmapSizeX(int32 MipIndex) const
 {
+	check(IsInParallelGameThread() || IsInGameThread());
 	return Component->GetHeightmap(bWorkOnEditingLayer)->Source.GetSizeX() >> MipIndex;
 }
 
 int32 FLandscapeComponentDataInterface::GetHeightmapSizeY(int32 MipIndex) const
 {
+	check(IsInParallelGameThread() || IsInGameThread());
 	return Component->GetHeightmap(bWorkOnEditingLayer)->Source.GetSizeY() >> MipIndex;
 }
 
diff --git a/Engine/Source/Runtime/Landscape/Public/LandscapeDataAccess.h b/Engine/Source/Runtime/Landscape/Public/LandscapeDataAccess.h
--- a/Engine/Source/Runtime/Landscape/Public/LandscapeDataAccess.h
+++ b/Engine/Source/Runtime/Landscape/Public/LandscapeDataAccess.h
@@ -93,7 +93,7 @@ public:
 		check(MipLevel < Texture->Source.GetNumMips());
 
 		TArray<FLockedMipDataInfo>* MipInfo = LockedMipInfoMap.Find(Texture);
-		if( MipInfo == NULL )
+		if( MipInfo == nullptr )
 		{
 			MipInfo = &LockedMipInfoMap.Add(Texture, TArray<FLockedMipDataInfo>() );
 			for (int32 i = 0; i < Texture->Source.GetNumMips(); ++i)
@@ -139,9 +139,6 @@ struct FLandscapeComponentDataInterfaceBase
 	// Accessors
 	void VertexIndexToXY(int32 VertexIndex, int32& OutX, int32& OutY) const
 	{
-//#if LANDSCAPE_VALIDATE_DATA_ACCESS
-//		check(MipLevel == 0);
-//#endif
 		OutX = VertexIndex % ComponentSizeVerts;
 		OutY = VertexIndex / ComponentSizeVerts;
 	}
@@ -149,9 +146,6 @@ struct FLandscapeComponentDataInterfaceBase
 	// Accessors
 	void QuadIndexToXY(int32 QuadIndex, int32& OutX, int32& OutY) const
 	{
-//#if LANDSCAPE_VALIDATE_DATA_ACCESS
-//		check(MipLevel == 0);
-//#endif
 		OutX = QuadIndex % (ComponentSizeVerts-1);
 		OutY = QuadIndex / (ComponentSizeVerts-1);
 	}
@@ -311,6 +305,7 @@ struct FLandscapeComponentDataInterface : public FLandscapeComponentDataInterfac
 	FColor* GetHeightData(int32 LocalX, int32 LocalY) const
 	{
 #if LANDSCAPE_VALIDATE_DATA_ACCESS
+		checkSlow(IsInParallelGameThread() || IsInGameThread());
 		check(Component);
 		check(HeightMipData);
 		check(LocalX >=0 && LocalY >=0 && LocalX < ComponentSizeVerts && LocalY < ComponentSizeVerts);
@@ -423,12 +418,10 @@ struct FLandscapeComponentDataInterface : public FLandscapeComponentDataInterfac
 
 private:
 	FLandscapeDataInterface DataInterface;
-	ULandscapeComponent* Component;
-	bool bWorkOnEditingLayer;
-
-private:
-	FColor* HeightMipData;
-	FColor* XYOffsetMipData;
+	ULandscapeComponent* Component = nullptr;
+	bool bWorkOnEditingLayer = false;
+	FColor* HeightMipData = nullptr;
+	FColor* XYOffsetMipData = nullptr;
 	
 };
 
