Landscape Editor - Adds nullptr warning for nanite PostLoad SetSourceLandscapeComponents
* Fixes crash when running in Single Exec Workflow mode

diff --git a/Engine/Source/Runtime/Landscape/Private/Landscape.cpp b/Engine/Source/Runtime/Landscape/Private/Landscape.cpp
--- a/Engine/Source/Runtime/Landscape/Private/Landscape.cpp
+++ b/Engine/Source/Runtime/Landscape/Private/Landscape.cpp
@@ -4664,9 +4664,12 @@ PRAGMA_DISABLE_DEPRECATION_WARNINGS
 	}
 PRAGMA_ENABLE_DEPRECATION_WARNINGS
 
-	// Fix for empty nanite components that have lost their link to their source landscape components
-	SetSourceComponentsForNaniteComponents(/*bSetEmptyComponentsOnly=*/ true);
-
+	if (GIsEditor)
+	{
+		// Fix for empty nanite components that have lost their link to their source landscape components
+		SetSourceComponentsForNaniteComponents(/*bSetEmptyComponentsOnly=*/ true);
+	}
+	
 	// Handle Nanite representation invalidation on load: 
 	if (!HasAnyFlags(RF_ClassDefaultObject) && !FPlatformProperties::RequiresCookedData())
 	{
@@ -6967,6 +6970,11 @@ void ALandscapeProxy::SetSourceComponentsForNaniteComponents(bool bSetEmptyCompo
 		const int32 EndComponentIndex = FMath::Min(LandscapeComponents.Num(), (NaniteComponentIndex + 1) * NaniteMaxComponents);
 		const int32 NumComponents = EndComponentIndex - StartComponentIndex;
 		ULandscapeNaniteComponent* NaniteComponent = NaniteComponents[NaniteComponentIndex];
+		if (NaniteComponent == nullptr)
+		{
+			UE_LOG(LogLandscape, Warning, TEXT("Landscape Nanite failed to set source components for NaniteComponents[%d]. Landscape Nanite requires rebuild. Actor: '%s' Package: '%s'"), NaniteComponentIndex, *GetActorNameOrLabel(), *GetPackage()->GetName());
+			continue;
+		}
 		// Optionally skip nanite components that already have valid source components set
 		if (NumComponents <= 0 || !LandscapeComponents.IsValidIndex(StartComponentIndex) || (!NaniteComponent->GetSourceLandscapeComponents().IsEmpty() && bSetEmptyComponentsOnly))
 		{
