[Landscape] Landscape componentt scene proxy delegate registration sanity check

#rnx

diff --git a/Engine/Source/Runtime/Landscape/Private/LandscapeModule.cpp b/Engine/Source/Runtime/Landscape/Private/LandscapeModule.cpp
--- a/Engine/Source/Runtime/Landscape/Private/LandscapeModule.cpp
+++ b/Engine/Source/Runtime/Landscape/Private/LandscapeModule.cpp
@@ -44,7 +44,12 @@ public:
 
 	virtual void SetLandscapeEditorServices(ILandscapeEditorServices* InLandscapeEditorServices) override { LandscapeEditorServices = InLandscapeEditorServices; }
 	virtual ILandscapeEditorServices* GetLandscapeEditorServices() const override { return LandscapeEditorServices; }
-	virtual void SetCreateLandscapeComponentSceneProxyDelegate(const UE::Landscape::FCreateLandscapeComponentSceneProxyDelegate& InDelegate) override { CreateLandscapeComponentSceneProxyDelegate = InDelegate; }
+	virtual void SetCreateLandscapeComponentSceneProxyDelegate(const UE::Landscape::FCreateLandscapeComponentSceneProxyDelegate& InDelegate) override
+	{
+		// Setting new delegate without reseting previous one is not allowed.
+		check(CreateLandscapeComponentSceneProxyDelegate.IsBound() != InDelegate.IsBound());
+		CreateLandscapeComponentSceneProxyDelegate = InDelegate;
+	}
 	virtual const UE::Landscape::FCreateLandscapeComponentSceneProxyDelegate& GetCreateLandscapeComponentSceneProxyDelegate() const override { return CreateLandscapeComponentSceneProxyDelegate; }
 
 private:
