Fixed bad landscape component scale / SectionBase after  creating a new WP landscape. This fixes a regression introduced by 45383226.
Before the change, we were simply lucky that PostRegisterAllComponents was called by the new actor label being set when the WP landscape creation process was spawning a new landscape streaming proxy. This would lead FixupProxiesTransform to be called, which would assign the proper transform to all the (new) proxies from the landscape. Without this, the components, when attached to the newly created proxy, would get their SectionBase refreshed with improper values (i.e. not aligned with the number of quads on the component). So we now forcefully fixup the transforms right after creation in order not to rely on PostRegisterAllComponents happening.
Also, use a local copy of the landscape proxies when propagating property changes to them, to avoid the list being reshuffled while iterating (because of de-registration/re-registration)

Both of these will be followed up by a better fix but this prevents introducing bad data in the meantime.

#rnx

diff --git a/Engine/Source/Runtime/Landscape/Private/LandscapeConfigHelper.cpp b/Engine/Source/Runtime/Landscape/Private/LandscapeConfigHelper.cpp
--- a/Engine/Source/Runtime/Landscape/Private/LandscapeConfigHelper.cpp
+++ b/Engine/Source/Runtime/Landscape/Private/LandscapeConfigHelper.cpp
@@ -67,7 +67,7 @@ ALandscapeProxy* FLandscapeConfigHelper::FindOrAddLandscapeStreamingProxy(UActor
 	ALandscape* Landscape = InLandscapeInfo->LandscapeActor.Get();
 	check(Landscape);
 
-	auto LandscapeProxyCreated = [InCellCoord, Landscape](APartitionActor* PartitionActor)
+	auto LandscapeProxyCreated = [InCellCoord, Landscape, InLandscapeInfo](APartitionActor* PartitionActor)
 	{
 		const FIntPoint CellLocation(static_cast<int32>(InCellCoord.X) * Landscape->GetGridSize(), static_cast<int32>(InCellCoord.Y) * Landscape->GetGridSize());
 
@@ -80,6 +80,9 @@ ALandscapeProxy* FLandscapeConfigHelper::FindOrAddLandscapeStreamingProxy(UActor
 		LandscapeProxy->SetActorLocationAndRotation(ProxyLocation, Landscape->GetActorRotation());
 		LandscapeProxy->LandscapeSectionOffset = FIntPoint(CellLocation.X, CellLocation.Y);
 		LandscapeProxy->SetIsSpatiallyLoaded(LandscapeProxy->GetLandscapeInfo()->AreNewLandscapeActorsSpatiallyLoaded());
+		// TODO [jonathan.bard] temp fix for the fact that PostRegisterAllComponents doesn't get called anymore, to prevent bad data from creeping in. 
+		//  A proper fix is coming soon
+		InLandscapeInfo->FixupProxiesTransform();
 	};
 
 	const bool bCreate = true;
diff --git a/Engine/Source/Runtime/Landscape/Private/LandscapeEdit.cpp b/Engine/Source/Runtime/Landscape/Private/LandscapeEdit.cpp
--- a/Engine/Source/Runtime/Landscape/Private/LandscapeEdit.cpp
+++ b/Engine/Source/Runtime/Landscape/Private/LandscapeEdit.cpp
@@ -6837,7 +6837,10 @@ void ALandscape::PostEditChangeProperty(FPropertyChangedEvent& PropertyChangedEv
 		if (bPropagateToProxies)
 		{
 			// Propagate Event to Proxies...
-			for (TWeakObjectPtr<ALandscapeStreamingProxy> ProxyPtr : Info->StreamingProxies)
+			// First, gather the proxies in a local list so that we don't get screwed by PostEditChangeProperty re-registering streaming proxies while iterating :
+			// TODO [jonathan.bard] : replace by ForEachLandscapeProxy (and make ForEachLandscapeProxy also iterate on the list's local copy)
+			TArray<TWeakObjectPtr<ALandscapeStreamingProxy>> StreamingProxies(Info->StreamingProxies);
+			for (TWeakObjectPtr<ALandscapeStreamingProxy> ProxyPtr : StreamingProxies)
 			{
 				if (ALandscapeProxy* Proxy = ProxyPtr.Get())
 				{
