Landscape Editor - Retopologize / XY offset removal

#was-backout 45624427

#DisabledValidations Preflight

diff --git a/Engine/Plugins/Enterprise/GLTFExporter/Source/GLTFExporter/Private/Utilities/GLTFLandscapeComponentDataInterface.cpp b/Engine/Plugins/Enterprise/GLTFExporter/Source/GLTFExporter/Private/Utilities/GLTFLandscapeComponentDataInterface.cpp
--- a/Engine/Plugins/Enterprise/GLTFExporter/Source/GLTFExporter/Private/Utilities/GLTFLandscapeComponentDataInterface.cpp
+++ b/Engine/Plugins/Enterprise/GLTFExporter/Source/GLTFExporter/Private/Utilities/GLTFLandscapeComponentDataInterface.cpp
@@ -22,15 +22,6 @@ FGLTFLandscapeComponentDataInterface::FGLTFLandscapeComponentDataInterface(const
 			HeightMipData = (FColor*)HeightMipDataSource.GetData();
 		}
 	}
-
-	if (Component.XYOffsetmapTexture)
-	{
-		if (MipLevel < Component.XYOffsetmapTexture->Source.GetNumMips())
-		{
-			Component.XYOffsetmapTexture->Source.GetMipData(XYOffsetMipDataSource, MipLevel);
-			XYOffsetMipData = (FColor*)XYOffsetMipDataSource.GetData();
-		}
-	}
 #else
 	if (HeightMapTexture)
 	{
@@ -42,18 +33,6 @@ FGLTFLandscapeComponentDataInterface::FGLTFLandscapeComponentDataInterface(const
 			HeightMipData = (FColor*)MipData;
 		}
 	}
-
-	if (Component.XYOffsetmapTexture)
-	{
-		FTexturePlatformData* PlatformData = Component.XYOffsetmapTexture->GetPlatformData();
-
-		if (MipLevel < PlatformData->Mips.Num())
-		{
-			void* MipData = nullptr;
-			PlatformData->Mips[MipLevel].BulkData.GetCopy(&MipData, false);
-			XYOffsetMipData = (FColor*)MipData;
-		}
-	}
 #endif
 	ComponentSubsectionSizeQuads = Component.SubsectionSizeQuads;
 
@@ -220,35 +199,15 @@ FColor FGLTFLandscapeComponentDataInterface::GetHeightData(int32 LocalX, int32 L
 	return HeightMipData[TexelX + HeightmapComponentOffsetX + (TexelY + HeightmapComponentOffsetY) * HeightmapStride];
 }
 
+// Deprecated
 FColor FGLTFLandscapeComponentDataInterface::GetXYOffsetData(int32 LocalX, int32 LocalY) const
 {
-	if (XYOffsetMipData == nullptr)
-	{
-		return FColor(0,0,0,0);
-	}
-
-	const int32 WeightmapSize = ((ComponentSubsectionSizeQuads + 1) * ComponentNumSubsections) >> MipLevel;
-	int32 SubNumX;
-	int32 SubNumY;
-	int32 SubX;
-	int32 SubY;
-	ComponentXYToSubsectionXY(LocalX, LocalY, SubNumX, SubNumY, SubX, SubY);
-
-	return XYOffsetMipData[SubX + SubNumX * ComponentSubsectionSizeVerts + (SubY + SubNumY * ComponentSubsectionSizeVerts) * WeightmapSize];
+	return FColor();
 }
 
+// Deprecated
 void FGLTFLandscapeComponentDataInterface::GetXYOffset(int32 X, int32 Y, float& XOffset, float& YOffset) const
 {
-	if (XYOffsetMipData)
-	{
-		FColor Texel = GetXYOffsetData(X, Y);
-		XOffset = ((float)((Texel.R << 8) + Texel.G) - 32768.f) * LANDSCAPE_XYOFFSET_SCALE;
-		YOffset = ((float)((Texel.B << 8) + Texel.A) - 32768.f) * LANDSCAPE_XYOFFSET_SCALE;
-	}
-	else
-	{
-		XOffset = YOffset = 0.f;
-	}
 }
 
 void FGLTFLandscapeComponentDataInterface::GetPositionNormalUV(int32 LocalX, int32 LocalY,
@@ -265,9 +224,7 @@ void FGLTFLandscapeComponentDataInterface::GetPositionNormalUV(int32 LocalX, int
 		float Height = LandscapeDataAccess::UnpackHeight(Data);
 
 		const float ScaleFactor = (float)ComponentSizeQuads / (float)(ComponentSizeVerts - 1);
-		float XOffset, YOffset;
-		GetXYOffset(LocalX, LocalY, XOffset, YOffset);
-		Position = FVector3f(LocalX * ScaleFactor + XOffset, LocalY * ScaleFactor + YOffset, Height);
+		Position = FVector3f(LocalX * ScaleFactor, LocalY * ScaleFactor, Height);
 
 		Normal = FVector3f(TangentZ.X, TangentZ.Y, TangentZ.Z);
 		UV = FVector2f(Component.GetSectionBase().X + LocalX * ScaleFactor, Component.GetSectionBase().Y + LocalY * ScaleFactor);
diff --git a/Engine/Plugins/Enterprise/GLTFExporter/Source/GLTFExporter/Public/Utilities/GLTFLandscapeComponentDataInterface.h b/Engine/Plugins/Enterprise/GLTFExporter/Source/GLTFExporter/Public/Utilities/GLTFLandscapeComponentDataInterface.h
--- a/Engine/Plugins/Enterprise/GLTFExporter/Source/GLTFExporter/Public/Utilities/GLTFLandscapeComponentDataInterface.h
+++ b/Engine/Plugins/Enterprise/GLTFExporter/Source/GLTFExporter/Public/Utilities/GLTFLandscapeComponentDataInterface.h
@@ -48,8 +48,10 @@ public:
 
 	FColor GetHeightData(int32 LocalX, int32 LocalY) const;
 
+	UE_DEPRECATED(5.7, "RetopologizeTool/XYOffset deprecated with the removal of non-edit layer landscapes")
 	FColor GetXYOffsetData(int32 LocalX, int32 LocalY) const;
 
+	UE_DEPRECATED(5.7, "RetopologizeTool/XYOffset deprecated with the removal of non-edit layer landscapes")
 	void GetXYOffset(int32 X, int32 Y, float& XOffset, float& YOffset) const;
 
 	void GetPositionNormalUV(int32 LocalX, int32 LocalY,
diff --git a/Engine/Plugins/Experimental/GPULightmass/Source/GPULightmass/Private/Scene/Scene.cpp b/Engine/Plugins/Experimental/GPULightmass/Source/GPULightmass/Private/Scene/Scene.cpp
--- a/Engine/Plugins/Experimental/GPULightmass/Source/GPULightmass/Private/Scene/Scene.cpp
+++ b/Engine/Plugins/Experimental/GPULightmass/Source/GPULightmass/Private/Scene/Scene.cpp
@@ -1215,7 +1215,7 @@ void FScene::AddGeometryInstanceFromComponent(ULandscapeComponent* InComponent)
 	InstanceRenderState.bCastShadow = InComponent->CastShadow && InComponent->bCastStaticShadow;
 
 	const int8 SubsectionSizeLog2 = FMath::CeilLogTwo(InComponent->SubsectionSizeQuads + 1);
-	InstanceRenderState.SharedBuffersKey = (SubsectionSizeLog2 & 0xf) | ((InComponent->NumSubsections & 0xf) << 4) | (InComponent->XYOffsetmapTexture == nullptr ? 0 : 1 << 31);
+	InstanceRenderState.SharedBuffersKey = (SubsectionSizeLog2 & 0xf) | ((InComponent->NumSubsections & 0xf) << 4);
 	InstanceRenderState.SharedBuffersKey |= 1 << 29; // Use this bit to indicate it is GPULightmass specific buffer (which only has FixedGridVertexFactory created)
 
 	TArray<UMaterialInterface*> AvailableMaterials;
@@ -1363,10 +1363,6 @@ void FScene::AddGeometryInstanceFromComponent(ULandscapeComponent* InComponent)
 
 			LandscapeParams.NormalmapTexture = Initializer.HeightmapTexture->TextureReference.TextureReferenceRHI;
 			LandscapeParams.NormalmapTextureSampler = TStaticSamplerState<SF_Point>::GetRHI();
-			
-			// No support for XYOffset
-			LandscapeParams.XYOffsetmapTexture = GBlackTexture->TextureRHI;
-			LandscapeParams.XYOffsetmapTextureSampler = GBlackTexture->SamplerStateRHI;
 
 			InstanceRenderStateRef->LandscapeUniformShaderParameters = MakeUnique<TUniformBuffer<FLandscapeUniformShaderParameters>>();
 			InstanceRenderStateRef->LandscapeUniformShaderParameters->SetContents(RHICmdList, LandscapeParams);
diff --git a/Engine/Shaders/Private/Definitions.usf b/Engine/Shaders/Private/Definitions.usf
--- a/Engine/Shaders/Private/Definitions.usf
+++ b/Engine/Shaders/Private/Definitions.usf
@@ -181,10 +181,6 @@
 #define NEEDS_LIGHTMAP_COORDINATE						0
 #endif
 
-#ifndef LANDSCAPE_XYOFFSET
-#define LANDSCAPE_XYOFFSET								0
-#endif
-
 #ifndef MATERIAL_SKY_ATMOSPHERE
 #define	MATERIAL_SKY_ATMOSPHERE							0
 #endif
diff --git a/Engine/Shaders/Private/LandscapeVertexFactory.ush b/Engine/Shaders/Private/LandscapeVertexFactory.ush
--- a/Engine/Shaders/Private/LandscapeVertexFactory.ush
+++ b/Engine/Shaders/Private/LandscapeVertexFactory.ush
@@ -16,10 +16,6 @@
 #include "/Engine/Generated/UniformBuffers/PrecomputedLightingBuffer.ush"
 #include "LightmapData.ush"
 
-#if LANDSCAPE_XYOFFSET
-	#define XYOFFSET_SCALE (1.0f/256.f)
-#endif
-
 /* Return index into buffers for this section instance. */
 int GetComponentLinearIndex()
 {
@@ -679,18 +675,6 @@ FVertexFactoryIntermediates GetVertexFactoryIntermediates(FVertexFactoryInput In
 	float4 SampleValueNextLOD = Texture2DSampleLevel(LandscapeParameters.HeightmapTexture, LandscapeParameters.HeightmapTextureSampler, SampleCoordsNextLOD, LodValue+1-Intermediates.LodBias.x);
 	float HeightNextLOD = DecodePackedHeight(SampleValueNextLOD.xy);
 
-#if LANDSCAPE_XYOFFSET // FEATURE_LEVEL >= FEATURE_LEVEL_SM4 only
-	float2 SampleCoords2 = float2(InputPositionLODAdjusted * LandscapeParameters.WeightmapUVScaleBias.xy + LandscapeParameters.WeightmapUVScaleBias.zw + Intermediates.InputPosition.zw * LandscapeParameters.SubsectionOffsetParams.zz);
-	float4 OffsetValue = Texture2DSampleLevel( LandscapeParameters.XYOffsetmapTexture, LandscapeParameters.XYOffsetmapTextureSampler, SampleCoords2, LodValue- Intermediates.LodBias.y );
-	float2 SampleCoordsNextLOD2 = float2(InputPositionNextLOD * LandscapeParameters.WeightmapUVScaleBias.xy + LandscapeParameters.WeightmapUVScaleBias.zw + Intermediates.InputPosition.zw * LandscapeParameters.SubsectionOffsetParams.zz);
-	float4 OffsetValueNextLOD = Texture2DSampleLevel( LandscapeParameters.XYOffsetmapTexture, LandscapeParameters.XYOffsetmapTextureSampler, SampleCoordsNextLOD2, LodValue+1-Intermediates.LodBias.y );
-	float2 XYOffset = float2(((OffsetValue.r * 255.0 * 256.0 + OffsetValue.g * 255.0) - 32768.0) * XYOFFSET_SCALE, ((OffsetValue.b * 255.0 * 256.0 + OffsetValue.a * 255.0) - 32768.0) * XYOFFSET_SCALE );
-	float2 XYOffsetNextLOD = float2(((OffsetValueNextLOD.r * 255.0 * 256.0 + OffsetValueNextLOD.g * 255.0) - 32768.0) * XYOFFSET_SCALE, ((OffsetValueNextLOD.b * 255.0 * 256.0 + OffsetValueNextLOD.a * 255.0) - 32768.0) * XYOFFSET_SCALE );
-
-	InputPositionLODAdjusted = InputPositionLODAdjusted + XYOffset;
-	InputPositionNextLOD = InputPositionNextLOD + XYOffsetNextLOD;
-#endif
-
 	Intermediates.LocalPosition = lerp( float3(InputPositionLODAdjusted, Height), float3(InputPositionNextLOD, HeightNextLOD), MorphAlpha );
 
 	float2 Normal = float2(SampleValue.b, SampleValue.a);
diff --git a/Engine/Source/Editor/EditorStyle/Private/StarshipStyle.cpp b/Engine/Source/Editor/EditorStyle/Private/StarshipStyle.cpp
--- a/Engine/Source/Editor/EditorStyle/Private/StarshipStyle.cpp
+++ b/Engine/Source/Editor/EditorStyle/Private/StarshipStyle.cpp
@@ -7228,8 +7228,6 @@ void FStarshipEditorStyle::FStyle::SetupLandscapeEditorStyle()
 			Set("LandscapeEditor.HydroErosionTool.Small", 		new IMAGE_BRUSH("Icons/Landscape/Landscape_HydroErosion_x40",      Icon20x20));
 			Set("LandscapeEditor.NoiseTool",        			new IMAGE_BRUSH("Icons/Landscape/Landscape_Noise_x40",             Icon20x20));
 			Set("LandscapeEditor.NoiseTool.Small",        		new IMAGE_BRUSH("Icons/Landscape/Landscape_Noise_x40",             Icon20x20));
-			Set("LandscapeEditor.RetopologizeTool", 			new IMAGE_BRUSH("Icons/Landscape/Landscape_Retopologize_x40",      Icon20x20));
-			Set("LandscapeEditor.RetopologizeTool.Small", 		new IMAGE_BRUSH("Icons/Landscape/Landscape_Retopologize_x40",      Icon20x20));
 			Set("LandscapeEditor.VisibilityTool",   			new IMAGE_BRUSH_SVG("Icons/icon_Landscape_Target_Visibility_16x",  Icon20x20));
 			Set("LandscapeEditor.VisibilityTool.Small",   		new IMAGE_BRUSH_SVG("Icons/icon_Landscape_Target_Visibility_16x",  Icon20x20));
 			Set("LandscapeEditor.BlueprintBrushTool", 			new IMAGE_BRUSH("Icons/Landscape/Landscape_BlueprintTool_x40",     Icon20x20));
diff --git a/Engine/Source/Editor/LandscapeEditor/Private/LandscapeEdMode.cpp b/Engine/Source/Editor/LandscapeEditor/Private/LandscapeEdMode.cpp
--- a/Engine/Source/Editor/LandscapeEditor/Private/LandscapeEdMode.cpp
+++ b/Engine/Source/Editor/LandscapeEditor/Private/LandscapeEdMode.cpp
@@ -133,7 +133,6 @@ FEdModeLandscape::FEdModeLandscape()
 	InitializeTool_Erosion();
 	InitializeTool_HydraErosion();
 	InitializeTool_Noise();
-	InitializeTool_Retopologize();
 	InitializeTool_NewLandscape();
 	InitializeTool_ResizeLandscape();
 	InitializeTool_ImportExport();
@@ -240,7 +239,6 @@ void FEdModeLandscape::UpdateToolModes()
 	ToolMode_Sculpt->ValidTools.Add(TEXT("Noise"));
 	ToolMode_Sculpt->ValidTools.Add(TEXT("Erosion"));
 	ToolMode_Sculpt->ValidTools.Add(TEXT("HydraErosion"));
-	ToolMode_Sculpt->ValidTools.Add(TEXT("Retopologize"));
 	ToolMode_Sculpt->ValidTools.Add(TEXT("Visibility"));
 
 	ToolMode_Sculpt->ValidTools.Add(TEXT("Mask"));
@@ -3601,14 +3599,6 @@ void FEdModeLandscape::DeleteLandscapeComponents(ULandscapeInfo* LandscapeInfo,
 			}
 		}
 
-		if (Component->XYOffsetmapTexture)
-		{
-			Component->XYOffsetmapTexture->SetFlags(RF_Transactional);
-			Component->XYOffsetmapTexture->Modify();
-			Component->XYOffsetmapTexture->MarkPackageDirty();
-			Component->XYOffsetmapTexture->ClearFlags(RF_Standalone);
-		}
-
 		ULandscapeHeightfieldCollisionComponent* CollisionComp = Component->GetCollisionComponent();
 		if (CollisionComp)
 		{
@@ -4227,14 +4217,6 @@ bool FEdModeLandscape::CanEditLayer(FText* Reason /*=nullptr*/, const ULandscape
 				return false;
 			}
 		}
-		if (CurrentTool->GetToolName() == FName("Retopologize"))
-		{
-			if (Reason)
-			{
-				*Reason = FText::Format(NSLOCTEXT("UnrealEd", "LandscapeLayersNoSupportForRetopologize", "{0} Tool is not available with the Landscape Edit Layer System and will be entirely deprecated for all landscape types in UE5.7."), CurrentTool->GetDisplayName());
-			}
-			return false;
-		}
 	}
 
 	return true;
diff --git a/Engine/Source/Editor/LandscapeEditor/Private/LandscapeEdModeComponentTools.cpp b/Engine/Source/Editor/LandscapeEditor/Private/LandscapeEdModeComponentTools.cpp
--- a/Engine/Source/Editor/LandscapeEditor/Private/LandscapeEdModeComponentTools.cpp
+++ b/Engine/Source/Editor/LandscapeEditor/Private/LandscapeEdModeComponentTools.cpp
@@ -544,15 +544,12 @@ public:
 	FLandscapeToolStrokeAddComponent(FEdModeLandscape* InEdMode, FEditorViewportClient* InViewportClient, const FLandscapeToolTarget& InTarget)
 		: Super(InEdMode, InViewportClient, InTarget)
 		, HeightCache(InTarget)
-		, XYOffsetCache(InTarget)
 	{
 	}
 
 	virtual ~FLandscapeToolStrokeAddComponent()
 	{
-		// We flush here so here ~FXYOffsetmapAccessor can safely lock the heightmap data to update bounds
 		HeightCache.Flush();
-		XYOffsetCache.Flush();
 	}
 
 	virtual void Apply(FEditorViewportClient* ViewportClient, FLandscapeBrush* Brush, const ULandscapeEditorObject* UISettings, const TArray<FLandscapeToolInteractorPosition>& InteractorPositions)
@@ -587,9 +584,7 @@ public:
 			TArray<uint16> Data;
 			TArray<FVector> XYOffsetData;
 			HeightCache.CacheData(X1, Y1, X2, Y2);
-			XYOffsetCache.CacheData(X1, Y1, X2, Y2);
 			HeightCache.GetCachedData(X1, Y1, X2, Y2, Data);
-			bool bHasXYOffset = XYOffsetCache.GetCachedData(X1, Y1, X2, Y2, XYOffsetData);
 
 			UWorld* World = ViewportClient->GetScene()->GetWorld();
 			
@@ -634,18 +629,12 @@ public:
 				}
 			}
 
-			// Need to register to use general height/xyoffset data update
+			// Need to register to use general height data update
 			for (int32 Idx = 0; Idx < NewComponents.Num(); Idx++)
 			{
 				NewComponents[Idx]->RegisterComponent();
 			}
 
-			if (bHasXYOffset)
-			{
-				XYOffsetCache.SetCachedData(X1, Y1, X2, Y2, XYOffsetData);
-			}
-			XYOffsetCache.Flush();
-
 			HeightCache.SetCachedData(X1, Y1, X2, Y2, Data);
 			HeightCache.Flush();
 
@@ -670,7 +659,6 @@ public:
 					NewComponent->UpdateBounds();
 					NewComponent->MarkRenderStateDirty();
 				}
-
 				GEngine->BroadcastOnActorMoved(Landscape);
 			}
 		}
@@ -678,7 +666,6 @@ public:
 
 protected:
 	FLandscapeHeightCache HeightCache;
-	FLandscapeXYOffsetCache<true> XYOffsetCache;
 };
 
 class FLandscapeToolAddComponent : public FLandscapeToolBase<FLandscapeToolStrokeAddComponent>
diff --git a/Engine/Source/Editor/LandscapeEditor/Private/LandscapeEdModeTools.h b/Engine/Source/Editor/LandscapeEditor/Private/LandscapeEdModeTools.h
--- a/Engine/Source/Editor/LandscapeEditor/Private/LandscapeEdModeTools.h
+++ b/Engine/Source/Editor/LandscapeEditor/Private/LandscapeEdModeTools.h
@@ -785,126 +785,6 @@ struct FLandscapeHeightCache : public TLandscapeEditCache<FHeightmapAccessorTool
 	}
 };
 
-//
-// FXYOffsetmapAccessor
-//
-template<bool bInUseInterp>
-struct FXYOffsetmapAccessor
-{
-	enum { bUseInterp = bInUseInterp };
-	FXYOffsetmapAccessor(ULandscapeInfo* InLandscapeInfo)
-	{
-		LandscapeInfo = InLandscapeInfo;
-		LandscapeEdit = new FLandscapeEditDataInterface(InLandscapeInfo);
-	}
-
-	FXYOffsetmapAccessor(const FLandscapeToolTarget& InTarget)
-		: FXYOffsetmapAccessor(InTarget.LandscapeInfo.Get())
-	{
-	}
-
-	// accessors
-	void GetData(int32& X1, int32& Y1, int32& X2, int32& Y2, TMap<FIntPoint, FVector>& Data)
-	{
-		// NOTE: GetXYOffsetData may shrink X1/Y1/X2/Y2 to reflect the sub-range that contains valid data
-		LandscapeEdit->GetXYOffsetData(X1, Y1, X2, Y2, Data);
-
-		// NOTE: GetHeightData may shrink X1/Y1/X2/Y2 to reflect the sub-range that contains valid data
-		// In general, GetHeight should not further shrink the range, as anywhere the XYOffsets are valid, the Heights should also be valid.
-		TMap<FIntPoint, uint16> NewHeights;
-		LandscapeEdit->GetHeightData(X1, Y1, X2, Y2, NewHeights);
-		for (int32 Y = Y1; Y <= Y2; ++Y)
-		{
-			for (int32 X = X1; X <= X2; ++X)
-			{
-				FVector* Value = Data.Find(FIntPoint(X, Y));
-				if (Value)
-				{
-					Value->Z = LandscapeDataAccess::GetLocalHeight(static_cast<uint16>(NewHeights.FindRef(FIntPoint(X, Y))));
-				}
-			}
-		}
-	}
-
-	void GetDataFast(int32 X1, int32 Y1, int32 X2, int32 Y2, TMap<FIntPoint, FVector>& Data)
-	{
-		LandscapeEdit->GetXYOffsetDataFast(X1, Y1, X2, Y2, Data);
-
-		// NOTE: this can modify X1/Y1/X2/Y2 to reflect the region with valid height data, which may be a subset of the original range
-		// We rely on the fact that generally anywhere the XYOffsets are valid, the Heights should also be valid
-		// So in theory we should write the Z value for every XY Offset that was queried above.
-		TMap<FIntPoint, uint16> NewHeights;
-		LandscapeEdit->GetHeightData(X1, Y1, X2, Y2, NewHeights);
-		for (int32 Y = Y1; Y <= Y2; ++Y)
-		{
-			for (int32 X = X1; X <= X2; ++X)
-			{
-				FVector* Value = Data.Find(FIntPoint(X, Y));
-				if (Value)
-				{
-					Value->Z = LandscapeDataAccess::GetLocalHeight(static_cast<uint16>(NewHeights.FindRef(FIntPoint(X, Y))));
-				}
-			}
-		}
-	}
-
-	void SetData(int32 X1, int32 Y1, int32 X2, int32 Y2, const FVector* Data, ELandscapeLayerPaintingRestriction PaintingRestriction = ELandscapeLayerPaintingRestriction::None)
-	{
-		TSet<ULandscapeComponent*> Components;
-		if (LandscapeInfo && LandscapeEdit->GetComponentsInRegion(X1, Y1, X2, Y2, &Components))
-		{
-			// Update data
-			ChangedComponents.Append(Components);
-
-			// Convert Height to uint16
-			TArray<uint16> NewHeights;
-			NewHeights.AddZeroed((Y2 - Y1 + 1) * (X2 - X1 + 1));
-			for (int32 Y = Y1; Y <= Y2; ++Y)
-			{
-				for (int32 X = X1; X <= X2; ++X)
-				{
-					NewHeights[X - X1 + (Y - Y1) * (X2 - X1 + 1)] = LandscapeDataAccess::GetTexHeight(static_cast<float>(Data[(X - X1 + (Y - Y1) * (X2 - X1 + 1))].Z));
-				}
-			}
-
-			LandscapeEdit->SetXYOffsetData(X1, Y1, X2, Y2, Data, 0); // XY Offset always need to be update before the height update
-			LandscapeEdit->SetHeightData(X1, Y1, X2, Y2, NewHeights.GetData(), 0, /*InCalcNormals=*/ false);
-		}
-	}
-
-	void Flush()
-	{
-		LandscapeEdit->Flush();
-	}
-
-	virtual ~FXYOffsetmapAccessor()
-	{
-		delete LandscapeEdit;
-		LandscapeEdit = NULL;
-
-		// Update the bounds for the components we edited
-		for (TSet<ULandscapeComponent*>::TConstIterator It(ChangedComponents); It; ++It)
-		{
-			(*It)->UpdateCachedBounds();
-			(*It)->UpdateComponentToWorld();
-		}
-	}
-
-private:
-	ULandscapeInfo* LandscapeInfo;
-	FLandscapeEditDataInterface* LandscapeEdit;
-	TSet<ULandscapeComponent*> ChangedComponents;
-};
-
-template<bool bInUseInterp>
-struct FLandscapeXYOffsetCache : public TLandscapeEditCache<FXYOffsetmapAccessor<bInUseInterp>, FVector>
-{
-	FLandscapeXYOffsetCache(const FLandscapeToolTarget& InTarget)
-		: TLandscapeEditCache<FXYOffsetmapAccessor<bInUseInterp>, FVector>(InTarget)
-	{
-	}
-};
-
 template<bool bInUseInterp, bool bInUseTotalNormalize>
 struct FAlphamapAccessorTool : public FAlphamapAccessor<bInUseInterp, bInUseTotalNormalize>
 {
diff --git a/Engine/Source/Editor/LandscapeEditor/Private/LandscapeEdModeXYOffset.cpp b/Engine/Source/Editor/LandscapeEditor/Private/LandscapeEdModeXYOffset.cpp
deleted file mode 100644
diff --git a/Engine/Source/Editor/LandscapeEditor/Private/LandscapeEditorCommands.cpp b/Engine/Source/Editor/LandscapeEditor/Private/LandscapeEditorCommands.cpp
--- a/Engine/Source/Editor/LandscapeEditor/Private/LandscapeEditorCommands.cpp
+++ b/Engine/Source/Editor/LandscapeEditor/Private/LandscapeEditorCommands.cpp
@@ -73,11 +73,6 @@ void FLandscapeEditorCommands::RegisterCommands()
 	UI_COMMAND(NoiseTool, "Noise", "Adds noise to the heightmap or blend layer", EUserInterfaceActionType::RadioButton, FInputChord());
 	NameToCommandMap.Add("Tool_Noise", NoiseTool);
 
-	UI_COMMAND(RetopologizeTool, "Retop", "Automatically adjusts landscape vertices with an X/Y offset map to improve vertex density on cliffs, reducing texture stretching.\n"
-								 "Note: An X/Y offset map makes the landscape slower to render and paint on with other tools, so only use if needed.\n\n"
-							     "Retop is not available with the Landscape Edit Layer System and will be deprecated for all landscape types in UE5.6.  If not using Edit Layers, you can set the CVar landscape.EnableRetopologizeTool to True to enable it.", EUserInterfaceActionType::RadioButton, FInputChord());
-	NameToCommandMap.Add("Tool_Retopologize", RetopologizeTool);
-
 	UI_COMMAND(VisibilityTool, "Visibility", "Mask out individual quads in the landscape, leaving a hole.", EUserInterfaceActionType::RadioButton, FInputChord());
 	NameToCommandMap.Add("Tool_Visibility", VisibilityTool);
 
diff --git a/Engine/Source/Editor/LandscapeEditor/Private/LandscapeEditorCommands.h b/Engine/Source/Editor/LandscapeEditor/Private/LandscapeEditorCommands.h
--- a/Engine/Source/Editor/LandscapeEditor/Private/LandscapeEditorCommands.h
+++ b/Engine/Source/Editor/LandscapeEditor/Private/LandscapeEditorCommands.h
@@ -46,7 +46,6 @@ public:
 	TSharedPtr<FUICommandInfo> ErosionTool;
 	TSharedPtr<FUICommandInfo> HydroErosionTool;
 	TSharedPtr<FUICommandInfo> NoiseTool;
-	TSharedPtr<FUICommandInfo> RetopologizeTool;
 	TSharedPtr<FUICommandInfo> VisibilityTool;
 	TSharedPtr<FUICommandInfo> BlueprintBrushTool;
 
diff --git a/Engine/Source/Editor/LandscapeEditor/Private/LandscapeEditorDetailCustomization_NewLandscape.cpp b/Engine/Source/Editor/LandscapeEditor/Private/LandscapeEditorDetailCustomization_NewLandscape.cpp
--- a/Engine/Source/Editor/LandscapeEditor/Private/LandscapeEditorDetailCustomization_NewLandscape.cpp
+++ b/Engine/Source/Editor/LandscapeEditor/Private/LandscapeEditorDetailCustomization_NewLandscape.cpp
@@ -1086,13 +1086,12 @@ void FLandscapeEditorDetailCustomization_NewLandscape::AddComponents(ULandscapeI
 		InLandscapeInfo->XYtoAddCollisionMap.Remove(ComponentCoordinate);
 	}
 
-	// Need to register to use general height/xyoffset data update
+	// Need to register to use general height data update
 	for (int32 Idx = 0; Idx < NewComponents.Num(); Idx++)
 	{
 		NewComponents[Idx]->RegisterComponent();
 	}
 
-	const bool bHasXYOffset = false;
 	ALandscape* Landscape = InLandscapeInfo->LandscapeActor.Get();
 
 	for (ULandscapeComponent* NewComponent : NewComponents)
diff --git a/Engine/Source/Editor/LandscapeEditor/Private/LandscapeEditorDetails.cpp b/Engine/Source/Editor/LandscapeEditor/Private/LandscapeEditorDetails.cpp
--- a/Engine/Source/Editor/LandscapeEditor/Private/LandscapeEditorDetails.cpp
+++ b/Engine/Source/Editor/LandscapeEditor/Private/LandscapeEditorDetails.cpp
@@ -238,7 +238,6 @@ FText FLandscapeEditorDetails::GetLocalizedName(FString Name)
 		LANDSCAPE_TOOL_LOCTEXT("ToolSet_Erosion", "Erosion"),
 		LANDSCAPE_TOOL_LOCTEXT("ToolSet_HydraErosion", "HydroErosion"),
 		LANDSCAPE_TOOL_LOCTEXT("ToolSet_Noise", "Noise"),
-		LANDSCAPE_TOOL_LOCTEXT("ToolSet_Retopologize", "Retopologize"),
 		LANDSCAPE_TOOL_LOCTEXT("ToolSet_Visibility", "Visibility"),
 		LANDSCAPE_TOOL_LOCTEXT("ToolSet_BlueprintBrush", "Blueprint Brushes"),
 
diff --git a/Engine/Source/Editor/LandscapeEditor/Private/SLandscapeEditor.cpp b/Engine/Source/Editor/LandscapeEditor/Private/SLandscapeEditor.cpp
--- a/Engine/Source/Editor/LandscapeEditor/Private/SLandscapeEditor.cpp
+++ b/Engine/Source/Editor/LandscapeEditor/Private/SLandscapeEditor.cpp
@@ -26,16 +26,6 @@
 
 #define LOCTEXT_NAMESPACE "LandscapeEditor"
 
-namespace UE::Landscape::Private {
-	static bool bEnableRetopoTool = false;
-	static FAutoConsoleVariableRef CVarEnableRetopoTool(
-		TEXT("landscape.EnableRetopologizeTool"),
-		bEnableRetopoTool,
-		TEXT("Enable the Retopologize tool.  The tool will be fully deprecated in UE5.7, but this cvar will enable it for 5.6"),
-		ECVF_Default
-	);
-}
-
 void SLandscapeAssetThumbnail::Construct(const FArguments& InArgs, UObject* Asset, TSharedRef<FAssetThumbnailPool> ThumbnailPool, const FName& ClassThumbnailBrushOverride)
 {
 	FIntPoint ThumbnailSize = InArgs._ThumbnailSize;
@@ -150,7 +140,6 @@ void FLandscapeToolKit::Init(const TSharedPtr<IToolkitHost>& InitToolkitHost)
 	MAP_TOOL("Erosion");
 	MAP_TOOL("HydraErosion");
 	MAP_TOOL("Noise");
-	MAP_TOOL("Retopologize");
 	MAP_TOOL("Visibility");
 	MAP_TOOL("BlueprintBrush");
 
@@ -268,10 +257,6 @@ void FLandscapeToolKit::BuildToolPalette(FName PaletteName, class FToolBarBuilde
 		ToolBarBuilder.AddToolBarButton(Commands.ErosionTool);
 		ToolBarBuilder.AddToolBarButton(Commands.HydroErosionTool);
 		ToolBarBuilder.AddToolBarButton(Commands.NoiseTool);
-		if (!Settings->InRestrictiveMode())
-		{
-			ToolBarBuilder.AddToolBarButton(Commands.RetopologizeTool);
-		}
 		ToolBarBuilder.AddToolBarButton(Commands.VisibilityTool);
 
 		ToolBarBuilder.AddToolBarButton(Commands.MirrorTool);
@@ -429,12 +414,6 @@ bool FLandscapeToolKit::IsToolEnabled(FName ToolName) const
 			return true;
 		}
 
-		if (ToolName == "Retopologize")
-		{
-			// TODO [jared.ritchie] this temporarily hides tool but need to fully remove landscape's retopologize tool / XY offset texture
-			return false;
-		}
-
 		// Other tools are available if there is an existing landscape.
 		if (LandscapeEdMode->GetLandscapeList().Num() > 0)
 		{
diff --git a/Engine/Source/Editor/StatsViewer/Private/StatsPages/PrimitiveStatsPage.cpp b/Engine/Source/Editor/StatsViewer/Private/StatsPages/PrimitiveStatsPage.cpp
--- a/Engine/Source/Editor/StatsViewer/Private/StatsPages/PrimitiveStatsPage.cpp
+++ b/Engine/Source/Editor/StatsViewer/Private/StatsPages/PrimitiveStatsPage.cpp
@@ -362,15 +362,6 @@ struct PrimitiveStatsGenerator
 							const SIZE_T HeightmapResourceSize = CurrentComponent->GetHeightmap()->GetResourceSizeBytes(EResourceSizeMode::EstimatedTotal);
 							NewStatsEntry->ResourceSize += (float)HeightmapResourceSize / 1024.0f;
 						}
-						if (CurrentComponent->XYOffsetmapTexture)
-						{
-							UniqueTextures.Add(CurrentComponent->XYOffsetmapTexture, &bNotUnique);
-							if (!bNotUnique)
-							{
-								const SIZE_T OffsetmapResourceSize = CurrentComponent->XYOffsetmapTexture->GetResourceSizeBytes(EResourceSizeMode::EstimatedTotal);
-								NewStatsEntry->ResourceSize += (float)OffsetmapResourceSize / 1024.f;
-							}
-						}
 
 						const TArray<UTexture2D*>& ComponentWeightmapTextures = CurrentComponent->GetWeightmapTextures();
 
diff --git a/Engine/Source/Editor/UnrealEd/Private/Lightmass/Lightmass.cpp b/Engine/Source/Editor/UnrealEd/Private/Lightmass/Lightmass.cpp
--- a/Engine/Source/Editor/UnrealEd/Private/Lightmass/Lightmass.cpp
+++ b/Engine/Source/Editor/UnrealEd/Private/Lightmass/Lightmass.cpp
@@ -2786,7 +2786,6 @@ FLightmassProcessor::~FLightmassProcessor()
 	delete DeferredMappings;
 
 	FLandscapeStaticLightingMesh::LandscapeUpscaleHeightDataCache.Empty();
-	FLandscapeStaticLightingMesh::LandscapeUpscaleXYOffsetDataCache.Empty();
 }
 
 /** Retrieve an exporter for the given channel name */
diff --git a/Engine/Source/Runtime/Landscape/Classes/LandscapeComponent.h b/Engine/Source/Runtime/Landscape/Classes/LandscapeComponent.h
--- a/Engine/Source/Runtime/Landscape/Classes/LandscapeComponent.h
+++ b/Engine/Source/Runtime/Landscape/Classes/LandscapeComponent.h
@@ -462,9 +462,9 @@ class ULandscapeComponent : public UPrimitiveComponent
 	UPROPERTY(TextExportTransient)
 	TArray<int8> LODIndexToMaterialIndex;
 
-	/** XYOffsetmap texture reference */
+	UE_DEPRECATED(5.7, "RetopologizeTool/XYOffset deprecated with the removal of non-edit layer landscapes")
 	UPROPERTY()
-	TObjectPtr<UTexture2D> XYOffsetmapTexture;
+	TObjectPtr<UTexture2D> XYOffsetmapTexture_DEPRECATED;
 
 	/** UV offset to component's weightmap data from component local coordinates*/
 	UPROPERTY()
@@ -1134,10 +1134,8 @@ public:
 	 * @param HeightmapTextureMipData: heightmap data
 	 * @param ComponentX1, ComponentY1, ComponentX2, ComponentY2: region to update
 	 * @param bUpdateBounds: Whether to update bounds from render component.
-	 * @param XYOffsetTextureMipData: xy-offset map data
-	 * @returns True if CollisionComponent was created in this update.
 	 */
-	void UpdateCollisionHeightData(const FColor* HeightmapTextureMipData, const FColor* SimpleCollisionHeightmapTextureData, int32 ComponentX1=0, int32 ComponentY1=0, int32 ComponentX2=MAX_int32, int32 ComponentY2=MAX_int32, bool bUpdateBounds=false, const FColor* XYOffsetTextureMipData=nullptr, bool bInUpdateHeightfieldRegion=true);
+	void UpdateCollisionHeightData(const FColor* HeightmapTextureMipData, const FColor* SimpleCollisionHeightmapTextureData, int32 ComponentX1=0, int32 ComponentY1=0, int32 ComponentX2=MAX_int32, int32 ComponentY2=MAX_int32, bool bUpdateBounds=false, bool bInUpdateHeightfieldRegion=true);
 
 	/**
 	 * Deletes Collision Component
@@ -1280,9 +1278,8 @@ protected:
 #if WITH_EDITOR
 	void RecreateCollisionComponent(bool bUseSimpleCollision);
 	void UpdateCollisionHeightBuffer(int32 InComponentX1, int32 InComponentY1, int32 InComponentX2, int32 InComponentY2, int32 InCollisionMipLevel, int32 InHeightmapSizeU, int32 InHeightmapSizeV,
-		const FColor* const InHeightmapTextureMipData, uint16* CollisionHeightData, uint16* GrassHeightData,
-		const FColor* const InXYOffsetTextureMipData, uint16* CollisionXYOffsetData);
-	void UpdateDominantLayerBuffer(int32 InComponentX1, int32 InComponentY1, int32 InComponentX2, int32 InComponentY2, int32 InCollisionMipLevel, int32 InWeightmapSizeU, int32 InDataLayerIdx, const TArray<uint8*>& InCollisionDataPtrs, const TArray<ULandscapeLayerInfoObject*>& InLayerInfos, uint8* DominantLayerData);
+		const FColor* const InHeightmapTextureMipData, uint16* CollisionHeightData, uint16* GrassHeightData);
+	void UpdateDominantLayerBuffer(int32 InComponentX1, int32 InComponentY1, int32 InComponentX2, int32 InComponentY2, int32 InCollisionMipLevel, int32 InWeightmapSizeU, int32 InVisibilityLayerIdx, const TArray<uint8*>& InCollisionDataPtrs, const TArray<ULandscapeLayerInfoObject*>& InLayerInfos, uint8* DominantLayerData);
 #endif
 
 	/** Whether the component type supports static lighting. */
diff --git a/Engine/Source/Runtime/Landscape/Classes/LandscapeMeshCollisionComponent.h b/Engine/Source/Runtime/Landscape/Classes/LandscapeMeshCollisionComponent.h
--- a/Engine/Source/Runtime/Landscape/Classes/LandscapeMeshCollisionComponent.h
+++ b/Engine/Source/Runtime/Landscape/Classes/LandscapeMeshCollisionComponent.h
@@ -15,31 +15,21 @@
 
 #include "LandscapeMeshCollisionComponent.generated.h"
 
-class UPhysicalMaterial;
-struct FNavigableGeometryExport;
-
 namespace Chaos
 {
 	class FTriangleMeshImplicitObject;
 }
 
 UCLASS()
-class ULandscapeMeshCollisionComponent : public ULandscapeHeightfieldCollisionComponent
+class UE_DEPRECATED(5.7, "RetopologizeTool/XYOffset deprecated with the removal of non-edit layer landscapes") ULandscapeMeshCollisionComponent_DEPRECATED : public ULandscapeHeightfieldCollisionComponent
 {
 	GENERATED_BODY()
 
 public:
-
-	ULandscapeMeshCollisionComponent();
-	virtual ~ULandscapeMeshCollisionComponent();
-
-	// Keep the possibility to share projected height field PhysX object with editor mesh collision objects...
-
-	/** Guid used to share PhysX heightfield objects in the editor */
 	UPROPERTY()
 	FGuid MeshGuid;
 
-	struct FTriMeshGeometryRef : public FRefCountedObject
+	struct UE_DEPRECATED(5.7, "RetopologizeTool/XYOffset deprecated with the removal of non-edit layer landscapes") FTriMeshGeometryRef : public FRefCountedObject
 	{
 		FGuid Guid;
 
@@ -64,41 +54,4 @@ public:
 
 	/** Physics engine version of heightfield data. */
 	TRefCountPtr<FTriMeshGeometryRef> MeshRef;
-
-	//~ Begin UActorComponent Interface.
-protected:
-	virtual void OnCreatePhysicsState() override;
-public:
-	virtual void ApplyWorldOffset(const FVector& InOffset, bool bWorldShift) override;
-	//~ End UActorComponent Interface.
-
-	//~ Begin USceneComponent Interface.
-	virtual void DestroyComponent(bool bPromoteChildren = false) override;
-	//~ End USceneComponent Interface.
-
-	//~ Begin UPrimitiveComponent Interface
-	virtual bool DoCustomNavigableGeometryExport(FNavigableGeometryExport& GeomExport) const override;
-	//End UPrimitiveComponent interface
-
-	//~ Begin INavRelevantInterface Interface
-	virtual bool SupportsGatheringGeometrySlices() const override { return false; }
-	//~ End INavRelevantInterface Interface
-
-	//~ Begin UObject Interface.
-	virtual void Serialize(FArchive& Ar) override;
-	virtual void BeginDestroy() override;
-	virtual void GetResourceSizeEx(FResourceSizeEx& CumulativeResourceSize) override;
-#if WITH_EDITOR
-	virtual void ExportCustomProperties(FOutputDevice& Out, uint32 Indent) override;
-	virtual void ImportCustomProperties(const TCHAR* SourceText, FFeedbackContext* Warn) override;
-
-	virtual bool CookCollisionData(const FName& Format, bool bUseDefaultMaterialOnly, bool bCheckDDC, TArray<uint8>& OutCookedData, TArray<UPhysicalMaterial*>& InOutMaterials) const override;
-	virtual uint32 ComputeCollisionHash() const override { return 0; }
-#endif
-	//~ End UObject Interface.
-
-	//~ Begin ULandscapeHeightfieldCollisionComponent Interface
-	virtual void CreateCollisionObject() override;
-	virtual bool RecreateCollision() override;
-	//~ End ULandscapeHeightfieldCollisionComponent Interface
 };
diff --git a/Engine/Source/Runtime/Landscape/Classes/LandscapeProxy.h b/Engine/Source/Runtime/Landscape/Classes/LandscapeProxy.h
--- a/Engine/Source/Runtime/Landscape/Classes/LandscapeProxy.h
+++ b/Engine/Source/Runtime/Landscape/Classes/LandscapeProxy.h
@@ -774,9 +774,7 @@ PRAGMA_ENABLE_DEPRECATION_WARNINGS
 	UPROPERTY(EditAnywhere, AdvancedDisplay, BlueprintReadOnly, Category = Lighting, meta = (LandscapeOverridable))
 	FLightingChannels LightingChannels;
 
-	/** Whether to use the landscape material's vertical world position offset when calculating static lighting.
-	* Note: Only z (vertical) offset is supported. XY offsets are ignored.
-	* Does not work correctly with an XY offset map (mesh collision) */
+	/** Whether to use the landscape material's vertical world position offset when calculating static lighting. Z (vertical) offset is supported */
 	UPROPERTY(EditAnywhere, AdvancedDisplay, Category = Lighting, meta = (LandscapeInherited))
 	uint32 bUseMaterialPositionOffsetInStaticLighting:1;
 
@@ -830,7 +828,6 @@ PRAGMA_ENABLE_DEPRECATION_WARNINGS
 	int32 CollisionMipLevel;
 
 	// If set higher than the "Collision Mip Level", this specifies the Landscape LOD to use for "simple collision" tests, otherwise the "Collision Mip Level" is used for both simple and complex collision.
-	// Does not work with an XY offset map (mesh collision)
 	UPROPERTY(EditAnywhere, Category=Collision, meta = (LandscapeOverridable))
 	int32 SimpleCollisionMipLevel;
 
@@ -848,9 +845,7 @@ PRAGMA_ENABLE_DEPRECATION_WARNINGS
 	UPROPERTY(EditAnywhere, BlueprintReadOnly, Category = Collision, meta = (LandscapeOverridable))
 	uint32 bGenerateOverlapEvents : 1;
 
-	/** Whether to bake the landscape material's vertical world position offset into the collision heightfield.
-		Note: Only z (vertical) offset is supported. XY offsets are ignored.
-		Does not work with an XY offset map (mesh collision) */
+	/** Whether to bake the landscape material's vertical world position offset into the collision heightfield. Only z (vertical) offset is supported. */
 	UPROPERTY(EditAnywhere, AdvancedDisplay, Category=Collision, meta = (LandscapeInherited))
 	uint32 bBakeMaterialPositionOffsetIntoCollision:1;
 
diff --git a/Engine/Source/Runtime/Landscape/Private/Landscape.cpp b/Engine/Source/Runtime/Landscape/Private/Landscape.cpp
--- a/Engine/Source/Runtime/Landscape/Private/Landscape.cpp
+++ b/Engine/Source/Runtime/Landscape/Private/Landscape.cpp
@@ -837,7 +837,6 @@ void ULandscapeComponent::Serialize(FArchive& Ar)
 
 		TArray<UObject**> TexturesAndMaterials;
 		TexturesAndMaterials.Add((UObject**)&HeightmapTexture);
-		TexturesAndMaterials.Add((UObject**)&XYOffsetmapTexture);
 		for (TObjectPtr<UTexture2D>& WeightmapTexture : WeightmapTextures)
 		{
 			TexturesAndMaterials.Add((UObject**)&static_cast<UTexture2D*&>(WeightmapTexture));
@@ -898,17 +897,14 @@ void ULandscapeComponent::Serialize(FArchive& Ar)
 		if (!Ar.CookingTarget()->SupportsFeature(ETargetPlatformFeatures::DeferredRendering))
 		{
 			// These are used for SM5 rendering
-			UTexture2D* BackupXYOffsetmapTexture = nullptr;
 			TArray<TObjectPtr<UMaterialInstanceConstant>> BackupMaterialInstances;
 			TArray<TObjectPtr<UTexture2D>> BackupWeightmapTextures;
 
-			Exchange(BackupXYOffsetmapTexture, XYOffsetmapTexture);
 			Exchange(BackupMaterialInstances, MaterialInstances);
 			Exchange(BackupWeightmapTextures, WeightmapTextures);
 
 			Super::Serialize(Ar);
 
-			Exchange(BackupXYOffsetmapTexture, XYOffsetmapTexture);
 			Exchange(BackupMaterialInstances, MaterialInstances);
 			Exchange(BackupWeightmapTextures, WeightmapTextures);
 		}
@@ -1251,14 +1247,7 @@ void ULandscapeComponent::PostLoad()
 		SetupPrecachePSOParams(PrecachePSOParams);
 
 		FPSOPrecacheVertexFactoryDataList VertexFactoryDataList;
-		if (!XYOffsetmapTexture)
-		{
-			VertexFactoryDataList.Add(FPSOPrecacheVertexFactoryData(&FLandscapeVertexFactory::StaticType));
-		}
-		else
-		{
-			VertexFactoryDataList.Add(FPSOPrecacheVertexFactoryData(&FLandscapeXYOffsetVertexFactory::StaticType));
-		}
+		VertexFactoryDataList.Add(FPSOPrecacheVertexFactoryData(&FLandscapeVertexFactory::StaticType));
 
 		// we need the fixed grid vertex factory for both virtual texturing and grass
 		if (NeedsFixedGridVertexFactory(GMaxRHIShaderPlatform))
@@ -1420,7 +1409,6 @@ PRAGMA_ENABLE_DEPRECATION_WARNINGS
 	};
 
 	ReparentObject(HeightmapTexture);
-	ReparentObject(XYOffsetmapTexture);
 
 	for (UTexture2D* WeightmapTexture : WeightmapTextures)
 	{
@@ -1909,11 +1897,6 @@ TArray<UTexture*> ULandscapeComponent::GetGeneratedTextures() const
 
 	OutTextures.Append(WeightmapTextures);
 
-	if (XYOffsetmapTexture)
-	{
-		OutTextures.Add(XYOffsetmapTexture);
-	}
-
 	TArray<UMaterialInstance*> OutMaterials;
 	for (UMaterialInstance* MaterialInstance : MaterialInstances)
 	{
diff --git a/Engine/Source/Runtime/Landscape/Private/LandscapeCollision.cpp b/Engine/Source/Runtime/Landscape/Private/LandscapeCollision.cpp
--- a/Engine/Source/Runtime/Landscape/Private/LandscapeCollision.cpp
+++ b/Engine/Source/Runtime/Landscape/Private/LandscapeCollision.cpp
@@ -175,36 +175,28 @@ void ULandscapeHeightfieldCollisionComponent::FHeightfieldGeometryRef::GetResour
 	}
 }
 
-TMap<FGuid, ULandscapeMeshCollisionComponent::FTriMeshGeometryRef* > GSharedMeshRefs;
-
-ULandscapeMeshCollisionComponent::FTriMeshGeometryRef::FTriMeshGeometryRef()
-{}
+PRAGMA_DISABLE_DEPRECATION_WARNINGS
+// Deprecated
+ULandscapeMeshCollisionComponent_DEPRECATED::FTriMeshGeometryRef::FTriMeshGeometryRef()
+{
+}
 
-ULandscapeMeshCollisionComponent::FTriMeshGeometryRef::FTriMeshGeometryRef(FGuid& InGuid)
+// Deprecated
+ULandscapeMeshCollisionComponent_DEPRECATED::FTriMeshGeometryRef::FTriMeshGeometryRef(FGuid& InGuid)
 	: Guid(InGuid)
-{}
-
-ULandscapeMeshCollisionComponent::FTriMeshGeometryRef::~FTriMeshGeometryRef()
 {
-	// Remove ourselves from the shared map.
-	GSharedMeshRefs.Remove(Guid);
 }
 
-void ULandscapeMeshCollisionComponent::FTriMeshGeometryRef::GetResourceSizeEx(FResourceSizeEx& CumulativeResourceSize)
+// Deprecated
+ULandscapeMeshCollisionComponent_DEPRECATED::FTriMeshGeometryRef::~FTriMeshGeometryRef()
 {
-	CumulativeResourceSize.AddDedicatedSystemMemoryBytes(sizeof(*this));
-
-	CumulativeResourceSize.AddDedicatedSystemMemoryBytes(UsedChaosMaterials.GetAllocatedSize());
+}
 
-	if (TrimeshGeometry.IsValid())
-	{
-		TArray<uint8> Data;
-		FMemoryWriter MemAr(Data);
-		Chaos::FChaosArchive ChaosAr(MemAr);
-		TrimeshGeometry->Serialize(ChaosAr);
-		CumulativeResourceSize.AddDedicatedSystemMemoryBytes(Data.Num());
-	}
+// Deprecated
+void ULandscapeMeshCollisionComponent_DEPRECATED::FTriMeshGeometryRef::GetResourceSizeEx(FResourceSizeEx& CumulativeResourceSize)
+{
 }
+PRAGMA_ENABLE_DEPRECATION_WARNINGS
 
 // Generate a new guid to force a recache of landscape collison derived data
 #define LANDSCAPE_COLLISION_DERIVEDDATA_VER	TEXT("75E2F3A08BE44420813DD2F2AD34021D")
@@ -1473,556 +1465,6 @@ bool ULandscapeHeightfieldCollisionComponent::CookCollisionData(const FName& For
 	return Succeeded;
 }
 
-bool ULandscapeMeshCollisionComponent::CookCollisionData(const FName& Format, bool bUseDefaultMaterialOnly, bool bCheckDDC, TArray<uint8>& OutCookedData, TArray<UPhysicalMaterial*>& InOutMaterials) const
-{
-	// Use existing cooked data unless !bCheckDDC in which case the data must be rebuilt.
-	if (bCheckDDC && OutCookedData.Num() > 0)
-	{
-		return true;
-	}
-
-	COOK_STAT(auto Timer = LandscapeCollisionCookStats::MeshUsageStats.TimeSyncWork());
-	// we have 2 versions of collision objects
-	const int32 CookedDataIndex = bUseDefaultMaterialOnly ? 0 : 1;
-
-	if (!GLandscapeCollisionSkipDDC && bCheckDDC)
-	{
-		// Ensure that content was saved with physical materials before using DDC data
-		if (GetLinkerUEVersion() >= VER_UE4_LANDSCAPE_SERIALIZE_PHYSICS_MATERIALS && MeshGuid.IsValid())
-		{
-			FString DDCKey = GetHFDDCKeyString(Format, bUseDefaultMaterialOnly, MeshGuid, InOutMaterials);
-
-			// Check if the speculatively-loaded data loaded and is what we wanted
-			if (SpeculativeDDCRequest.IsValid() && DDCKey == SpeculativeDDCRequest->GetKey())
-			{
-				// If we have a DDC request in flight, just time the synchronous cycles used.
-				COOK_STAT(auto WaitTimer = LandscapeCollisionCookStats::MeshUsageStats.TimeAsyncWait());
-				SpeculativeDDCRequest->WaitAsynchronousCompletion();
-				bool bSuccess = SpeculativeDDCRequest->GetAsynchronousResults(OutCookedData);
-				// World will clean up remaining reference
-				SpeculativeDDCRequest.Reset();
-				if (bSuccess)
-				{
-					COOK_STAT(Timer.Cancel());
-					COOK_STAT(WaitTimer.AddHit(OutCookedData.Num()));
-					bShouldSaveCookedDataToDDC[CookedDataIndex] = false;
-					return true;
-				}
-				else
-				{
-					// If the DDC request failed, then we waited for nothing and will build the resource anyway. Just ignore the wait timer and treat it all as sync time.
-					COOK_STAT(WaitTimer.Cancel());
-				}
-			}
-
-			if (GetDerivedDataCacheRef().GetSynchronous(*DDCKey, OutCookedData, GetPathName()))
-			{
-				COOK_STAT(Timer.AddHit(OutCookedData.Num()));
-				bShouldSaveCookedDataToDDC[CookedDataIndex] = false;
-				return true;
-			}
-		}
-	}
-	
-	ALandscapeProxy* Proxy = GetLandscapeProxy();
-	UPhysicalMaterial* DefMaterial = (Proxy && Proxy->DefaultPhysMaterial != nullptr) ? Proxy->DefaultPhysMaterial : GEngine->DefaultPhysMaterial;
-
-	// List of materials which is actually used by trimesh
-	InOutMaterials.Empty();
-
-	TArray<FVector3f>		Vertices;
-	TArray<FTriIndices>		Indices;
-	TArray<uint16>			MaterialIndices;
-
-	const int32 CollisionSizeVerts = CollisionSizeQuads + 1;
-	const int32 SimpleCollisionSizeVerts = SimpleCollisionSizeQuads > 0 ? SimpleCollisionSizeQuads + 1 : 0;
-	const int32 NumVerts = FMath::Square(CollisionSizeVerts);
-	const int32 NumSimpleVerts = FMath::Square(SimpleCollisionSizeVerts);
-
-	const uint16* Heights = (const uint16*)CollisionHeightData.LockReadOnly();
-	const uint16* XYOffsets = (const uint16*)CollisionXYOffsetData.LockReadOnly();
-	checkf(CollisionHeightData.GetElementCount() == NumVerts + NumSimpleVerts,
-		TEXT("Invalid collision height data element count : %d found while there should be (%d samples + %d simple samples)"), CollisionHeightData.GetElementCount(), NumVerts, NumSimpleVerts);
-	check(CollisionXYOffsetData.GetElementCount() == NumVerts * 2);
-
-	const uint8* DominantLayers = nullptr;
-	if (DominantLayerData.GetElementCount() > 0)
-	{
-		DominantLayers = (const uint8*)DominantLayerData.LockReadOnly();
-	}
-
-	// Scale all verts into temporary vertex buffer.
-	Vertices.SetNumUninitialized(NumVerts);
-	for (int32 i = 0; i < NumVerts; i++)
-	{
-		int32 X = i % CollisionSizeVerts;
-		int32 Y = i / CollisionSizeVerts;
-		Vertices[i].Set(X + ((float)XYOffsets[i * 2] - 32768.f) * LANDSCAPE_XYOFFSET_SCALE, Y + ((float)XYOffsets[i * 2 + 1] - 32768.f) * LANDSCAPE_XYOFFSET_SCALE, LandscapeDataAccess::GetLocalHeight(Heights[i]));
-	}
-
-	const int32 NumTris = FMath::Square(CollisionSizeQuads) * 2;
-	Indices.SetNumUninitialized(NumTris);
-	if (DominantLayers)
-	{
-		MaterialIndices.SetNumUninitialized(NumTris);
-	}
-
-	int32 TriangleIdx = 0;
-	for (int32 y = 0; y < CollisionSizeQuads; y++)
-	{
-		for (int32 x = 0; x < CollisionSizeQuads; x++)
-		{
-			int32 DataIdx = x + y * CollisionSizeVerts;
-			bool bHole = false;
-
-			int32 MaterialIndex = 0; // Default physical material.
-			if (!bUseDefaultMaterialOnly && DominantLayers)
-			{
-				uint8 DominantLayerIdx = DominantLayers[DataIdx];
-				if (ComponentLayerInfos.IsValidIndex(DominantLayerIdx))
-				{
-					ULandscapeLayerInfoObject* Layer = ComponentLayerInfos[DominantLayerIdx];
-					if (Layer == ALandscapeProxy::VisibilityLayer)
-					{
-						// If it's a hole, override with the hole flag.
-						bHole = true;
-					}
-					else
-					{
-						UPhysicalMaterial* DominantMaterial = Layer && Layer->PhysMaterial ? ToRawPtr(Layer->PhysMaterial) : DefMaterial;
-						MaterialIndex = InOutMaterials.AddUnique(DominantMaterial);
-					}
-				}
-			}
-
-			FTriIndices& TriIndex1 = Indices[TriangleIdx];
-			if (bHole)
-			{
-				TriIndex1.v0 = (x + 0) + (y + 0) * CollisionSizeVerts;
-				TriIndex1.v1 = TriIndex1.v0;
-				TriIndex1.v2 = TriIndex1.v0;
-			}
-			else
-			{
-				TriIndex1.v0 = (x + 0) + (y + 0) * CollisionSizeVerts;
-				TriIndex1.v1 = (x + 1) + (y + 1) * CollisionSizeVerts;
-				TriIndex1.v2 = (x + 1) + (y + 0) * CollisionSizeVerts;
-			}
-
-			if (DominantLayers)
-			{
-				MaterialIndices[TriangleIdx] = static_cast<uint16>(MaterialIndex);
-			}
-			TriangleIdx++;
-
-			FTriIndices& TriIndex2 = Indices[TriangleIdx];
-			if (bHole)
-			{
-				TriIndex2.v0 = (x + 0) + (y + 0) * CollisionSizeVerts;
-				TriIndex2.v1 = TriIndex2.v0;
-				TriIndex2.v2 = TriIndex2.v0;
-			}
-			else
-			{
-				TriIndex2.v0 = (x + 0) + (y + 0) * CollisionSizeVerts;
-				TriIndex2.v1 = (x + 0) + (y + 1) * CollisionSizeVerts;
-				TriIndex2.v2 = (x + 1) + (y + 1) * CollisionSizeVerts;
-			}
-
-			if (DominantLayers)
-			{
-				MaterialIndices[TriangleIdx] = static_cast<uint16>(MaterialIndex);
-			}
-			TriangleIdx++;
-		}
-	}
-
-	CollisionHeightData.Unlock();
-	CollisionXYOffsetData.Unlock();
-	if (DominantLayers)
-	{
-		DominantLayerData.Unlock();
-	}
-
-	// Add the default physical material to be used used when we have no dominant data.
-	if (InOutMaterials.Num() == 0)
-	{
-		InOutMaterials.Add(DefMaterial);
-	}
-
-	TArray<uint8> OutData;
-	bool Result = false;
-
-	FCookBodySetupInfo CookInfo;
-	FTriMeshCollisionData& MeshDesc = CookInfo.TriangleMeshDesc;
-	MeshDesc.bFlipNormals = true;
-	MeshDesc.Vertices = MoveTemp(Vertices);
-	MeshDesc.Indices = MoveTemp(Indices);
-	MeshDesc.MaterialIndices = MoveTemp(MaterialIndices);
-	CookInfo.bCookTriMesh = true;
-	TArray<int32> FaceRemap;
-	TArray<int32> VertexRemap;
-	Chaos::FTriangleMeshImplicitObjectPtr Trimesh = Chaos::Cooking::BuildSingleTrimesh(MeshDesc, FaceRemap, VertexRemap);
-
-	if(Trimesh.IsValid())
-	{
-		FMemoryWriter Ar(OutData);
-		Chaos::FChaosArchive ChaosAr(Ar);
-		ChaosAr << Trimesh;
-		Result = OutData.Num() > 0;
-	}
-
-	if (Result)
-	{
-		COOK_STAT(Timer.AddMiss(OutData.Num()));
-		OutCookedData.SetNumUninitialized(OutData.Num());
-		FMemory::Memcpy(OutCookedData.GetData(), OutData.GetData(), OutData.Num());
-
-		if (!GLandscapeCollisionSkipDDC && bShouldSaveCookedDataToDDC[CookedDataIndex] && MeshGuid.IsValid())
-		{
-			GetDerivedDataCacheRef().Put(*GetHFDDCKeyString(Format, bUseDefaultMaterialOnly, MeshGuid, InOutMaterials), OutCookedData, GetPathName());
-			bShouldSaveCookedDataToDDC[CookedDataIndex] = false;
-		}
-	}
-	else
-	{
-		// We didn't actually build anything, so just track the cycles.
-		COOK_STAT(Timer.TrackCyclesOnly());
-		OutCookedData.Empty();
-		InOutMaterials.Empty();
-	}
-
-	return Result;
-}
-#endif //WITH_EDITOR
-
-void ULandscapeMeshCollisionComponent::CreateCollisionObject()
-{
-	// If we have not created a heightfield yet - do it now.
-	if (!IsValidRef(MeshRef))
-	{
-		FTriMeshGeometryRef* ExistingMeshRef = nullptr;
-		bool bCheckDDC = true;
-
-		if (!MeshGuid.IsValid())
-		{
-			MeshGuid = FGuid::NewGuid();
-			bCheckDDC = false;
-		}
-		else
-		{
-			// Look for a heightfield object with the current Guid (this occurs with PIE)
-			ExistingMeshRef = GSharedMeshRefs.FindRef(MeshGuid);
-		}
-
-		if (ExistingMeshRef)
-		{
-			MeshRef = ExistingMeshRef;
-		}
-		else
-		{
-#if WITH_EDITOR
-		    // This should only occur if a level prior to VER_UE4_LANDSCAPE_COLLISION_DATA_COOKING 
-		    // was resaved using a commandlet and not saved in the editor, or if a PhysicalMaterial asset was deleted.
-		    if (CookedPhysicalMaterials.Num() == 0 || CookedPhysicalMaterials.Contains(nullptr))
-		    {
-			    bCheckDDC = false;
-		    }
-
-			// Create cooked physics data
-			static FName PhysicsFormatName(FPlatformProperties::GetPhysicsFormat());
-			CookCollisionData(PhysicsFormatName, false, bCheckDDC, CookedCollisionData, MutableView(CookedPhysicalMaterials));
-#endif // WITH_EDITOR
-
-			if (CookedCollisionData.IsEmpty())
-			{
-				if (bCookedCollisionDataWasDeleted)
-				{
-					// only complain if we actually deleted the data.. otherwise it may have been intentional?
-					UE_LOG(LogLandscape, Warning, TEXT("Tried to create mesh collision for component '%s', but the collision data was deleted!"), *GetName());
-				}
-			}
-			else
-			{
-				MeshRef = GSharedMeshRefs.Add(MeshGuid, new FTriMeshGeometryRef(MeshGuid));
-
-				// Create physics objects
-				FMemoryReader Reader(CookedCollisionData);
-				Chaos::FChaosArchive Ar(Reader);
-				Ar << MeshRef->TrimeshGeometry;
-
-				for (UPhysicalMaterial* PhysicalMaterial : CookedPhysicalMaterials)
-				{
-					MeshRef->UsedChaosMaterials.Add(PhysicalMaterial->GetPhysicsMaterial());
-				}
-
-				// Release cooked collison data
-				// In cooked builds created collision object will never be deleted while component is alive, so we don't need this data anymore
-				if (FPlatformProperties::RequiresCookedData() || GetWorld()->IsGameWorld())
-				{
-					CookedCollisionData.Empty();
-					bCookedCollisionDataWasDeleted = true;
-				}
-
-#if WITH_EDITOR
-				// Create collision mesh for the landscape editor (no holes in it)
-				if (!GetWorld()->IsGameWorld())
-				{
-					TArray<UPhysicalMaterial*> CookedMaterialsEd;
-					if (CookCollisionData(PhysicsFormatName, true, bCheckDDC, CookedCollisionDataEd, CookedMaterialsEd))
-					{
-						FMemoryReader EdReader(CookedCollisionData);
-						Chaos::FChaosArchive EdAr(EdReader);
-						EdAr << MeshRef->EditorTrimeshGeometry;
-					}
-				}
-#endif //WITH_EDITOR
-			}
-		}
-	}
-}
-
-
-ULandscapeMeshCollisionComponent::ULandscapeMeshCollisionComponent()
-{
-	// make landscape always create? 
-	bAlwaysCreatePhysicsState = true;
-}
-
-ULandscapeMeshCollisionComponent::~ULandscapeMeshCollisionComponent() = default;
-
-struct FMeshCollisionInitHelper
-{
-	FMeshCollisionInitHelper() = delete;
-	FMeshCollisionInitHelper(TRefCountPtr<ULandscapeMeshCollisionComponent::FTriMeshGeometryRef> InMeshRef, UWorld* InWorld, UPrimitiveComponent* InComponent, FBodyInstance* InBodyInstance)
-		: ComponentToWorld(FTransform::Identity)
-		, ComponentScale(FVector::OneVector)
-		, CollisionScale(1.0f)
-		, MeshRef(InMeshRef)
-		, World(InWorld)
-		, Component(InComponent)
-		, TargetInstance(InBodyInstance)
-	{
-		check(World);
-		PhysScene = World->GetPhysicsScene();
-		check(PhysScene);
-		check(Component);
-		check(TargetInstance);
-	}
-
-	void SetComponentScale3D(const FVector& InScale)
-	{
-		ComponentScale = InScale;
-	}
-
-	void SetCollisionScale(float InScale)
-	{
-		CollisionScale = InScale;
-	}
-
-	void SetComponentToWorld(const FTransform& InTransform)
-	{
-		ComponentToWorld = InTransform;
-	}
-
-	void SetFilters(const FCollisionFilterData& InQueryFilter, const FCollisionFilterData& InSimulationFilter)
-	{
-		QueryFilter = InQueryFilter;
-		SimulationFilter = InSimulationFilter;
-	}
-
-	void SetEditorFilter(const FCollisionFilterData& InFilter)
-	{
-		QueryFilterEd = InFilter;
-	}
-
-	bool IsGeometryValid() const
-	{
-		return MeshRef->TrimeshGeometry.IsValid();
-	}
-
-	void CreateActors()
-	{
-		Chaos::FShapesArray ShapeArray;
-		TArray<Chaos::FImplicitObjectPtr> Geometries;
-		
-		FActorCreationParams Params;
-		Params.InitialTM = ComponentToWorld;
-		Params.InitialTM.SetScale3D(FVector::OneVector);
-		Params.bQueryOnly = false;
-		Params.bStatic = true;
-		Params.Scene = PhysScene;
-
-		FPhysicsInterface::CreateActor(Params, ActorHandle);
-
-		FVector Scale = FVector(ComponentScale.X * CollisionScale, ComponentScale.Y * CollisionScale, ComponentScale.Z);
-
-		{
-			Chaos::FImplicitObjectPtr ScaledTrimesh = MakeImplicitObjectPtr<Chaos::TImplicitObjectScaled<Chaos::FTriangleMeshImplicitObject>>(MeshRef->TrimeshGeometry, Scale);
-			TUniquePtr<Chaos::FPerShapeData> NewShape = Chaos::FShapeInstanceProxy::Make(ShapeArray.Num(), ScaledTrimesh);
-
-			NewShape->SetQueryData(QueryFilter);
-			NewShape->SetSimData(SimulationFilter);
-			NewShape->SetCollisionTraceType(Chaos::EChaosCollisionTraceFlag::Chaos_CTF_UseComplexAsSimple);
-			NewShape->SetMaterials(MeshRef->UsedChaosMaterials);
-
-			Geometries.Emplace(MoveTemp(ScaledTrimesh));
-			ShapeArray.Emplace(MoveTemp(NewShape));
-		}
-
-#if WITH_EDITOR
-		if(!World->IsGameWorld())
-		{
-			Chaos::FImplicitObjectPtr ScaledTrimeshEd = MakeImplicitObjectPtr<Chaos::TImplicitObjectScaled<Chaos::FTriangleMeshImplicitObject>>(MeshRef->EditorTrimeshGeometry, Scale);
-			TUniquePtr<Chaos::FPerShapeData> NewEdShape = Chaos::FShapeInstanceProxy::Make(ShapeArray.Num(), ScaledTrimeshEd);
-
-			NewEdShape->SetQueryData(QueryFilterEd);
-			NewEdShape->SetSimEnabled(false);
-			NewEdShape->SetCollisionTraceType(Chaos::EChaosCollisionTraceFlag::Chaos_CTF_UseComplexAsSimple);
-			NewEdShape->SetMaterial(GEngine->DefaultPhysMaterial->GetPhysicsMaterial());
-
-			Geometries.Emplace(MoveTemp(ScaledTrimeshEd));
-			ShapeArray.Emplace(MoveTemp(NewEdShape));
-		}
-#endif // WITH_EDITOR
-
-		if(Geometries.Num() == 1)
-		{
-			ActorHandle->GetGameThreadAPI().SetGeometry(Geometries[0]);
-		}
-		else
-		{
-			ActorHandle->GetGameThreadAPI().SetGeometry(MakeImplicitObjectPtr<Chaos::FImplicitObjectUnion>(MoveTemp(Geometries)));
-		}
-
-		for(TUniquePtr<Chaos::FPerShapeData>& Shape : ShapeArray)
-		{
-			Chaos::FRigidTransform3 WorldTransform = Chaos::FRigidTransform3(ActorHandle->GetGameThreadAPI().X(), ActorHandle->GetGameThreadAPI().R());
-			Shape->UpdateShapeBounds(WorldTransform);
-		}
-
-		ActorHandle->GetGameThreadAPI().MergeShapesArray(MoveTemp(ShapeArray));
-
-		TargetInstance->PhysicsUserData = FPhysicsUserData(TargetInstance);
-		TargetInstance->OwnerComponent = Component;
-		TargetInstance->SetPhysicsActor(ActorHandle);
-
-		ActorHandle->GetGameThreadAPI().SetUserData(&TargetInstance->PhysicsUserData);
-	}
-
-	void AddToScene()
-	{
-		check(PhysScene);
-
-		TArray<FPhysicsActorHandle> Actors;
-		Actors.Add(ActorHandle);
-
-		FPhysicsCommand::ExecuteWrite(PhysScene, [&]()
-		{
-			PhysScene->AddActorsToScene_AssumesLocked(Actors, true);
-		});
-		PhysScene->AddToComponentMaps(Component, ActorHandle);
-
-		if(TargetInstance->bNotifyRigidBodyCollision)
-		{
-			PhysScene->RegisterForCollisionEvents(Component);
-		}
-	}
-
-private:
-	FTransform ComponentToWorld;
-	FVector ComponentScale;
-	float CollisionScale;
-	TRefCountPtr<ULandscapeMeshCollisionComponent::FTriMeshGeometryRef> MeshRef;
-	FPhysScene* PhysScene;
-	FCollisionFilterData QueryFilter;
-	FCollisionFilterData SimulationFilter;
-	FCollisionFilterData QueryFilterEd;
-	UWorld* World;
-	UPrimitiveComponent* Component;
-	FBodyInstance* TargetInstance;
-
-	FPhysicsActorHandle ActorHandle;
-};
-
-
-void ULandscapeMeshCollisionComponent::OnCreatePhysicsState()
-{
-	USceneComponent::OnCreatePhysicsState(); // route OnCreatePhysicsState, skip PrimitiveComponent implementation
-
-	if (!BodyInstance.IsValidBodyInstance())
-	{
-		// This will do nothing, because we create trimesh at component PostLoad event, unless we destroyed it explicitly
-		CreateCollisionObject();
-
-		if (IsValidRef(MeshRef))
-		{
-			FMeshCollisionInitHelper Initializer(MeshRef, GetWorld(), this, &BodyInstance);
-
-			// Make transform for this landscape component PxActor
-			FTransform LandscapeComponentTransform = GetComponentToWorld();
-			FMatrix LandscapeComponentMatrix = LandscapeComponentTransform.ToMatrixWithScale();
-
-			FVector LandscapeScale = LandscapeComponentMatrix.ExtractScaling();
-
-			Initializer.SetComponentToWorld(LandscapeComponentTransform);
-			Initializer.SetComponentScale3D(LandscapeScale);
-			Initializer.SetCollisionScale(CollisionScale);
-
-			if(Initializer.IsGeometryValid())
-			{
-				// Setup filtering
-				FCollisionFilterData QueryFilterData, SimFilterData;
-				CreateShapeFilterData(static_cast<uint8>(GetCollisionObjectType()), FMaskFilter(0), GetOwner()->GetUniqueID(), GetCollisionResponseToChannels(), GetUniqueID(), 0, QueryFilterData, SimFilterData, false, false, true);
-				QueryFilterData.Word3 |= (EPDF_SimpleCollision | EPDF_ComplexCollision);
-				SimFilterData.Word3 |= (EPDF_SimpleCollision | EPDF_ComplexCollision);
-
-				Initializer.SetFilters(QueryFilterData, SimFilterData);
-
-#if WITH_EDITOR
-				FCollisionResponseContainer EdResponse;
-				EdResponse.SetAllChannels(ECollisionResponse::ECR_Ignore);
-				EdResponse.SetResponse(ECollisionChannel::ECC_Visibility, ECR_Block);
-						FCollisionFilterData QueryFilterDataEd, SimFilterDataEd;
-				CreateShapeFilterData(ECollisionChannel::ECC_Visibility, FMaskFilter(0), GetOwner()->GetUniqueID(), EdResponse, GetUniqueID(), 0, QueryFilterDataEd, SimFilterDataEd, true, false, true);
-						QueryFilterDataEd.Word3 |= (EPDF_SimpleCollision | EPDF_ComplexCollision);
-
-				Initializer.SetEditorFilter(QueryFilterDataEd);
-#endif // WITH_EDITOR
-
-				Initializer.CreateActors();
-				Initializer.AddToScene();
-			}
-			else
-			{
-				UE_LOG(LogLandscape, Warning, TEXT("ULandscapeMeshCollisionComponent::OnCreatePhysicsState(): TriMesh invalid"));
-			}
-		}
-	}
-}
-
-void ULandscapeMeshCollisionComponent::ApplyWorldOffset(const FVector& InOffset, bool bWorldShift)
-{
-	Super::ApplyWorldOffset(InOffset, bWorldShift);
-
-	if (!bWorldShift || !FPhysScene::SupportsOriginShifting())
-	{
-		RecreatePhysicsState();
-	}
-}
-
-void ULandscapeMeshCollisionComponent::DestroyComponent(bool bPromoteChildren/*= false*/)
-{
-	ALandscapeProxy* Proxy = GetLandscapeProxy();
-	if (Proxy)
-	{
-		Proxy->CollisionComponents.Remove(this);
-	}
-
-	Super::DestroyComponent(bPromoteChildren);
-}
-
-#if WITH_EDITOR
 uint32 ULandscapeHeightfieldCollisionComponent::ComputeCollisionHash() const
 {
 	uint32 Hash = 0;
@@ -2151,17 +1593,6 @@ void ULandscapeHeightfieldCollisionComponent::BeginDestroy()
 	}
 }
 
-void ULandscapeMeshCollisionComponent::BeginDestroy()
-{
-	if (!HasAnyFlags(RF_ClassDefaultObject))
-	{
-		MeshRef = nullptr;
-		MeshGuid = FGuid();
-	}
-
-	Super::BeginDestroy();
-}
-
 bool ULandscapeHeightfieldCollisionComponent::RecreateCollision()
 {
 	if (!HasAnyFlags(RF_ClassDefaultObject))
@@ -2372,21 +1803,6 @@ void ULandscapeHeightfieldCollisionComponent::SnapFoliageInstances(const FBox& I
 }
 #endif // WITH_EDITORONLY_DATA
 
-bool ULandscapeMeshCollisionComponent::RecreateCollision()
-{
-	TRefCountPtr<FTriMeshGeometryRef> TriMeshLifetimeExtender = nullptr; // Ensure heightfield data is alive until removed from physics world
-
-	if (!HasAnyFlags(RF_ClassDefaultObject))
-	{
-		TriMeshLifetimeExtender = MeshRef;
-		MeshRef = nullptr; // Ensure data will be recreated
-		MeshGuid = FGuid();
-		CachedHeightFieldSamples.Empty();
-	}
-
-	return Super::RecreateCollision();
-}
-
 void ULandscapeHeightfieldCollisionComponent::Serialize(FArchive& Ar)
 {
 	LLM_SCOPE(ELLMTag::Landscape);
@@ -2430,62 +1846,18 @@ void ULandscapeHeightfieldCollisionComponent::Serialize(FArchive& Ar)
 		else
 		{
 #if WITH_EDITORONLY_DATA
-			// For PIE, we won't need the source height data if we already have a shared reference to the heightfield
-			if (!(Ar.GetPortFlags() & PPF_DuplicateForPIE) || !HeightfieldGuid.IsValid() || GSharedMeshRefs.FindRef(HeightfieldGuid) == nullptr)
-			{
-				CollisionHeightData.Serialize(Ar, this);
-				DominantLayerData.Serialize(Ar, this);
+			CollisionHeightData.Serialize(Ar, this);
+			DominantLayerData.Serialize(Ar, this);
 
-				if (Ar.CustomVer(FFortniteMainBranchObjectVersion::GUID) >= FFortniteMainBranchObjectVersion::LandscapePhysicalMaterialRenderData)
-				{
-					PhysicalMaterialRenderData.Serialize(Ar, this);
-				}
+			if (Ar.CustomVer(FFortniteMainBranchObjectVersion::GUID) >= FFortniteMainBranchObjectVersion::LandscapePhysicalMaterialRenderData)
+			{
+				PhysicalMaterialRenderData.Serialize(Ar, this);
 			}
 #endif//WITH_EDITORONLY_DATA
 		}
 	}
 }
 
-void ULandscapeMeshCollisionComponent::Serialize(FArchive& Ar)
-{
-	LLM_SCOPE(ELLMTag::Landscape);
-
-	Super::Serialize(Ar);
-
-	if (Ar.UEVer() < VER_UE4_LANDSCAPE_COLLISION_DATA_COOKING)
-	{
-#if WITH_EDITORONLY_DATA
-		// conditional serialization in later versions
-		CollisionXYOffsetData.Serialize(Ar, this);
-#endif// WITH_EDITORONLY_DATA
-	}
-
-	// Physics cooking mesh data
-	bool bCooked = false;
-	if (Ar.UEVer() >= VER_UE4_ADD_COOKED_TO_LANDSCAPE)
-	{
-		bCooked = Ar.IsCooking();
-		Ar << bCooked;
-	}
-
-	if (FPlatformProperties::RequiresCookedData() && !bCooked && Ar.IsLoading())
-	{
-		UE_LOG(LogPhysics, Fatal, TEXT("This platform requires cooked packages, and physics data was not cooked into %s."), *GetFullName());
-	}
-
-	if (bCooked)
-	{
-		// triangle mesh cooked data should be serialized in ULandscapeHeightfieldCollisionComponent
-	}
-	else if (Ar.UEVer() >= VER_UE4_LANDSCAPE_COLLISION_DATA_COOKING)
-	{
-#if WITH_EDITORONLY_DATA		
-		// we serialize raw collision data only with non-cooked content
-		CollisionXYOffsetData.Serialize(Ar, this);
-#endif// WITH_EDITORONLY_DATA
-	}
-}
-
 bool ULandscapeHeightfieldCollisionComponent::IsShown(const FEngineShowFlags& ShowFlags) const
 {
 	return ShowFlags.Landscape;
@@ -2564,24 +1936,6 @@ void ULandscapeHeightfieldCollisionComponent::PrepareGeometryExportSync()
 	}
 }
 
-bool ULandscapeMeshCollisionComponent::DoCustomNavigableGeometryExport(FNavigableGeometryExport& GeomExport) const
-{
-	check(IsInGameThread());
-
-	if (IsValidRef(MeshRef))
-	{
-		FTransform MeshToW = GetComponentTransform();
-		MeshToW.MultiplyScale3D(FVector(CollisionScale, CollisionScale, 1.f));
-
-		if (MeshRef->TrimeshGeometry != nullptr)
-		{
-			GeomExport.ExportChaosTriMesh(MeshRef->TrimeshGeometry.GetReference(), MeshToW);
-		}
-	}
-
-	return false;
-}
-
 void ULandscapeHeightfieldCollisionComponent::PostLoad()
 {
 	Super::PostLoad();
@@ -2994,123 +2348,6 @@ void ULandscapeHeightfieldCollisionComponent::ImportCustomProperties(const TCHAR
 	}
 }
 
-void ULandscapeMeshCollisionComponent::ExportCustomProperties(FOutputDevice& Out, uint32 Indent)
-{
-	if (HasAnyFlags(RF_ClassDefaultObject))
-	{
-		return;
-	}
-
-	Super::ExportCustomProperties(Out, Indent);
-
-	uint16* XYOffsets = (uint16*)CollisionXYOffsetData.Lock(LOCK_READ_ONLY);
-	int32 NumOffsets = FMath::Square(CollisionSizeQuads + 1) * 2;
-	check(CollisionXYOffsetData.GetElementCount() == NumOffsets);
-
-	Out.Logf(TEXT("%sCustomProperties CollisionXYOffsetData "), FCString::Spc(Indent));
-	for (int32 i = 0; i < NumOffsets; i++)
-	{
-		Out.Logf(TEXT("%d "), XYOffsets[i]);
-	}
-
-	CollisionXYOffsetData.Unlock();
-	Out.Logf(TEXT("\r\n"));
-}
-
-void ULandscapeMeshCollisionComponent::ImportCustomProperties(const TCHAR* SourceText, FFeedbackContext* Warn)
-{
-	if (FParse::Command(&SourceText, TEXT("CollisionHeightData")))
-	{
-		int32 CollisionSizeVerts = CollisionSizeQuads + 1;
-		int32 SimpleCollisionSizeVerts = SimpleCollisionSizeQuads > 0 ? SimpleCollisionSizeQuads + 1 : 0;
-		int32 NumHeights = FMath::Square(CollisionSizeVerts) + FMath::Square(SimpleCollisionSizeVerts);
-
-		CollisionHeightData.Lock(LOCK_READ_WRITE);
-		uint16* Heights = (uint16*)CollisionHeightData.Realloc(NumHeights);
-		FMemory::Memzero(Heights, sizeof(uint16)*NumHeights);
-
-		FParse::Next(&SourceText);
-		int32 i = 0;
-		while (FChar::IsDigit(*SourceText))
-		{
-			if (i < NumHeights)
-			{
-				Heights[i++] = static_cast<uint16>(FCString::Atoi(SourceText));
-				while (FChar::IsDigit(*SourceText))
-				{
-					SourceText++;
-				}
-			}
-
-			FParse::Next(&SourceText);
-		}
-
-		CollisionHeightData.Unlock();
-
-		if (i != NumHeights)
-		{
-			Warn->Log(*NSLOCTEXT("Core", "SyntaxError", "Syntax Error").ToString());
-		}
-	}
-	else if (FParse::Command(&SourceText, TEXT("DominantLayerData")))
-	{
-		int32 NumDominantLayerSamples = FMath::Square(CollisionSizeQuads + 1);
-
-		DominantLayerData.Lock(LOCK_READ_WRITE);
-		uint8* DominantLayerSamples = (uint8*)DominantLayerData.Realloc(NumDominantLayerSamples);
-		FMemory::Memzero(DominantLayerSamples, NumDominantLayerSamples);
-
-		FParse::Next(&SourceText);
-		int32 i = 0;
-		while (SourceText[0] && SourceText[1])
-		{
-			if (i < NumDominantLayerSamples)
-			{
-				DominantLayerSamples[i++] = static_cast<uint8>(FParse::HexDigit(SourceText[0]) * 16 + FParse::HexDigit(SourceText[1]));
-			}
-			SourceText += 2;
-		}
-
-		DominantLayerData.Unlock();
-
-		if (i != NumDominantLayerSamples)
-		{
-			Warn->Log(*NSLOCTEXT("Core", "SyntaxError", "Syntax Error").ToString());
-		}
-	}
-	else if (FParse::Command(&SourceText, TEXT("CollisionXYOffsetData")))
-	{
-		int32 NumOffsets = FMath::Square(CollisionSizeQuads + 1) * 2;
-
-		CollisionXYOffsetData.Lock(LOCK_READ_WRITE);
-		uint16* Offsets = (uint16*)CollisionXYOffsetData.Realloc(NumOffsets);
-		FMemory::Memzero(Offsets, sizeof(uint16)*NumOffsets);
-
-		FParse::Next(&SourceText);
-		int32 i = 0;
-		while (FChar::IsDigit(*SourceText))
-		{
-			if (i < NumOffsets)
-			{
-				Offsets[i++] = static_cast<uint16>(FCString::Atoi(SourceText));
-				while (FChar::IsDigit(*SourceText))
-				{
-					SourceText++;
-				}
-			}
-
-			FParse::Next(&SourceText);
-		}
-
-		CollisionXYOffsetData.Unlock();
-
-		if (i != NumOffsets)
-		{
-			Warn->Log(*NSLOCTEXT("Core", "SyntaxError", "Syntax Error").ToString());
-		}
-	}
-}
-
 #endif // WITH_EDITOR
 
 ULandscapeInfo* ULandscapeHeightfieldCollisionComponent::GetLandscapeInfo() const
@@ -3349,17 +2586,6 @@ void ULandscapeHeightfieldCollisionComponent::GetResourceSizeEx(FResourceSizeEx&
 	CachedHeightFieldSamples.GetResourceSizeEx(CumulativeResourceSize);
 }
 
-void ULandscapeMeshCollisionComponent::GetResourceSizeEx(FResourceSizeEx& CumulativeResourceSize)
-{
-	Super::GetResourceSizeEx(CumulativeResourceSize);
-
-	if (IsValidRef(MeshRef))
-	{
-		MeshRef->GetResourceSizeEx(CumulativeResourceSize);
-	}
-}
-
-
 TOptional<float> ALandscapeProxy::GetHeightAtLocation(FVector Location, EHeightfieldSource HeightFieldSource) const
 {
 	TOptional<float> Height;
diff --git a/Engine/Source/Runtime/Landscape/Private/LandscapeDataAccess.cpp b/Engine/Source/Runtime/Landscape/Private/LandscapeDataAccess.cpp
--- a/Engine/Source/Runtime/Landscape/Private/LandscapeDataAccess.cpp
+++ b/Engine/Source/Runtime/Landscape/Private/LandscapeDataAccess.cpp
@@ -28,7 +28,6 @@ FLandscapeComponentDataInterface::FLandscapeComponentDataInterface(ULandscapeCom
 	, Component(InComponent)
 	, bWorkOnEditingLayer(bInWorkOnEditingLayer)
 	, HeightMipData(nullptr)
-	, XYOffsetMipData(nullptr)
 {
 	check(IsInParallelGameThread() || IsInGameThread());
 
@@ -36,10 +35,6 @@ FLandscapeComponentDataInterface::FLandscapeComponentDataInterface(ULandscapeCom
 	if (MipLevel < HeightMapTexture->Source.GetNumMips())
 	{
 		HeightMipData = (FColor*)DataInterface.LockMip(HeightMapTexture, MipLevel);
-		if (Component->XYOffsetmapTexture)
-		{
-			XYOffsetMipData = (FColor*)DataInterface.LockMip(Component->XYOffsetmapTexture, MipLevel);
-		}
 	}
 }
 
@@ -51,10 +46,6 @@ FLandscapeComponentDataInterface::~FLandscapeComponentDataInterface()
 	{
 		UTexture2D* HeightMapTexture = Component->GetHeightmap(bWorkOnEditingLayer);
 		DataInterface.UnlockMip(HeightMapTexture, MipLevel);
-		if (Component->XYOffsetmapTexture)
-		{
-			DataInterface.UnlockMip(Component->XYOffsetmapTexture, MipLevel);
-		}
 	}
 }
 
@@ -140,30 +131,16 @@ bool FLandscapeComponentDataInterface::GetWeightmapTextureData(ULandscapeLayerIn
 	return true;
 }
 
+// Deprecated
 FColor* FLandscapeComponentDataInterface::GetXYOffsetData(int32 LocalX, int32 LocalY) const
 {
-#if LANDSCAPE_VALIDATE_DATA_ACCESS
-	checkSlow(IsInParallelGameThread() || IsInGameThread());
-	check(Component);
-	check(LocalX >= 0 && LocalY >= 0 && LocalX < ComponentSizeQuads + 1 && LocalY < ComponentSizeQuads + 1);
-#endif
-
-	const int32 WeightmapSize = ((Component->SubsectionSizeQuads + 1) * Component->NumSubsections) >> MipLevel;
-	int32 SubNumX;
-	int32 SubNumY;
-	int32 SubX;
-	int32 SubY;
-	ComponentXYToSubsectionXY(LocalX, LocalY, SubNumX, SubNumY, SubX, SubY);
-
-	return &XYOffsetMipData[SubX + SubNumX*SubsectionSizeVerts + (SubY + SubNumY*SubsectionSizeVerts)*WeightmapSize];
+	return nullptr;
 }
 
 FVector FLandscapeComponentDataInterface::GetLocalVertex(int32 LocalX, int32 LocalY) const
 {
 	const float ScaleFactor = (float)ComponentSizeQuads / (float)(ComponentSizeVerts - 1);
-	float XOffset, YOffset;
-	GetXYOffset(LocalX, LocalY, XOffset, YOffset);
-	return FVector(LocalX * ScaleFactor + XOffset, LocalY * ScaleFactor + YOffset, LandscapeDataAccess::GetLocalHeight(GetHeight(LocalX, LocalY)));
+	return FVector(LocalX * ScaleFactor, LocalY * ScaleFactor, LandscapeDataAccess::GetLocalHeight(GetHeight(LocalX, LocalY)));
 }
 
 FVector FLandscapeComponentDataInterface::GetWorldVertex(int32 LocalX, int32 LocalY) const
@@ -193,9 +170,7 @@ void FLandscapeComponentDataInterface::GetWorldPositionTangents(int32 LocalX, in
 	float Height = LandscapeDataAccess::UnpackHeight(*Data);
 
 	const float ScaleFactor = (float)ComponentSizeQuads / (float)(ComponentSizeVerts - 1);
-	float XOffset, YOffset;
-	GetXYOffset(LocalX, LocalY, XOffset, YOffset);
-	WorldPos = Component->GetComponentTransform().TransformPosition(FVector(LocalX * ScaleFactor + XOffset, LocalY * ScaleFactor + YOffset, Height));
+	WorldPos = Component->GetComponentTransform().TransformPosition(FVector(LocalX * ScaleFactor, LocalY * ScaleFactor, Height));
 	WorldTangentX = Component->GetComponentTransform().TransformVectorNoScale(WorldTangentX);
 	WorldTangentY = Component->GetComponentTransform().TransformVectorNoScale(WorldTangentY);
 	WorldTangentZ = Component->GetComponentTransform().TransformVectorNoScale(WorldTangentZ);
diff --git a/Engine/Source/Runtime/Landscape/Private/LandscapeEdit.cpp b/Engine/Source/Runtime/Landscape/Private/LandscapeEdit.cpp
--- a/Engine/Source/Runtime/Landscape/Private/LandscapeEdit.cpp
+++ b/Engine/Source/Runtime/Landscape/Private/LandscapeEdit.cpp
@@ -1508,18 +1508,16 @@ void ULandscapeComponent::ClearDirtyCollisionHeightData()
 	LayerDirtyCollisionHeightData = FIntRect();
 }
 
-void ULandscapeComponent::UpdateCollisionHeightData(const FColor* const HeightmapTextureMipData, const FColor* const SimpleCollisionHeightmapTextureData, int32 ComponentX1/*=0*/, int32 ComponentY1/*=0*/, int32 ComponentX2/*=MAX_int32*/, int32 ComponentY2/*=MAX_int32*/, bool bUpdateBounds/*=false*/, const FColor* XYOffsetTextureMipData/*=nullptr*/, bool bInUpdateHeightfieldRegion/*=true*/)
+void ULandscapeComponent::UpdateCollisionHeightData(const FColor* const HeightmapTextureMipData, const FColor* const SimpleCollisionHeightmapTextureData, int32 ComponentX1/*=0*/, int32 ComponentY1/*=0*/, int32 ComponentX2/*=MAX_int32*/, int32 ComponentY2/*=MAX_int32*/, bool bUpdateBounds/*=false*/, bool bInUpdateHeightfieldRegion/*=true*/)
 {
 	TRACE_CPUPROFILER_EVENT_SCOPE(ULandscapeComponent::UpdateCollisionHeightData);
 	
 	ALandscapeProxy* Proxy = GetLandscapeProxy();
 	FIntPoint ComponentKey = GetSectionBase() / ComponentSizeQuads;
 	ULandscapeHeightfieldCollisionComponent* CollisionComp = GetCollisionComponent();
-	ULandscapeMeshCollisionComponent* MeshCollisionComponent = Cast<ULandscapeMeshCollisionComponent>(CollisionComp);
 	ULandscapeHeightfieldCollisionComponent* OldCollisionComponent = CollisionComp;
 
-	// Simple collision is not currently supported with mesh collision components
-	const bool bUsingSimpleCollision = (SimpleCollisionMipLevel > CollisionMipLevel && SimpleCollisionHeightmapTextureData && !XYOffsetmapTexture);
+	const bool bUsingSimpleCollision = (SimpleCollisionMipLevel > CollisionMipLevel && SimpleCollisionHeightmapTextureData);
 	
 	FCollisionSize CollisionSize = FCollisionSize::Create(NumSubsections, SubsectionSizeQuads, CollisionMipLevel);
 	FCollisionSize SimpleCollisionSize = FCollisionSize::CreateSimple(bUsingSimpleCollision, NumSubsections, SubsectionSizeQuads, SimpleCollisionMipLevel);
@@ -1527,7 +1525,6 @@ void ULandscapeComponent::UpdateCollisionHeightData(const FColor* const Heightma
 	const int32 TotalCollisionSize = CollisionSize.SizeVertsSquare + SimpleCollisionSize.SizeVertsSquare;
 
 	uint16* CollisionHeightData = nullptr;
-	uint16* CollisionXYOffsetData = nullptr;
 	bool CreatedNew = false;
 	bool ChangeType = false;
 
@@ -1541,8 +1538,7 @@ void ULandscapeComponent::UpdateCollisionHeightData(const FColor* const Heightma
 	}
 	ClearDirtyCollisionHeightData();
 
-	// Existing collision component is same type with collision
-	if (CollisionComp && ((XYOffsetmapTexture == nullptr) == (MeshCollisionComponent == nullptr)))
+	if (CollisionComp)
 	{
 		ComponentX1 = FMath::Clamp(ComponentX1, 0, ComponentSizeQuads);
 		ComponentY1 = FMath::Clamp(ComponentY1, 0, ComponentSizeQuads);
@@ -1572,22 +1568,15 @@ void ULandscapeComponent::UpdateCollisionHeightData(const FColor* const Heightma
 
 		RecreateCollisionComponent(bUsingSimpleCollision);
 		CollisionComp = GetCollisionComponent();
-        MeshCollisionComponent = Cast<ULandscapeMeshCollisionComponent>(CollisionComp);
 	}
 
 	CollisionHeightData = (uint16*)CollisionComp->CollisionHeightData.Lock(LOCK_READ_WRITE);
-	
-	if (XYOffsetmapTexture && MeshCollisionComponent)
-	{
-		CollisionXYOffsetData = (uint16*)MeshCollisionComponent->CollisionXYOffsetData.Lock(LOCK_READ_WRITE);
-	}
 
 	const int32 HeightmapSizeU = GetHeightmap()->Source.GetSizeX();
 	const int32 HeightmapSizeV = GetHeightmap()->Source.GetSizeY();
 
 	// Handle Material WPO baked into heightfield collision
-	// Material WPO is not currently supported for mesh collision components
-	const bool bUsingGrassMapHeights = Proxy->bBakeMaterialPositionOffsetIntoCollision && !MeshCollisionComponent && GrassData->HasData() && !IsGrassMapOutdated();
+	const bool bUsingGrassMapHeights = Proxy->bBakeMaterialPositionOffsetIntoCollision && GrassData->HasData() && !IsGrassMapOutdated();
 	uint16* GrassHeights = nullptr;
 	if (bUsingGrassMapHeights)
 	{
@@ -1604,31 +1593,21 @@ void ULandscapeComponent::UpdateCollisionHeightData(const FColor* const Heightma
 		}
 	}
 
-	UpdateCollisionHeightBuffer(ComponentX1, ComponentY1, ComponentX2, ComponentY2, CollisionMipLevel, HeightmapSizeU, HeightmapSizeV, HeightmapTextureMipData, CollisionHeightData, GrassHeights, XYOffsetTextureMipData, CollisionXYOffsetData);
+	UpdateCollisionHeightBuffer(ComponentX1, ComponentY1, ComponentX2, ComponentY2, CollisionMipLevel, HeightmapSizeU, HeightmapSizeV, HeightmapTextureMipData, CollisionHeightData, GrassHeights);
 		
 	if (bUsingSimpleCollision)
 	{
 		uint16* SimpleCollisionGrassHeights = bUsingGrassMapHeights && GrassData->HeightMipData.Contains(SimpleCollisionMipLevel) ? GrassData->HeightMipData[SimpleCollisionMipLevel].GetData() : nullptr;
 		uint16* const SimpleCollisionHeightData = CollisionHeightData + CollisionSize.SizeVertsSquare;
-		UpdateCollisionHeightBuffer(ComponentX1, ComponentY1, ComponentX2, ComponentY2, SimpleCollisionMipLevel, HeightmapSizeU, HeightmapSizeV, SimpleCollisionHeightmapTextureData, SimpleCollisionHeightData, SimpleCollisionGrassHeights, nullptr, nullptr);
+		UpdateCollisionHeightBuffer(ComponentX1, ComponentY1, ComponentX2, ComponentY2, SimpleCollisionMipLevel, HeightmapSizeU, HeightmapSizeV, SimpleCollisionHeightmapTextureData, SimpleCollisionHeightData, SimpleCollisionGrassHeights);
 	}
 
 	CollisionComp->CollisionHeightData.Unlock();
 
-	if (XYOffsetmapTexture && MeshCollisionComponent)
-	{
-		MeshCollisionComponent->CollisionXYOffsetData.Unlock();
-	}
-
 	// If we updated an existing component, we need to update the phys x heightfield edit data
 	if (!CreatedNew && bInUpdateHeightfieldRegion)
 	{
-		if (MeshCollisionComponent)
-		{
-			// Will be done once for XY Offset data update in FXYOffsetmapAccessor() destructor with UpdateCachedBounds()
-			//MeshCollisionComponent->RecreateCollision();
-		}
-		else if (CollisionMipLevel == 0)
+		if (CollisionMipLevel == 0)
 		{
 			CollisionComp->UpdateHeightfieldRegion(ComponentX1, ComponentY1, ComponentX2, ComponentY2);
 		}
@@ -1698,22 +1677,16 @@ void ULandscapeComponent::UpdateCollisionData(bool bInUpdateHeightfieldRegion)
 	{
 		verify( GetHeightmap()->Source.GetMipData(SimpleCollisionMipData, SimpleCollisionMipLevel) );
 	}
-	if (XYOffsetmapTexture)
-	{
-		verify( XYOffsetmapTexture->Source.GetMipData(XYOffsetMipData, CollisionMipLevel) );
-	}
 
 	UpdateCollisionHeightData(
 		(FColor*)CollisionMipData.GetData(),
 		SimpleCollisionMipLevel > CollisionMipLevel ? (FColor*)SimpleCollisionMipData.GetData() : nullptr,
-		0, 0, MAX_int32, MAX_int32, true,
-		XYOffsetmapTexture ? (FColor*)XYOffsetMipData.GetData() : nullptr, bInUpdateHeightfieldRegion);
+		0, 0, MAX_int32, MAX_int32, true, bInUpdateHeightfieldRegion);
 }
 
 void ULandscapeComponent::RecreateCollisionComponent(bool bUseSimpleCollision)
 {
 	ULandscapeHeightfieldCollisionComponent* CollisionComp = GetCollisionComponent();
-	ULandscapeMeshCollisionComponent* MeshCollisionComponent = nullptr;
 	TArray<uint8> DominantLayerData;
 	TArray<ULandscapeLayerInfoObject*> LayerInfos;
 	ALandscapeProxy* Proxy = GetLandscapeProxy();
@@ -1743,16 +1716,7 @@ void ULandscapeComponent::RecreateCollisionComponent(bool bUseSimpleCollision)
 		CollisionComp = nullptr;
 	}
 
-	if (XYOffsetmapTexture)
-	{
-		MeshCollisionComponent = NewObject<ULandscapeMeshCollisionComponent>(Proxy, NAME_None, RF_Transactional);
-		CollisionComp = MeshCollisionComponent;
-	}
-	else
-	{
-		MeshCollisionComponent = nullptr;
-		CollisionComp = NewObject<ULandscapeHeightfieldCollisionComponent>(Proxy, NAME_None, RF_Transactional);
-	}
+	CollisionComp = NewObject<ULandscapeHeightfieldCollisionComponent>(Proxy, NAME_None, RF_Transactional);
 
 	CollisionComp->SetRelativeLocation(GetRelativeLocation());
 	CollisionComp->SetupAttachment(Proxy->GetRootComponent(), NAME_None);
@@ -1771,15 +1735,6 @@ void ULandscapeComponent::RecreateCollisionComponent(bool bUseSimpleCollision)
 	uint16* CollisionHeightData = (uint16*)CollisionComp->CollisionHeightData.Realloc(TotalCollisionSize);
 	FMemory::Memzero(CollisionHeightData, TotalCollisionSize * CollisionComp->CollisionHeightData.GetElementSize());
 	CollisionComp->CollisionHeightData.Unlock();
-	
-	if (XYOffsetmapTexture && MeshCollisionComponent)
-	{
-		// Need XYOffsetData for Collision Component
-		MeshCollisionComponent->CollisionXYOffsetData.Lock(LOCK_READ_WRITE);
-		uint16* CollisionXYOffsetData = (uint16*)MeshCollisionComponent->CollisionXYOffsetData.Realloc(TotalCollisionSize * 2);
-		FMemory::Memzero(CollisionXYOffsetData, TotalCollisionSize * 2 * MeshCollisionComponent->CollisionXYOffsetData.GetElementSize());
-		MeshCollisionComponent->CollisionXYOffsetData.Unlock();
-	}
 
 	if (DominantLayerData.Num())
 	{
@@ -1797,8 +1752,7 @@ void ULandscapeComponent::RecreateCollisionComponent(bool bUseSimpleCollision)
 }
 
 void ULandscapeComponent::UpdateCollisionHeightBuffer(	int32 InComponentX1, int32 InComponentY1, int32 InComponentX2, int32 InComponentY2, int32 InCollisionMipLevel, int32 InHeightmapSizeU, int32 InHeightmapSizeV,
-														const FColor* const InHeightmapTextureMipData, uint16* OutCollisionHeightData, uint16* InGrassHeightData,
-														const FColor* const InXYOffsetTextureMipData, uint16* OutCollisionXYOffsetData)
+														const FColor* const InHeightmapTextureMipData, uint16* OutCollisionHeightData, uint16* InGrassHeightData)
 {
 	FCollisionSize CollisionSize = FCollisionSize::Create(NumSubsections, SubsectionSizeQuads, InCollisionMipLevel);
 
@@ -1816,8 +1770,6 @@ void ULandscapeComponent::UpdateCollisionHeightBuffer(	int32 InComponentX1, int3
 	const int32 HeightmapOffsetX = FMath::RoundToInt32(HeightmapScaleBias.Z * (float)InHeightmapSizeU) >> InCollisionMipLevel;
 	const int32 HeightmapOffsetY = FMath::RoundToInt32(HeightmapScaleBias.W * (float)InHeightmapSizeV) >> InCollisionMipLevel;
 
-	const int32 XYMipSizeU = XYOffsetmapTexture ? XYOffsetmapTexture->Source.GetSizeX() >> InCollisionMipLevel : 0;
-
 	for (int32 SubsectionY = SubSectionY1; SubsectionY < SubSectionY2; ++SubsectionY)
 	{
 		for (int32 SubsectionX = SubSectionX1; SubsectionX < SubSectionX2; ++SubsectionX)
@@ -1867,28 +1819,13 @@ void ULandscapeComponent::UpdateCollisionHeightBuffer(	int32 InComponentX1, int3
 
 						CollisionHeight = NewHeight;
 					}
-
-					if (XYOffsetmapTexture && InXYOffsetTextureMipData && OutCollisionXYOffsetData)
-					{
-						const int32 TexX = CollisionSize.SubsectionSizeVerts * SubsectionX + VertX;
-						const int32 TexY = CollisionSize.SubsectionSizeVerts * SubsectionY + VertY;
-						const FColor& TexData = InXYOffsetTextureMipData[TexX + TexY * XYMipSizeU];
-
-						// Copy collision data
-						const uint16 NewXOffset = TexData.R << 8 | TexData.G;
-						const uint16 NewYOffset = TexData.B << 8 | TexData.A;
-
-						const int32 XYIndex = CompVertX + CompVertY * CollisionSize.SizeVerts;
-						OutCollisionXYOffsetData[XYIndex * 2] = NewXOffset;
-						OutCollisionXYOffsetData[XYIndex * 2 + 1] = NewYOffset;
-					}
 				}
 			}
 		}
 	}
 }
 
-void ULandscapeComponent::UpdateDominantLayerBuffer(int32 InComponentX1, int32 InComponentY1, int32 InComponentX2, int32 InComponentY2, int32 InCollisionMipLevel, int32 InWeightmapSizeU, int32 InDataLayerIdx, const TArray<uint8*>& InCollisionDataPtrs, const TArray<ULandscapeLayerInfoObject*>& InLayerInfos, uint8* OutDominantLayerData)
+void ULandscapeComponent::UpdateDominantLayerBuffer(int32 InComponentX1, int32 InComponentY1, int32 InComponentX2, int32 InComponentY2, int32 InCollisionMipLevel, int32 InWeightmapSizeU, int32 InVisibilityLayerIdx, const TArray<uint8*>& InCollisionDataPtrs, const TArray<ULandscapeLayerInfoObject*>& InLayerInfos, uint8* OutDominantLayerData)
 {
 	const int32 MipSizeU = InWeightmapSizeU >> InCollisionMipLevel;
 
@@ -1940,7 +1877,7 @@ void ULandscapeComponent::UpdateDominantLayerBuffer(int32 InComponentX1, int32 I
 						const uint8 LayerWeight = InCollisionDataPtrs[LayerIdx][DataOffset];
 						const uint8 LayerMinimumWeight = InLayerInfos[LayerIdx] ? (uint8)(InLayerInfos[LayerIdx]->MinimumCollisionRelevanceWeight * 255) :  0;
 
-						if (LayerIdx == InDataLayerIdx) // Override value for hole
+						if (LayerIdx == InVisibilityLayerIdx) // Override value for visibility layer holes
 						{
 							if (LayerWeight > 170) // 255 * 0.66...
 							{
@@ -1976,15 +1913,14 @@ void ULandscapeComponent::UpdateCollisionLayerData(const FColor* const* const We
 
 	if (CollisionComp)
 	{
-		// Simple collision is not currently supported with mesh collision components
-		const bool bUsingSimpleCollision = (SimpleCollisionMipLevel > CollisionMipLevel && SimpleCollisionWeightmapTextureMipData && !XYOffsetmapTexture);
+		const bool bUsingSimpleCollision = (SimpleCollisionMipLevel > CollisionMipLevel && SimpleCollisionWeightmapTextureMipData);
 
 		TArray<ULandscapeLayerInfoObject*> CandidateLayers;
 		TArray<uint8*> CandidateDataPtrs;
 		TArray<uint8*> SimpleCollisionDataPtrs;
 
 		bool bExistingLayerMismatch = false;
-		int32 DataLayerIdx = INDEX_NONE;
+		int32 VisibilityLayerIdx = INDEX_NONE;
 
 		const TArray<FWeightmapLayerAllocationInfo>& ComponentWeightmapLayerAllocations = GetWeightmapLayerAllocations(false);
 		const TArray<UTexture2D*>& ComponentWeightmapsTexture = GetWeightmapTextures(false);
@@ -1994,7 +1930,7 @@ void ULandscapeComponent::UpdateCollisionLayerData(const FColor* const* const We
 		{
 			const FWeightmapLayerAllocationInfo& AllocInfo = ComponentWeightmapLayerAllocations[AllocIdx];
 			ULandscapeLayerInfoObject* LayerInfo = AllocInfo.LayerInfo;
-			if (LayerInfo == ALandscapeProxy::VisibilityLayer || AllocInfo.IsAllocated())
+			if (LayerInfo != nullptr && AllocInfo.IsAllocated())
 			{
 				int32 Idx = CandidateLayers.Add(LayerInfo);
 				CandidateDataPtrs.Add(((uint8*)WeightmapTextureMipData[AllocInfo.WeightmapTextureIndex]) + ChannelOffsets[AllocInfo.WeightmapTextureChannel]);
@@ -2012,7 +1948,7 @@ void ULandscapeComponent::UpdateCollisionLayerData(const FColor* const* const We
 
 				if (LayerInfo == ALandscapeProxy::VisibilityLayer)
 				{
-					DataLayerIdx = Idx;
+					VisibilityLayerIdx = Idx;
 					bExistingLayerMismatch = true; // always rebuild whole component for hole
 				}
 			}
@@ -2057,12 +1993,12 @@ void ULandscapeComponent::UpdateCollisionLayerData(const FColor* const* const We
 			const int32 WeightmapSizeU = ComponentWeightmapsTexture[0]->Source.GetSizeX();
 						
 			// gmartin: WeightmapScaleBias not handled?			
-			UpdateDominantLayerBuffer(ComponentX1, ComponentY1, ComponentX2, ComponentY2, CollisionMipLevel, WeightmapSizeU, DataLayerIdx, CandidateDataPtrs, CollisionComp->ComponentLayerInfos, DominantLayerData);
+			UpdateDominantLayerBuffer(ComponentX1, ComponentY1, ComponentX2, ComponentY2, CollisionMipLevel, WeightmapSizeU, VisibilityLayerIdx, CandidateDataPtrs, CollisionComp->ComponentLayerInfos, DominantLayerData);
 
 			if (bUsingSimpleCollision)
 			{
 				uint8* const SimpleCollisionHeightData = DominantLayerData + CollisionSize.SizeVertsSquare;
-				UpdateDominantLayerBuffer(ComponentX1, ComponentY1, ComponentX2, ComponentY2, SimpleCollisionMipLevel, WeightmapSizeU, DataLayerIdx, SimpleCollisionDataPtrs, CollisionComp->ComponentLayerInfos, SimpleCollisionHeightData);
+				UpdateDominantLayerBuffer(ComponentX1, ComponentY1, ComponentX2, ComponentY2, SimpleCollisionMipLevel, WeightmapSizeU, VisibilityLayerIdx, SimpleCollisionDataPtrs, CollisionComp->ComponentLayerInfos, SimpleCollisionHeightData);
 			}
 
 			CollisionComp->DominantLayerData.Unlock();
@@ -5450,18 +5386,20 @@ void ALandscapeProxy::RemoveXYOffsets()
 	for (int32 i = 0; i < LandscapeComponents.Num(); ++i)
 	{
 		ULandscapeComponent* Comp = LandscapeComponents[i];
-		if (Comp && Comp->XYOffsetmapTexture)
+PRAGMA_DISABLE_DEPRECATION_WARNINGS
+		if (Comp && Comp->XYOffsetmapTexture_DEPRECATED)
 		{
-			Comp->XYOffsetmapTexture->SetFlags(RF_Transactional);
-			Comp->XYOffsetmapTexture->Modify();
-			Comp->XYOffsetmapTexture->MarkPackageDirty();
-			Comp->XYOffsetmapTexture->ClearFlags(RF_Standalone);
+			Comp->XYOffsetmapTexture_DEPRECATED->SetFlags(RF_Transactional);
+			Comp->XYOffsetmapTexture_DEPRECATED->Modify();
+			Comp->XYOffsetmapTexture_DEPRECATED->MarkPackageDirty();
+			Comp->XYOffsetmapTexture_DEPRECATED->ClearFlags(RF_Standalone);
 			Comp->Modify();
 			Comp->MarkPackageDirty();
-			Comp->XYOffsetmapTexture = nullptr;
+			Comp->XYOffsetmapTexture_DEPRECATED = nullptr;
 			Comp->MarkRenderStateDirty();
 			bFoundXYOffset = true;
 		}
+PRAGMA_ENABLE_DEPRECATION_WARNINGS
 	}
 
 	if (bFoundXYOffset)
@@ -5744,16 +5682,6 @@ ALandscapeProxy* ULandscapeInfo::MoveComponentsToProxy(const TArray<ULandscapeCo
 	LandscapeProxy->Modify();
 	LandscapeProxy->MarkPackageDirty();
 
-	// Handle XY-offset textures (these don't need splitting, as they aren't currently shared between components like heightmaps/weightmaps can be)
-	for (ULandscapeComponent* Component : TargetSelectedComponents)
-	{
-		if (Component->XYOffsetmapTexture)
-		{
-			Component->XYOffsetmapTexture->Modify();
-			Component->XYOffsetmapTexture->Rename(nullptr, LandscapeProxy);
-		}
-	}
-
 	// Change Weight maps...
 	{
 		FLandscapeEditDataInterface LandscapeEdit(this);
diff --git a/Engine/Source/Runtime/Landscape/Private/LandscapeEditInterface.cpp b/Engine/Source/Runtime/Landscape/Private/LandscapeEditInterface.cpp
--- a/Engine/Source/Runtime/Landscape/Private/LandscapeEditInterface.cpp
+++ b/Engine/Source/Runtime/Landscape/Private/LandscapeEditInterface.cpp
@@ -288,19 +288,14 @@ void FLandscapeEditDataInterface::SetHeightData(int32 X1, int32 Y1, int32 X2, in
 		VertexNormals = new FVector[NumVertsX*NumVertsY];
 		FMemory::Memzero(VertexNormals, NumVertsX*NumVertsY * sizeof(FVector));
 
-		// Need to consider XYOffset for XY displacemented map
-		FVector2D* XYOffsets = new FVector2D[NumVertsX*NumVertsY];
-		FMemory::Memzero(XYOffsets, NumVertsX*NumVertsY * sizeof(FVector2D));
-		GetXYOffsetDataFast(X1, Y1, X2, Y2, XYOffsets, 0);
-
 		for (int32 Y = 0; Y < NumVertsY - 1; Y++)
 		{
 			for (int32 X = 0; X < NumVertsX - 1; X++)
 			{
-				FVector Vert00 = FVector(XYOffsets[(X + 0) + NumVertsX*(Y + 0)].X + 0.0f, XYOffsets[(X + 0) + NumVertsX*(Y + 0)].Y + 0.0f, LandscapeDataAccess::GetLocalHeight(InData[(X + 0) + InStride*(Y + 0)])) * DrawScale;
-				FVector Vert01 = FVector(XYOffsets[(X + 0) + NumVertsX*(Y + 0)].X + 0.0f, XYOffsets[(X + 0) + NumVertsX*(Y + 0)].Y + 1.0f, LandscapeDataAccess::GetLocalHeight(InData[(X + 0) + InStride*(Y + 1)])) * DrawScale;
-				FVector Vert10 = FVector(XYOffsets[(X + 0) + NumVertsX*(Y + 0)].X + 1.0f, XYOffsets[(X + 0) + NumVertsX*(Y + 0)].Y + 0.0f, LandscapeDataAccess::GetLocalHeight(InData[(X + 1) + InStride*(Y + 0)])) * DrawScale;
-				FVector Vert11 = FVector(XYOffsets[(X + 0) + NumVertsX*(Y + 0)].X + 1.0f, XYOffsets[(X + 0) + NumVertsX*(Y + 0)].Y + 1.0f, LandscapeDataAccess::GetLocalHeight(InData[(X + 1) + InStride*(Y + 1)])) * DrawScale;
+				FVector Vert00 = FVector(0.0f, 0.0f, LandscapeDataAccess::GetLocalHeight(InData[(X + 0) + InStride*(Y + 0)])) * DrawScale;
+				FVector Vert01 = FVector(0.0f, 1.0f, LandscapeDataAccess::GetLocalHeight(InData[(X + 0) + InStride*(Y + 1)])) * DrawScale;
+				FVector Vert10 = FVector(1.0f, 0.0f, LandscapeDataAccess::GetLocalHeight(InData[(X + 1) + InStride*(Y + 0)])) * DrawScale;
+				FVector Vert11 = FVector(1.0f, 1.0f, LandscapeDataAccess::GetLocalHeight(InData[(X + 1) + InStride*(Y + 1)])) * DrawScale;
 
 				FVector FaceNormal1 = ((Vert00 - Vert10) ^ (Vert10 - Vert11)).GetSafeNormal();
 				FVector FaceNormal2 = ((Vert11 - Vert01) ^ (Vert01 - Vert00)).GetSafeNormal();
@@ -312,8 +307,6 @@ void FLandscapeEditDataInterface::SetHeightData(int32 X1, int32 Y1, int32 X2, in
 				VertexNormals[(X + 1 + NumVertsX*(Y + 1))] += FaceNormal1 + FaceNormal2;
 			}
 		}
-
-		delete[] XYOffsets;
 	}
 
 	for (int32 ComponentIndexY = ComponentIndexY1; ComponentIndexY <= ComponentIndexY2; ComponentIndexY++)
@@ -338,20 +331,12 @@ void FLandscapeEditDataInterface::SetHeightData(int32 X1, int32 Y1, int32 X2, in
 			}
 
 			UTexture2D* Heightmap = InHeightmap != nullptr ? InHeightmap : Component->GetHeightmap(GetEditLayer());
-			UTexture2D* XYOffsetmapTexture = InXYOffsetmapTexture != nullptr ? InXYOffsetmapTexture : ToRawPtr(Component->XYOffsetmapTexture);
 
 			Component->Modify(GetShouldDirtyPackage());
 
 			FLandscapeTextureDataInfo* TexDataInfo = GetTextureDataInfo(Heightmap);
 			FColor* HeightmapTextureData = (FColor*)TexDataInfo->GetMipData(0);
 
-			FColor* XYOffsetMipData = nullptr;
-			if (XYOffsetmapTexture)
-			{
-				FLandscapeTextureDataInfo* XYTexDataInfo = GetTextureDataInfo(XYOffsetmapTexture);
-				XYOffsetMipData = (FColor*)XYTexDataInfo->GetMipData(Component->CollisionMipLevel);
-			}
-
 			// Find the texture data corresponding to this vertex
 			int32 SizeU = Heightmap->Source.GetSizeX();
 			int32 SizeV = Heightmap->Source.GetSizeY();
@@ -463,133 +448,9 @@ void FLandscapeEditDataInterface::SetHeightData(int32 X1, int32 Y1, int32 X2, in
 	}
 }
 
-//
-// RecalculateNormals - Regenerate normals for the entire landscape. Called after modifying DrawScale3D.
-//
+// Deprecated
 void FLandscapeEditDataInterface::RecalculateNormals()
 {
-	if (!LandscapeInfo) return;
-
-	// we should only ever be calculating normals on the final runtime layer.
-	// if this gets hit, then we are wasting time calculating them on edit layers.
-	ensure(!GetEditLayer().IsValid());
-
-	// Recalculate normals for each component in turn
-	for( auto It = LandscapeInfo->XYtoComponentMap.CreateIterator(); It; ++It )
-	{
-		ULandscapeComponent* Component = It.Value();
-
-		// one extra row of vertex either side of the component
-		int32 X1 = Component->GetSectionBase().X-1;
-		int32 Y1 = Component->GetSectionBase().Y-1;
-		int32 X2 = Component->GetSectionBase().X+ComponentSizeQuads+1;
-		int32 Y2 = Component->GetSectionBase().Y+ComponentSizeQuads+1;
-		const int32 OriginalX1 = X1;
-		const int32 OriginalY1 = Y1;
-		const int32 OriginalX2 = X2;
-		const int32 OriginalY2 = Y2;
-		int32 Stride = ComponentSizeQuads+3; 
-		int32 StrideSquared = FMath::Square(Stride);
-
-		uint16* HeightData = new uint16[StrideSquared];
-		FVector* VertexNormals = new FVector[StrideSquared];
-		FMemory::Memzero(VertexNormals, StrideSquared * sizeof(FVector));
-		FVector2D* XYOffsets = new FVector2D[StrideSquared];
-		FMemory::Memzero(XYOffsets, StrideSquared * sizeof(FVector2D));
-
-		// Get XY offset
-		GetXYOffsetDataFast(X1,Y1,X2,Y2,XYOffsets,0);
-		check(X1 == OriginalX1 && Y1 == OriginalY1 && X2 == OriginalX2 && Y2 == OriginalY2);	// because we're querying within an existing component, we assume the data should always exist
-
-		// Get the vertex positions for entire quad
-		GetHeightData(X1,Y1,X2,Y2,HeightData,0);
-		check(X1 == OriginalX1 && Y1 == OriginalY1 && X2 == OriginalX2 && Y2 == OriginalY2);	// because we're querying within an existing component, we assume the data should always exist
-
-		// Contribute face normals for all triangles contributing to this components' normals
-		for( int32 Y=0;Y<Stride-1;Y++ )
-		{
-			for( int32 X=0;X<Stride-1;X++ )
-			{
-				FVector Vert00 = FVector(XYOffsets[(X + 0) + Stride * (Y + 0)].X + 0.0f, XYOffsets[(X + 0) + Stride * (Y + 0)].Y + 0.0f, LandscapeDataAccess::GetLocalHeight(HeightData[(X + 0) + Stride * (Y + 0)])) * DrawScale;
-				FVector Vert01 = FVector(XYOffsets[(X + 0) + Stride * (Y + 0)].X + 0.0f, XYOffsets[(X + 0) + Stride * (Y + 0)].Y + 1.0f, LandscapeDataAccess::GetLocalHeight(HeightData[(X + 0) + Stride * (Y + 1)])) * DrawScale;
-				FVector Vert10 = FVector(XYOffsets[(X + 0) + Stride * (Y + 0)].X + 1.0f, XYOffsets[(X + 0) + Stride * (Y + 0)].Y + 0.0f, LandscapeDataAccess::GetLocalHeight(HeightData[(X + 1) + Stride * (Y + 0)])) * DrawScale;
-				FVector Vert11 = FVector(XYOffsets[(X + 0) + Stride * (Y + 0)].X + 1.0f, XYOffsets[(X + 0) + Stride * (Y + 0)].Y + 1.0f, LandscapeDataAccess::GetLocalHeight(HeightData[(X + 1) + Stride * (Y + 1)])) * DrawScale;
-
-				FVector FaceNormal1 = ((Vert00-Vert10) ^ (Vert10-Vert11)).GetSafeNormal();
-				FVector FaceNormal2 = ((Vert11-Vert01) ^ (Vert01-Vert00)).GetSafeNormal(); 
-
-				// contribute to the vertex normals.
-				VertexNormals[(X+1 + Stride*(Y+0))] += FaceNormal1;
-				VertexNormals[(X+0 + Stride*(Y+1))] += FaceNormal2;
-				VertexNormals[(X+0 + Stride*(Y+0))] += FaceNormal1 + FaceNormal2;
-				VertexNormals[(X+1 + Stride*(Y+1))] += FaceNormal1 + FaceNormal2;
-			}
-		}
-
-		// Find the texture data corresponding to this vertex
-		UTexture2D* EditHeightmap = Component->GetHeightmap(GetEditLayer());
-		int32 SizeU = EditHeightmap->Source.GetSizeX();
-		int32 SizeV = EditHeightmap->Source.GetSizeY();
-		int32 HeightmapOffsetX = static_cast<int32>(Component->HeightmapScaleBias.Z * SizeU);
-		int32 HeightmapOffsetY = static_cast<int32>(Component->HeightmapScaleBias.W * SizeV);
-
-		FLandscapeTextureDataInfo* TexDataInfo = GetTextureDataInfo(EditHeightmap);
-		FColor* HeightmapTextureData = (FColor*)TexDataInfo->GetMipData(0);
-
-		// Apply vertex normals to the component
-		for( int32 SubIndexY=0;SubIndexY<Component->NumSubsections;SubIndexY++ )
-		{
-			for( int32 SubIndexX=0;SubIndexX<Component->NumSubsections;SubIndexX++ )
-			{
-				for( int32 SubY=0;SubY<=SubsectionSizeQuads;SubY++ )
-				{
-					for( int32 SubX=0;SubX<=SubsectionSizeQuads;SubX++ )
-					{
-						// Read from local VertexNormals cache which doesn't have duplicated values.
-						const int32 ReadX = (SubsectionSizeQuads) * SubIndexX + SubX;
-						const int32 ReadY = (SubsectionSizeQuads) * SubIndexY + SubY;
-						const int32 ReadIndex = (ReadX + 1) + (ReadY + 1) * Stride;
-						checkSlow(ReadIndex < StrideSquared);
-						const FVector Normal = VertexNormals[ReadIndex].GetSafeNormal();
-
-						// Write to (shared) Heightmap texture which has duplicated values on subsection and section borders.
-						const int32 WriteX = HeightmapOffsetX + (SubsectionSizeQuads + 1) * SubIndexX + SubX;
-						const int32 WriteY = HeightmapOffsetY + (SubsectionSizeQuads + 1) * SubIndexY + SubY;
-						const int32 WriteIndex = WriteX + WriteY * SizeU;
-						FColor& TexData = HeightmapTextureData[WriteIndex];
-
-						// Update the channels containing the normals
-						TexData.B = static_cast<uint8>(FMath::RoundToInt32( 127.5f * (Normal.X + 1.0f) ));
-						TexData.A = static_cast<uint8>(FMath::RoundToInt32( 127.5f * (Normal.Y + 1.0f) ));
-					}
-				}
-			}
-		}
-
-		delete[] XYOffsets;
-		delete[] HeightData;
-		delete[] VertexNormals;
-
-		// Record the areas of the texture we need to re-upload
-		int32 TexX1 = HeightmapOffsetX;
-		int32 TexY1 = HeightmapOffsetY;
-		int32 TexX2 = HeightmapOffsetX + (SubsectionSizeQuads+1) * Component->NumSubsections - 1;
-		int32 TexY2 = HeightmapOffsetY + (SubsectionSizeQuads+1) * Component->NumSubsections - 1;
-		TexDataInfo->AddMipUpdateRegion(0,TexX1,TexY1,TexX2,TexY2);
-
-		// Work out how many mips should be calculated directly from one component's data.
-		// The remaining mips are calculated on a per texture basis.
-		// eg if subsection is 7x7 quads, we need one 3 mips total: (8x8, 4x4, 2x2 verts)
-		int32 BaseNumMips = FMath::CeilLogTwo(SubsectionSizeQuads+1);
-		TArray<FColor*> MipData;
-		MipData.AddUninitialized(BaseNumMips);
-		MipData[0] = HeightmapTextureData;
-		for( int32 MipIdx=1;MipIdx<BaseNumMips;MipIdx++ )
-		{
-			MipData[MipIdx] = (FColor*)TexDataInfo->GetMipData(MipIdx);
-		}
-		Component->GenerateHeightmapMips( MipData, 0, 0, ComponentSizeQuads, ComponentSizeQuads, TexDataInfo );
-	}
 }
 
 template<typename TStoreData>
@@ -4599,839 +4460,76 @@ void FLandscapeEditDataInterface::SetEditToolTextureData(int32 X1, int32 Y1, int
 	}
 }
 
-template<typename T>
-void FLandscapeEditDataInterface::SetXYOffsetDataTempl(int32 X1, int32 Y1, int32 X2, int32 Y2, const T* Data, int32 Stride)
-{
-	if( Stride==0 )
-	{
-		Stride = (1+X2-X1);
-	}
-
-	check(ComponentSizeQuads > 0);
-	// Find component range for this block of data
-	int32 ComponentIndexX1, ComponentIndexY1, ComponentIndexX2, ComponentIndexY2;
-	ALandscape::CalcComponentIndicesNoOverlap(X1, Y1, X2, Y2, ComponentSizeQuads, ComponentIndexX1, ComponentIndexY1, ComponentIndexX2, ComponentIndexY2);
-
-	const FColor DefaultValue(128, 0, 128, 0);
-
-	for( int32 ComponentIndexY=ComponentIndexY1;ComponentIndexY<=ComponentIndexY2;ComponentIndexY++ )
-	{
-		for( int32 ComponentIndexX=ComponentIndexX1;ComponentIndexX<=ComponentIndexX2;ComponentIndexX++ )
-		{	
-			FIntPoint ComponentKey(ComponentIndexX,ComponentIndexY);
-			ULandscapeComponent* Component = LandscapeInfo->XYtoComponentMap.FindRef(ComponentKey);
-
-			UTexture2D* XYOffsetTexture = NULL;
-			if( Component==NULL)
-			{
-				continue;
-			}
-			else
-			{
-				if (Component->XYOffsetmapTexture == NULL)
-				{
-					Component->Modify(GetShouldDirtyPackage());
-					//FlushRenderingCommands();
-					// Construct Texture...
-					int32 WeightmapSize = (Component->SubsectionSizeQuads+1) * Component->NumSubsections;
-					XYOffsetTexture = Component->GetLandscapeProxy()->CreateLandscapeTexture(WeightmapSize, WeightmapSize, TEXTUREGROUP_Terrain_Weightmap, TSF_BGRA8);
-					// Alloc dummy mips
-					const bool bClear = true;
-					ULandscapeComponent::CreateEmptyTextureMips(XYOffsetTexture, ELandscapeTextureUsage::FinalData, ELandscapeTextureType::Unknown, bClear);
-
-					// TODO: we flag somewhere that we need to use XYOffset Factory, so we can recompute shaders based on this change
-
-					XYOffsetTexture->PostEditChange();
-
-					//FlushRenderingCommands();
-					SetTextureValue(XYOffsetTexture, DefaultValue);
-					FLandscapeTextureDataInfo* TexDataInfo = GetTextureDataInfo(XYOffsetTexture);
-					int32 NumMips = XYOffsetTexture->Source.GetNumMips();
-					TArray<FColor*> TextureMipData;
-					TextureMipData.AddUninitialized(NumMips);
-					for( int32 MipIdx=0;MipIdx<NumMips;MipIdx++ )
-					{
-						TextureMipData[MipIdx] = (FColor*)TexDataInfo->GetMipData(MipIdx);
-					}
-					ULandscapeComponent::UpdateWeightmapMips(ComponentNumSubsections, SubsectionSizeQuads, XYOffsetTexture, TextureMipData, 0, 0, MAX_int32, MAX_int32, TexDataInfo);
-
-					Component->XYOffsetmapTexture = XYOffsetTexture;
-					FComponentReregisterContext ReregisterContext(Component);
-				}
-				else
-				{
-					XYOffsetTexture = Component->XYOffsetmapTexture;
-				}
-			}
-
-			FLandscapeTextureDataInfo* TexDataInfo = GetTextureDataInfo(XYOffsetTexture);
-			FColor* XYOffsetTextureData = (FColor*)TexDataInfo->GetMipData(0);
-
-			// Find the texture data corresponding to this vertex
-			int32 SizeU = XYOffsetTexture->Source.GetSizeX();
-			int32 SizeV = XYOffsetTexture->Source.GetSizeY();
-			int32 WeightmapOffsetX = static_cast<int32>(Component->WeightmapScaleBias.Z * SizeU);
-			int32 WeightmapOffsetY = static_cast<int32>(Component->WeightmapScaleBias.W * SizeV);
-
-			// Find coordinates of box that lies inside component
-			int32 ComponentX1 = FMath::Clamp<int32>(X1-ComponentIndexX*ComponentSizeQuads, 0, ComponentSizeQuads);
-			int32 ComponentY1 = FMath::Clamp<int32>(Y1-ComponentIndexY*ComponentSizeQuads, 0, ComponentSizeQuads);
-			int32 ComponentX2 = FMath::Clamp<int32>(X2-ComponentIndexX*ComponentSizeQuads, 0, ComponentSizeQuads);
-			int32 ComponentY2 = FMath::Clamp<int32>(Y2-ComponentIndexY*ComponentSizeQuads, 0, ComponentSizeQuads);
-
-			// Find subsection range for this box
-			int32 SubIndexX1 = FMath::Clamp<int32>((ComponentX1-1) / SubsectionSizeQuads,0,ComponentNumSubsections-1);	// -1 because we need to pick up vertices shared between subsections
-			int32 SubIndexY1 = FMath::Clamp<int32>((ComponentY1-1) / SubsectionSizeQuads,0,ComponentNumSubsections-1);
-			int32 SubIndexX2 = FMath::Clamp<int32>(ComponentX2 / SubsectionSizeQuads,0,ComponentNumSubsections-1);
-			int32 SubIndexY2 = FMath::Clamp<int32>(ComponentY2 / SubsectionSizeQuads,0,ComponentNumSubsections-1);
-
-			for( int32 SubIndexY=SubIndexY1;SubIndexY<=SubIndexY2;SubIndexY++ )
-			{
-				for( int32 SubIndexX=SubIndexX1;SubIndexX<=SubIndexX2;SubIndexX++ )
-				{
-					// Find coordinates of box that lies inside subsection
-					int32 SubX1 = FMath::Clamp<int32>(ComponentX1-SubsectionSizeQuads*SubIndexX, 0, SubsectionSizeQuads);
-					int32 SubY1 = FMath::Clamp<int32>(ComponentY1-SubsectionSizeQuads*SubIndexY, 0, SubsectionSizeQuads);
-					int32 SubX2 = FMath::Clamp<int32>(ComponentX2-SubsectionSizeQuads*SubIndexX, 0, SubsectionSizeQuads);
-					int32 SubY2 = FMath::Clamp<int32>(ComponentY2-SubsectionSizeQuads*SubIndexY, 0, SubsectionSizeQuads);
-
-					// Update texture data for the box that lies inside subsection
-					for( int32 SubY=SubY1;SubY<=SubY2;SubY++ )
-					{
-						for( int32 SubX=SubX1;SubX<=SubX2;SubX++ )
-						{
-							int32 LandscapeX = SubIndexX*SubsectionSizeQuads + ComponentIndexX*ComponentSizeQuads + SubX;
-							int32 LandscapeY = SubIndexY*SubsectionSizeQuads + ComponentIndexY*ComponentSizeQuads + SubY;
-							checkSlow( LandscapeX >= X1 && LandscapeX <= X2 );
-							checkSlow( LandscapeY >= Y1 && LandscapeY <= Y2 );
-
-							// Find the input data corresponding to this vertex
-							int32 DataIndex = (LandscapeX-X1) + Stride * (LandscapeY-Y1);
-							const T& Value = Data[DataIndex];
-
-							int32 TexX = WeightmapOffsetX + (SubsectionSizeQuads+1) * SubIndexX + SubX;
-							int32 TexY = WeightmapOffsetY + (SubsectionSizeQuads+1) * SubIndexY + SubY;
-							FColor& TexData = XYOffsetTextureData[ TexX + TexY * SizeU ];
-
-							uint16 XOffset = FMath::Clamp<uint16>(static_cast<uint16>(Value.X * LANDSCAPE_INV_XYOFFSET_SCALE + 32768.0), 0, 65535);
-							uint16 YOffset = FMath::Clamp<uint16>(static_cast<uint16>(Value.Y * LANDSCAPE_INV_XYOFFSET_SCALE + 32768.0), 0, 65535);
-
-							TexData.R = XOffset >> 8;
-							TexData.G = XOffset & 255;
-							TexData.B = YOffset >> 8;
-							TexData.A = YOffset & 255;
-						}
-					}
-
-					// Record the areas of the texture we need to re-upload
-					int32 TexX1 = WeightmapOffsetX + (SubsectionSizeQuads+1) * SubIndexX + SubX1;
-					int32 TexY1 = WeightmapOffsetY + (SubsectionSizeQuads+1) * SubIndexY + SubY1;
-					int32 TexX2 = WeightmapOffsetX + (SubsectionSizeQuads+1) * SubIndexX + SubX2;
-					int32 TexY2 = WeightmapOffsetY + (SubsectionSizeQuads+1) * SubIndexY + SubY2;
-					TexDataInfo->AddMipUpdateRegion(0,TexX1,TexY1,TexX2,TexY2);
-				}
-			}
-
-			// Update mipmaps
-			int32 NumMips = XYOffsetTexture->Source.GetNumMips();
-			TArray<FColor*> TextureMipData;
-			TextureMipData.AddUninitialized(NumMips);
-			for( int32 MipIdx=0;MipIdx<NumMips;MipIdx++ )
-			{
-				TextureMipData[MipIdx] = (FColor*)TexDataInfo->GetMipData(MipIdx);
-			}
-			ULandscapeComponent::UpdateWeightmapMips(ComponentNumSubsections, SubsectionSizeQuads, XYOffsetTexture, TextureMipData, ComponentX1, ComponentY1, ComponentX2, ComponentY2, TexDataInfo);
-		}
-	}
-}
-
+// Deprecated
 void FLandscapeEditDataInterface::SetXYOffsetData(int32 X1, int32 Y1, int32 X2, int32 Y2, const FVector2D* Data, int32 Stride)
 {
-	SetXYOffsetDataTempl<FVector2D>(X1, Y1, X2, Y2, Data, Stride);
 }
 
+// Deprecated
 void FLandscapeEditDataInterface::SetXYOffsetData(int32 X1, int32 Y1, int32 X2, int32 Y2, const FVector* Data, int32 Stride)
 {
-	SetXYOffsetDataTempl<FVector>(X1, Y1, X2, Y2, Data, Stride);
 }
 
+// Deprecated
 FVector2D FLandscapeEditDataInterface::GetXYOffsetmapData(const ULandscapeComponent* Component, int32 TexU, int32 TexV, FColor* TextureData/* = NULL*/)
 {
-	check(Component);
-	if (!TextureData && Component->XYOffsetmapTexture)
-	{
-		FLandscapeTextureDataInfo* TexDataInfo = GetTextureDataInfo(Component->XYOffsetmapTexture);
-		TextureData = (FColor*)TexDataInfo->GetMipData(0);	
-	}
-
-	if (TextureData)
-	{
-		int32 SizeU = Component->NumSubsections * (Component->SubsectionSizeQuads + 1);
-		int32 SizeV = Component->NumSubsections * (Component->SubsectionSizeQuads + 1);
-
-		int32 TexX = TexU;
-		int32 TexY = TexV;
-		FColor& TexData = TextureData[ TexX + TexY * SizeU ];
-		return FVector2D(((TexData.R * 256.0 + TexData.G) - 32768.0) * LANDSCAPE_XYOFFSET_SCALE, ((TexData.B * 256.0 + TexData.A) - 32768.0) * LANDSCAPE_XYOFFSET_SCALE );
-	}
-	return FVector2D::ZeroVector;
+	return FVector2D();
 }
 
-// XYOffset Interpolation version
-template<typename TStoreData>
-void FLandscapeEditDataInterface::GetXYOffsetDataTempl(int32& ValidX1, int32& ValidY1, int32& ValidX2, int32& ValidY2, TStoreData& StoreData)
+// Deprecated
+void FLandscapeEditDataInterface::GetXYOffsetData(int32& X1, int32& Y1, int32& X2, int32& Y2, FVector2D* Data, int32 Stride)
 {
-	// Copy variables
-	int32 X1 = ValidX1, X2 = ValidX2, Y1 = ValidY1, Y2 = ValidY2;
-	ValidX1 = INT_MAX; ValidX2 = INT_MIN; ValidY1 = INT_MAX; ValidY2 = INT_MIN;
-
-	int32 ComponentIndexX1, ComponentIndexY1, ComponentIndexX2, ComponentIndexY2;
-	ALandscape::CalcComponentIndicesNoOverlap(X1, Y1, X2, Y2, ComponentSizeQuads, ComponentIndexX1, ComponentIndexY1, ComponentIndexX2, ComponentIndexY2);
-
-	int32 ComponentSizeX = ComponentIndexX2 - ComponentIndexX1 + 1;
-	int32 ComponentSizeY = ComponentIndexY2 - ComponentIndexY1 + 1;
+}
+// Deprecated
+void FLandscapeEditDataInterface::GetXYOffsetData(int32& X1, int32& Y1, int32& X2, int32& Y2, TMap<FIntPoint, FVector2D>& SparseData)
+{
+}
 
-	// Neighbor Components
-	ULandscapeComponent* BorderComponent[4] = { 0, 0, 0, 0 };
-	ULandscapeComponent* CornerComponent[4] = { 0, 0, 0, 0 };
-	bool NoBorderX1 = false, NoBorderX2 = false;
-	TArray<bool> NoBorderY1, NoBorderY2, ComponentDataExist;
-	TArray<ULandscapeComponent*> BorderComponentY1, BorderComponentY2;
-	ComponentDataExist.Empty(ComponentSizeX*ComponentSizeY);
-	ComponentDataExist.AddZeroed(ComponentSizeX*ComponentSizeY);
-	bool bHasMissingValue = false;
+// Deprecated
+void FLandscapeEditDataInterface::GetXYOffsetData(int32& X1, int32& Y1, int32& X2, int32& Y2, FVector* Data, int32 Stride)
+{
+}
 
-	FLandscapeTextureDataInfo* NeighborTexDataInfo[4] = { 0, 0, 0, 0 };
-	FColor* NeighborXYOffsetmapTextureData[4] = { 0, 0, 0, 0 };
-	FVector2D CornerValues[4] = { FVector2D::ZeroVector, FVector2D::ZeroVector, FVector2D::ZeroVector, FVector2D::ZeroVector };
+// Deprecated
+void FLandscapeEditDataInterface::GetXYOffsetData(int32& X1, int32& Y1, int32& X2, int32& Y2, TMap<FIntPoint, FVector>& SparseData)
+{
+}
 
-	int32 EdgeCoord = (SubsectionSizeQuads + 1) * ComponentNumSubsections - 1; //ComponentSizeQuads;
+const ALandscape* FLandscapeEditDataInterface::GetTargetLandscape() const
+{
+	return LandscapeInfo ? LandscapeInfo->LandscapeActor.Get() : nullptr;
+}
 
-	TArray<FColor> EmptyXYOffset;
-	int32 XYOffsetSize = (LandscapeInfo->SubsectionSizeQuads + 1) * LandscapeInfo->ComponentNumSubsections;
-	XYOffsetSize = XYOffsetSize * XYOffsetSize;
-	EmptyXYOffset.Empty(XYOffsetSize);
-	for (int32 i = 0; i < XYOffsetSize; ++i)
-	{
-		EmptyXYOffset.Add(FColor(128, 0, 128, 0));
-	}
+// Deprecated 
+bool FLandscapeEditDataInterface::CanHaveLandscapeLayersContent() const
+{
+	return true;
+}
 
-	// initial loop....
-	for (int32 ComponentIndexY = ComponentIndexY1; ComponentIndexY <= ComponentIndexY2; ComponentIndexY++)
-	{
-		NoBorderX1 = false;
-		NoBorderX2 = false;
-		BorderComponent[0] = BorderComponent[1] = NULL;
-		for (int32 ComponentIndexX = ComponentIndexX1; ComponentIndexX <= ComponentIndexX2; ComponentIndexX++)
-		{
-			BorderComponent[2] = BorderComponent[3] = NULL;
-			int32 ComponentIndexXY = ComponentSizeX*(ComponentIndexY - ComponentIndexY1) + ComponentIndexX - ComponentIndexX1;
-			int32 ComponentIndexXX = ComponentIndexX - ComponentIndexX1;
-			int32 ComponentIndexYY = ComponentIndexY - ComponentIndexY1;
-			ComponentDataExist[ComponentIndexXY] = false;
-			ULandscapeComponent* Component = LandscapeInfo->XYtoComponentMap.FindRef(FIntPoint(ComponentIndexX, ComponentIndexY));
+// Deprecated 
+bool FLandscapeEditDataInterface::HasLandscapeLayersContent() const
+{
+	return true;
+}
 
-			FLandscapeTextureDataInfo* TexDataInfo = NULL;
-			FColor* XYOffsetmapTextureData = NULL;
-			uint8 CornerSet = 0;
-			bool ExistLeft = ComponentIndexXX > 0 && ComponentDataExist[ComponentIndexXX - 1 + ComponentIndexYY * ComponentSizeX];
-			bool ExistUp = ComponentIndexYY > 0 && ComponentDataExist[ComponentIndexXX + (ComponentIndexYY - 1) * ComponentSizeX];
+// Deprecated
+void FLandscapeEditDataInterface::GetXYOffsetDataFast(const int32 X1, const int32 Y1, const int32 X2, const int32 Y2, FVector2D* Data, int32 Stride)
+{
+}
 
-			if (Component)
-			{
-				if (Component->XYOffsetmapTexture)
-				{
-					TexDataInfo = GetTextureDataInfo(Component->XYOffsetmapTexture);
-					XYOffsetmapTextureData = (FColor*)TexDataInfo->GetMipData(0);
-				}
-				else
-				{
-					XYOffsetmapTextureData = EmptyXYOffset.GetData();
-				}
-				ComponentDataExist[ComponentIndexXY] = true;
-				// Update valid region
-				ValidX1 = FMath::Min<int32>(Component->GetSectionBase().X, ValidX1);
-				ValidX2 = FMath::Max<int32>(Component->GetSectionBase().X + ComponentSizeQuads, ValidX2);
-				ValidY1 = FMath::Min<int32>(Component->GetSectionBase().Y, ValidY1);
-				ValidY2 = FMath::Max<int32>(Component->GetSectionBase().Y + ComponentSizeQuads, ValidY2);
-			}
-			else
-			{
-				if (!bHasMissingValue)
-				{
-					NoBorderY1.Empty(ComponentSizeX);
-					NoBorderY2.Empty(ComponentSizeX);
-					NoBorderY1.AddZeroed(ComponentSizeX);
-					NoBorderY2.AddZeroed(ComponentSizeX);
-					BorderComponentY1.Empty(ComponentSizeX);
-					BorderComponentY2.Empty(ComponentSizeX);
-					BorderComponentY1.AddZeroed(ComponentSizeX);
-					BorderComponentY2.AddZeroed(ComponentSizeX);
-					bHasMissingValue = true;
-				}
-
-				// Search for neighbor component for interpolation
-				bool bShouldSearchX = (BorderComponent[1] && BorderComponent[1]->GetSectionBase().X / ComponentSizeQuads <= ComponentIndexX);
-				bool bShouldSearchY = (BorderComponentY2[ComponentIndexXX] && BorderComponentY2[ComponentIndexXX]->GetSectionBase().Y / ComponentSizeQuads <= ComponentIndexY);
-				// Search for left-closest component
-				if (bShouldSearchX || (!NoBorderX1 && !BorderComponent[0]))
-				{
-					NoBorderX1 = true;
-					for (int32 X = ComponentIndexX - 1; X >= ComponentIndexX1; X--)
-					{
-						BorderComponent[0] = LandscapeInfo->XYtoComponentMap.FindRef(FIntPoint(X, ComponentIndexY));
-						if (BorderComponent[0])
-						{
-							NoBorderX1 = false;
-							if (BorderComponent[0]->XYOffsetmapTexture)
-							{
-								NeighborTexDataInfo[0] = GetTextureDataInfo(BorderComponent[0]->XYOffsetmapTexture);
-								NeighborXYOffsetmapTextureData[0] = (FColor*)NeighborTexDataInfo[0]->GetMipData(0);
-							}
-							else
-							{
-								NeighborXYOffsetmapTextureData[0] = EmptyXYOffset.GetData();
-							}
-							break;
-						}
-					}
-				}
-				// Search for right-closest component
-				if (bShouldSearchX || (!NoBorderX2 && !BorderComponent[1]))
-				{
-					NoBorderX2 = true;
-					for (int32 X = ComponentIndexX + 1; X <= ComponentIndexX2; X++)
-					{
-						BorderComponent[1] = LandscapeInfo->XYtoComponentMap.FindRef(FIntPoint(X, ComponentIndexY));
-						if (BorderComponent[1])
-						{
-							NoBorderX2 = false;
-							if (BorderComponent[1]->XYOffsetmapTexture)
-							{
-								NeighborTexDataInfo[1] = GetTextureDataInfo(BorderComponent[1]->XYOffsetmapTexture);
-								NeighborXYOffsetmapTextureData[1] = (FColor*)NeighborTexDataInfo[1]->GetMipData(0);
-							}
-							else
-							{
-								NeighborXYOffsetmapTextureData[1] = EmptyXYOffset.GetData();
-							}
-							break;
-						}
-					}
-				}
-				// Search for up-closest component
-				if (bShouldSearchY || (!NoBorderY1[ComponentIndexXX] && !BorderComponentY1[ComponentIndexXX]))
-				{
-					NoBorderY1[ComponentIndexXX] = true;
-					for (int32 Y = ComponentIndexY - 1; Y >= ComponentIndexY1; Y--)
-					{
-						BorderComponentY1[ComponentIndexXX] = BorderComponent[2] = LandscapeInfo->XYtoComponentMap.FindRef(FIntPoint(ComponentIndexX, Y));
-						if (BorderComponent[2])
-						{
-							NoBorderY1[ComponentIndexXX] = false;
-							if (BorderComponent[2]->XYOffsetmapTexture)
-							{
-								NeighborTexDataInfo[2] = GetTextureDataInfo(BorderComponent[2]->XYOffsetmapTexture);
-								NeighborXYOffsetmapTextureData[2] = (FColor*)NeighborTexDataInfo[2]->GetMipData(0);
-							}
-							else
-							{
-								NeighborXYOffsetmapTextureData[2] = EmptyXYOffset.GetData();
-							}
-							break;
-						}
-					}
-				}
-				else
-				{
-					BorderComponent[2] = BorderComponentY1[ComponentIndexXX];
-					if (BorderComponent[2])
-					{
-						if (BorderComponent[2]->XYOffsetmapTexture)
-						{
-							NeighborTexDataInfo[2] = GetTextureDataInfo(BorderComponent[2]->XYOffsetmapTexture);
-							NeighborXYOffsetmapTextureData[2] = (FColor*)NeighborTexDataInfo[2]->GetMipData(0);
-						}
-						else
-						{
-							NeighborXYOffsetmapTextureData[2] = EmptyXYOffset.GetData();
-						}
-
-					}
-				}
-				// Search for bottom-closest component
-				if (bShouldSearchY || (!NoBorderY2[ComponentIndexXX] && !BorderComponentY2[ComponentIndexXX]))
-				{
-					NoBorderY2[ComponentIndexXX] = true;
-					for (int32 Y = ComponentIndexY + 1; Y <= ComponentIndexY2; Y++)
-					{
-						BorderComponentY2[ComponentIndexXX] = BorderComponent[3] = LandscapeInfo->XYtoComponentMap.FindRef(FIntPoint(ComponentIndexX, Y));
-						if (BorderComponent[3])
-						{
-							NoBorderY2[ComponentIndexXX] = false;
-							if (BorderComponent[3]->XYOffsetmapTexture)
-							{
-								NeighborTexDataInfo[3] = GetTextureDataInfo(BorderComponent[3]->XYOffsetmapTexture);
-								NeighborXYOffsetmapTextureData[3] = (FColor*)NeighborTexDataInfo[3]->GetMipData(0);
-							}
-							else
-							{
-								NeighborXYOffsetmapTextureData[3] = EmptyXYOffset.GetData();
-							}
-							break;
-						}
-					}
-				}
-				else
-				{
-					BorderComponent[3] = BorderComponentY2[ComponentIndexXX];
-					if (BorderComponent[3])
-					{
-						if (BorderComponent[3]->XYOffsetmapTexture)
-						{
-							NeighborTexDataInfo[3] = GetTextureDataInfo(BorderComponent[3]->XYOffsetmapTexture);
-							NeighborXYOffsetmapTextureData[3] = (FColor*)NeighborTexDataInfo[3]->GetMipData(0);
-						}
-						else
-						{
-							NeighborXYOffsetmapTextureData[3] = EmptyXYOffset.GetData();
-						}
-					}
-				}
-
-				CornerComponent[0] = ComponentIndexX >= ComponentIndexX1 && ComponentIndexY >= ComponentIndexY1 ?
-					LandscapeInfo->XYtoComponentMap.FindRef(FIntPoint((ComponentIndexX - 1), (ComponentIndexY - 1))) : NULL;
-				CornerComponent[1] = ComponentIndexX <= ComponentIndexX2 && ComponentIndexY >= ComponentIndexY1 ?
-					LandscapeInfo->XYtoComponentMap.FindRef(FIntPoint((ComponentIndexX + 1), (ComponentIndexY - 1))) : NULL;
-				CornerComponent[2] = ComponentIndexX >= ComponentIndexX1 && ComponentIndexY <= ComponentIndexY2 ?
-					LandscapeInfo->XYtoComponentMap.FindRef(FIntPoint((ComponentIndexX - 1), (ComponentIndexY + 1))) : NULL;
-				CornerComponent[3] = ComponentIndexX <= ComponentIndexX2 && ComponentIndexY <= ComponentIndexY2 ?
-					LandscapeInfo->XYtoComponentMap.FindRef(FIntPoint((ComponentIndexX + 1), (ComponentIndexY + 1))) : NULL;
-
-				if (CornerComponent[0])
-				{
-					CornerSet |= 1;
-					CornerValues[0] = GetXYOffsetmapData(CornerComponent[0], EdgeCoord, EdgeCoord);
-				}
-				else if ((ExistLeft || ExistUp) && X1 <= ComponentIndexX*ComponentSizeQuads && Y1 <= ComponentIndexY*ComponentSizeQuads)
-				{
-					CornerSet |= 1;
-					CornerValues[0] = FVector2D(StoreData.Load(ComponentIndexX*ComponentSizeQuads, ComponentIndexY*ComponentSizeQuads));
-				}
-				else if (BorderComponent[0])
-				{
-					CornerSet |= 1;
-					CornerValues[0] = GetXYOffsetmapData(BorderComponent[0], EdgeCoord, 0, NeighborXYOffsetmapTextureData[0]);
-				}
-				else if (BorderComponent[2])
-				{
-					CornerSet |= 1;
-					CornerValues[0] = GetXYOffsetmapData(BorderComponent[2], 0, EdgeCoord, NeighborXYOffsetmapTextureData[2]);
-				}
-
-				if (CornerComponent[1])
-				{
-					CornerSet |= 1 << 1;
-					CornerValues[1] = GetXYOffsetmapData(CornerComponent[1], 0, EdgeCoord);
-				}
-				else if (ExistUp && X2 >= (ComponentIndexX + 1)*ComponentSizeQuads)
-				{
-					CornerSet |= 1 << 1;
-					CornerValues[1] = FVector2D(StoreData.Load((ComponentIndexX + 1)*ComponentSizeQuads, ComponentIndexY*ComponentSizeQuads));
-				}
-				else if (BorderComponent[1])
-				{
-					CornerSet |= 1 << 1;
-					CornerValues[1] = GetXYOffsetmapData(BorderComponent[1], 0, 0, NeighborXYOffsetmapTextureData[1]);
-				}
-				else if (BorderComponent[2])
-				{
-					CornerSet |= 1 << 1;
-					CornerValues[1] = GetXYOffsetmapData(BorderComponent[2], EdgeCoord, EdgeCoord, NeighborXYOffsetmapTextureData[2]);
-				}
-
-				if (CornerComponent[2])
-				{
-					CornerSet |= 1 << 2;
-					CornerValues[2] = GetXYOffsetmapData(CornerComponent[2], EdgeCoord, 0);
-				}
-				else if (ExistLeft && Y2 >= (ComponentIndexY + 1)*ComponentSizeQuads) // Use data already stored for 0, 2
-				{
-					CornerSet |= 1 << 2;
-					CornerValues[2] = FVector2D(StoreData.Load(ComponentIndexX*ComponentSizeQuads, (ComponentIndexY + 1)*ComponentSizeQuads));
-				}
-				else if (BorderComponent[0])
-				{
-					CornerSet |= 1 << 2;
-					CornerValues[2] = GetXYOffsetmapData(BorderComponent[0], EdgeCoord, EdgeCoord, NeighborXYOffsetmapTextureData[0]);
-				}
-				else if (BorderComponent[3])
-				{
-					CornerSet |= 1 << 2;
-					CornerValues[2] = GetXYOffsetmapData(BorderComponent[3], 0, 0, NeighborXYOffsetmapTextureData[3]);
-				}
-
-				if (CornerComponent[3])
-				{
-					CornerSet |= 1 << 3;
-					CornerValues[3] = GetXYOffsetmapData(CornerComponent[3], 0, 0);
-				}
-				else if (BorderComponent[1])
-				{
-					CornerSet |= 1 << 3;
-					CornerValues[3] = GetXYOffsetmapData(BorderComponent[1], 0, EdgeCoord, NeighborXYOffsetmapTextureData[1]);
-				}
-				else if (BorderComponent[3])
-				{
-					CornerSet |= 1 << 3;
-					CornerValues[3] = GetXYOffsetmapData(BorderComponent[3], EdgeCoord, 0, NeighborXYOffsetmapTextureData[3]);
-				}
-
-				FillCornerValues(CornerSet, CornerValues);
-				ComponentDataExist[ComponentIndexXY] = ExistLeft || ExistUp || (BorderComponent[0] || BorderComponent[1] || BorderComponent[2] || BorderComponent[3]) || CornerSet;
-			}
-
-			if (!ComponentDataExist[ComponentIndexXY])
-			{
-				continue;
-			}
-
-			// Find coordinates of box that lies inside component
-			int32 ComponentX1 = FMath::Clamp<int32>(X1 - ComponentIndexX*ComponentSizeQuads, 0, ComponentSizeQuads);
-			int32 ComponentY1 = FMath::Clamp<int32>(Y1 - ComponentIndexY*ComponentSizeQuads, 0, ComponentSizeQuads);
-			int32 ComponentX2 = FMath::Clamp<int32>(X2 - ComponentIndexX*ComponentSizeQuads, 0, ComponentSizeQuads);
-			int32 ComponentY2 = FMath::Clamp<int32>(Y2 - ComponentIndexY*ComponentSizeQuads, 0, ComponentSizeQuads);
-
-			// Find subsection range for this box
-			int32 SubIndexX1 = FMath::Clamp<int32>((ComponentX1 - 1) / SubsectionSizeQuads, 0, ComponentNumSubsections - 1);	// -1 because we need to pick up vertices shared between subsections
-			int32 SubIndexY1 = FMath::Clamp<int32>((ComponentY1 - 1) / SubsectionSizeQuads, 0, ComponentNumSubsections - 1);
-			int32 SubIndexX2 = FMath::Clamp<int32>(ComponentX2 / SubsectionSizeQuads, 0, ComponentNumSubsections - 1);
-			int32 SubIndexY2 = FMath::Clamp<int32>(ComponentY2 / SubsectionSizeQuads, 0, ComponentNumSubsections - 1);
-
-			for (int32 SubIndexY = SubIndexY1; SubIndexY <= SubIndexY2; SubIndexY++)
-			{
-				for (int32 SubIndexX = SubIndexX1; SubIndexX <= SubIndexX2; SubIndexX++)
-				{
-					// Find coordinates of box that lies inside subsection
-					int32 SubX1 = FMath::Clamp<int32>(ComponentX1 - SubsectionSizeQuads*SubIndexX, 0, SubsectionSizeQuads);
-					int32 SubY1 = FMath::Clamp<int32>(ComponentY1 - SubsectionSizeQuads*SubIndexY, 0, SubsectionSizeQuads);
-					int32 SubX2 = FMath::Clamp<int32>(ComponentX2 - SubsectionSizeQuads*SubIndexX, 0, SubsectionSizeQuads);
-					int32 SubY2 = FMath::Clamp<int32>(ComponentY2 - SubsectionSizeQuads*SubIndexY, 0, SubsectionSizeQuads);
-
-					// Update texture data for the box that lies inside subsection
-					for (int32 SubY = SubY1; SubY <= SubY2; SubY++)
-					{
-						for (int32 SubX = SubX1; SubX <= SubX2; SubX++)
-						{
-							int32 LandscapeX = SubIndexX*SubsectionSizeQuads + ComponentIndexX*ComponentSizeQuads + SubX;
-							int32 LandscapeY = SubIndexY*SubsectionSizeQuads + ComponentIndexY*ComponentSizeQuads + SubY;
-
-							// Find the input data corresponding to this vertex
-							if (Component)
-							{
-								// Find the texture data corresponding to this vertex
-								FVector2D XYOffset = GetXYOffsetmapData(Component, (SubsectionSizeQuads + 1) * SubIndexX + SubX, (SubsectionSizeQuads + 1) * SubIndexY + SubY, XYOffsetmapTextureData);
-								StoreData.Store(LandscapeX, LandscapeY, XYOffset);
-							}
-							else
-							{
-								// Find the texture data corresponding to this vertex
-								FVector2D Value[4] = { FVector2D::ZeroVector, FVector2D::ZeroVector, FVector2D::ZeroVector, FVector2D::ZeroVector };
-								int32 Dist[4] = { INT_MAX, INT_MAX, INT_MAX, INT_MAX };
-								FVector2D ValueX = FVector2D::ZeroVector, ValueY = FVector2D::ZeroVector;
-								bool Exist[4] = { false, false, false, false };
-
-								// Use data already stored for 0, 2
-								if (ExistLeft)
-								{
-									Value[0] = FVector2D(StoreData.Load(ComponentIndexX*ComponentSizeQuads, LandscapeY));
-									Dist[0] = LandscapeX - (ComponentIndexX*ComponentSizeQuads);
-									Exist[0] = true;
-								}
-								else if (BorderComponent[0])
-								{
-									Value[0] = GetXYOffsetmapData(BorderComponent[0], EdgeCoord, (SubsectionSizeQuads + 1) * SubIndexY + SubY, NeighborXYOffsetmapTextureData[0]);
-									Dist[0] = LandscapeX - (BorderComponent[0]->GetSectionBase().X + ComponentSizeQuads);
-									Exist[0] = true;
-								}
-								else
-								{
-									if ((CornerSet & 1) && (CornerSet & 1 << 2))
-									{
-										int32 Dist1 = LandscapeY - (ComponentIndexY*ComponentSizeQuads);
-										int32 Dist2 = ((ComponentIndexY + 1)*ComponentSizeQuads) - LandscapeY;
-										Value[0] = (Dist2 * CornerValues[0] + Dist1 * CornerValues[2]) / (Dist1 + Dist2);
-										Dist[0] = LandscapeX - (ComponentIndexX*ComponentSizeQuads);
-										Exist[0] = true;
-									}
-								}
-
-								if (BorderComponent[1])
-								{
-									Value[1] = GetXYOffsetmapData(BorderComponent[1], 0, (SubsectionSizeQuads + 1) * SubIndexY + SubY, NeighborXYOffsetmapTextureData[1]);
-									Dist[1] = (BorderComponent[1]->GetSectionBase().X) - LandscapeX;
-									Exist[1] = true;
-								}
-								else
-								{
-									if ((CornerSet & 1 << 1) && (CornerSet & 1 << 3))
-									{
-										int32 Dist1 = LandscapeY - (ComponentIndexY*ComponentSizeQuads);
-										int32 Dist2 = ((ComponentIndexY + 1)*ComponentSizeQuads) - LandscapeY;
-										Value[1] = (Dist2 * CornerValues[1] + Dist1 * CornerValues[3]) / (Dist1 + Dist2);
-										Dist[1] = (ComponentIndexX + 1)*ComponentSizeQuads - LandscapeX;
-										Exist[1] = true;
-									}
-								}
-
-								if (ExistUp)
-								{
-									Value[2] = FVector2D(StoreData.Load(LandscapeX, ComponentIndexY*ComponentSizeQuads));
-									Dist[2] = LandscapeY - (ComponentIndexY*ComponentSizeQuads);
-									Exist[2] = true;
-								}
-								else if (BorderComponent[2])
-								{
-									Value[2] = GetXYOffsetmapData(BorderComponent[2], (SubsectionSizeQuads + 1) * SubIndexX + SubX, EdgeCoord, NeighborXYOffsetmapTextureData[2]);
-									Dist[2] = LandscapeY - (BorderComponent[2]->GetSectionBase().Y + ComponentSizeQuads);
-									Exist[2] = true;
-								}
-								else
-								{
-									if ((CornerSet & 1) && (CornerSet & 1 << 1))
-									{
-										int32 Dist1 = LandscapeX - (ComponentIndexX*ComponentSizeQuads);
-										int32 Dist2 = (ComponentIndexX + 1)*ComponentSizeQuads - LandscapeX;
-										Value[2] = (Dist2 * CornerValues[0] + Dist1 * CornerValues[1]) / (Dist1 + Dist2);
-										Dist[2] = LandscapeY - (ComponentIndexY*ComponentSizeQuads);
-										Exist[2] = true;
-									}
-								}
-
-								if (BorderComponent[3])
-								{
-									Value[3] = GetXYOffsetmapData(BorderComponent[3], (SubsectionSizeQuads + 1) * SubIndexX + SubX, 0, NeighborXYOffsetmapTextureData[3]);
-									Dist[3] = (BorderComponent[3]->GetSectionBase().Y) - LandscapeY;
-									Exist[3] = true;
-								}
-								else
-								{
-									if ((CornerSet & 1 << 2) && (CornerSet & 1 << 3))
-									{
-										int32 Dist1 = LandscapeX - (ComponentIndexX*ComponentSizeQuads);
-										int32 Dist2 = (ComponentIndexX + 1)*ComponentSizeQuads - LandscapeX;
-										Value[3] = (Dist2 * CornerValues[2] + Dist1 * CornerValues[3]) / (Dist1 + Dist2);
-										Dist[3] = (ComponentIndexY + 1)*ComponentSizeQuads - LandscapeY;
-										Exist[3] = true;
-									}
-								}
-
-								CalcInterpValue<FVector2D>(Dist, Exist, Value, ValueX, ValueY);
-
-								FVector2D FinalValue = FVector2D::ZeroVector; // Default Value
-								if ((Exist[0] || Exist[1]) && (Exist[2] || Exist[3]))
-								{
-									FinalValue = CalcValueFromValueXY<FVector2D>(Dist, ValueX, ValueY, CornerSet, CornerValues);
-								}
-								else if ((BorderComponent[0] || BorderComponent[1]))
-								{
-									FinalValue = ValueX;
-								}
-								else if ((BorderComponent[2] || BorderComponent[3]))
-								{
-									FinalValue = ValueY;
-								}
-								else if ((Exist[0] || Exist[1]))
-								{
-									FinalValue = ValueX;
-								}
-								else if ((Exist[2] || Exist[3]))
-								{
-									FinalValue = ValueY;
-								}
-
-								StoreData.Store(LandscapeX, LandscapeY, FinalValue);
-							}
-						}
-					}
-				}
-			}
-		}
-	}
-
-	if (bHasMissingValue)
-	{
-		// generate something to fill in any missing data within the requested region
-		CalcMissingValues<FVector2D, TStoreData, FVector2D>(X1, X2, Y1, Y2,
-			ComponentIndexX1, ComponentIndexX2, ComponentIndexY1, ComponentIndexY2,
-			ComponentSizeX, ComponentSizeY, CornerValues,
-			NoBorderY1, NoBorderY2, ComponentDataExist, StoreData);
-
-		// clamp the returned bounds to the original request bounds
-		// this means the returned bounds are allowed to shrink (if there is only a small part of the request that actually exists)
-		// but it is not allowed to expand beyond the original request
-		ValidX1 = FMath::Max<int32>(X1, ValidX1);
-		ValidX2 = FMath::Min<int32>(X2, ValidX2);
-		ValidY1 = FMath::Max<int32>(Y1, ValidY1);
-		ValidY2 = FMath::Min<int32>(Y2, ValidY2);
-	}
-	else
-	{
-		// if there were no missing values, then the entire requested bounds is returned
-		ValidX1 = X1;
-		ValidX2 = X2;
-		ValidY1 = Y1;
-		ValidY2 = Y2;
-	}
-}
-
-void FLandscapeEditDataInterface::GetXYOffsetData(int32& X1, int32& Y1, int32& X2, int32& Y2, FVector2D* Data, int32 Stride)
-{
-	if (Stride == 0)
-	{
-		Stride = (1 + X2 - X1);
-	}
-	TArrayStoreData<FVector2D> ArrayStoreData(X1, Y1, Data, Stride);
-	GetXYOffsetDataTempl(X1, Y1, X2, Y2, ArrayStoreData);
-}
-
-void FLandscapeEditDataInterface::GetXYOffsetData(int32& X1, int32& Y1, int32& X2, int32& Y2, TMap<FIntPoint, FVector2D>& SparseData)
-{
-	TSparseStoreData<FVector2D> SparseStoreData(SparseData);
-	GetXYOffsetDataTempl(X1, Y1, X2, Y2, SparseStoreData);
-}
-
-void FLandscapeEditDataInterface::GetXYOffsetData(int32& X1, int32& Y1, int32& X2, int32& Y2, FVector* Data, int32 Stride)
-{
-	if (Stride == 0)
-	{
-		Stride = (1 + X2 - X1);
-	}
-	TArrayStoreData<FVector> ArrayStoreData(X1, Y1, Data, Stride);
-	GetXYOffsetDataTempl(X1, Y1, X2, Y2, ArrayStoreData);
-}
-
-void FLandscapeEditDataInterface::GetXYOffsetData(int32& X1, int32& Y1, int32& X2, int32& Y2, TMap<FIntPoint, FVector>& SparseData)
-{
-	TSparseStoreData<FVector> SparseStoreData(SparseData);
-	GetXYOffsetDataTempl(X1, Y1, X2, Y2, SparseStoreData);
-}
-
-template<typename TStoreData>
-void FLandscapeEditDataInterface::GetXYOffsetDataTemplFast(const int32 X1, const int32 Y1, const int32 X2, const int32 Y2, TStoreData& StoreData)
-{
-	int32 ComponentIndexX1, ComponentIndexY1, ComponentIndexX2, ComponentIndexY2;
-	ALandscape::CalcComponentIndicesNoOverlap(X1, Y1, X2, Y2, ComponentSizeQuads, ComponentIndexX1, ComponentIndexY1, ComponentIndexX2, ComponentIndexY2);
-
-	for( int32 ComponentIndexY=ComponentIndexY1;ComponentIndexY<=ComponentIndexY2;ComponentIndexY++ )
-	{
-		for( int32 ComponentIndexX=ComponentIndexX1;ComponentIndexX<=ComponentIndexX2;ComponentIndexX++ )
-		{		
-			ULandscapeComponent* Component = LandscapeInfo->XYtoComponentMap.FindRef(FIntPoint(ComponentIndexX,ComponentIndexY));
-
-			FLandscapeTextureDataInfo* TexDataInfo = NULL;
-			FColor* OffsetTextureData = NULL;
-			if( Component && Component->XYOffsetmapTexture )
-			{
-				TexDataInfo = GetTextureDataInfo(Component->XYOffsetmapTexture);
-				OffsetTextureData = (FColor*)TexDataInfo->GetMipData(0);
-			}
-
-			// Find coordinates of box that lies inside component
-			int32 ComponentX1 = FMath::Clamp<int32>(X1-ComponentIndexX*ComponentSizeQuads, 0, ComponentSizeQuads);
-			int32 ComponentY1 = FMath::Clamp<int32>(Y1-ComponentIndexY*ComponentSizeQuads, 0, ComponentSizeQuads);
-			int32 ComponentX2 = FMath::Clamp<int32>(X2-ComponentIndexX*ComponentSizeQuads, 0, ComponentSizeQuads);
-			int32 ComponentY2 = FMath::Clamp<int32>(Y2-ComponentIndexY*ComponentSizeQuads, 0, ComponentSizeQuads);
-
-			// Find subsection range for this box
-			int32 SubIndexX1 = FMath::Clamp<int32>((ComponentX1-1) / SubsectionSizeQuads,0,ComponentNumSubsections-1);	// -1 because we need to pick up vertices shared between subsections
-			int32 SubIndexY1 = FMath::Clamp<int32>((ComponentY1-1) / SubsectionSizeQuads,0,ComponentNumSubsections-1);
-			int32 SubIndexX2 = FMath::Clamp<int32>(ComponentX2 / SubsectionSizeQuads,0,ComponentNumSubsections-1);
-			int32 SubIndexY2 = FMath::Clamp<int32>(ComponentY2 / SubsectionSizeQuads,0,ComponentNumSubsections-1);
-
-			for( int32 SubIndexY=SubIndexY1;SubIndexY<=SubIndexY2;SubIndexY++ )
-			{
-				for( int32 SubIndexX=SubIndexX1;SubIndexX<=SubIndexX2;SubIndexX++ )
-				{
-					// Find coordinates of box that lies inside subsection
-					int32 SubX1 = FMath::Clamp<int32>(ComponentX1-SubsectionSizeQuads*SubIndexX, 0, SubsectionSizeQuads);
-					int32 SubY1 = FMath::Clamp<int32>(ComponentY1-SubsectionSizeQuads*SubIndexY, 0, SubsectionSizeQuads);
-					int32 SubX2 = FMath::Clamp<int32>(ComponentX2-SubsectionSizeQuads*SubIndexX, 0, SubsectionSizeQuads);
-					int32 SubY2 = FMath::Clamp<int32>(ComponentY2-SubsectionSizeQuads*SubIndexY, 0, SubsectionSizeQuads);
-
-					// Update texture data for the box that lies inside subsection
-					for( int32 SubY=SubY1;SubY<=SubY2;SubY++ )
-					{
-						for( int32 SubX=SubX1;SubX<=SubX2;SubX++ )
-						{
-							int32 LandscapeX = SubIndexX*SubsectionSizeQuads + ComponentIndexX*ComponentSizeQuads + SubX;
-							int32 LandscapeY = SubIndexY*SubsectionSizeQuads + ComponentIndexY*ComponentSizeQuads + SubY;
-
-							// Find the input data corresponding to this vertex
-							if( Component && OffsetTextureData )
-							{
-								FVector2D Value = GetXYOffsetmapData(Component, SubX, SubY, OffsetTextureData);
-								StoreData.Store(LandscapeX, LandscapeY, Value);
-							}
-							else
-							{
-								StoreData.Store(LandscapeX, LandscapeY, FVector2D(0.0f, 0.0f) );
-							}
-						}
-					}
-				}
-			}
-		}
-	}
-}
-
-const ALandscape* FLandscapeEditDataInterface::GetTargetLandscape() const
-{
-	return LandscapeInfo ? LandscapeInfo->LandscapeActor.Get() : nullptr;
-}
-
-// Deprecated 
-bool FLandscapeEditDataInterface::CanHaveLandscapeLayersContent() const
-{
-	return true;
-}
-
-// Deprecated 
-bool FLandscapeEditDataInterface::HasLandscapeLayersContent() const
-{
-	return true;
-}
-
-void FLandscapeEditDataInterface::GetXYOffsetDataFast(const int32 X1, const int32 Y1, const int32 X2, const int32 Y2, FVector2D* Data, int32 Stride)
-{
-	if( Stride==0 )
-	{
-		Stride = (1+X2-X1);
-	}
-	TArrayStoreData<FVector2D> ArrayStoreData(X1, Y1, Data, Stride);
-	GetXYOffsetDataTemplFast(X1, Y1, X2, Y2, ArrayStoreData);
-}
-
-void FLandscapeEditDataInterface::GetXYOffsetDataFast(const int32 X1, const int32 Y1, const int32 X2, const int32 Y2, TMap<FIntPoint, FVector2D>& SparseData)
-{
-	TSparseStoreData<FVector2D> SparseStoreData(SparseData);
-	GetXYOffsetDataTemplFast(X1, Y1, X2, Y2, SparseStoreData);
-}
+// Deprecated
+void FLandscapeEditDataInterface::GetXYOffsetDataFast(const int32 X1, const int32 Y1, const int32 X2, const int32 Y2, TMap<FIntPoint, FVector2D>& SparseData)
+{
+}
 
+// Deprecated
 void FLandscapeEditDataInterface::GetXYOffsetDataFast(const int32 X1, const int32 Y1, const int32 X2, const int32 Y2, FVector* Data, int32 Stride)
 {
-	if( Stride==0 )
-	{
-		Stride = (1+X2-X1);
-	}
-	TArrayStoreData<FVector> ArrayStoreData(X1, Y1, Data, Stride);
-	GetXYOffsetDataTemplFast(X1, Y1, X2, Y2, ArrayStoreData);
 }
 
+// Deprecated
 void FLandscapeEditDataInterface::GetXYOffsetDataFast(const int32 X1, const int32 Y1, const int32 X2, const int32 Y2, TMap<FIntPoint, FVector>& SparseData)
 {
-	TSparseStoreData<FVector> SparseStoreData(SparseData);
-	GetXYOffsetDataTemplFast(X1, Y1, X2, Y2, SparseStoreData);
 }
 
 //
diff --git a/Engine/Source/Runtime/Landscape/Private/LandscapeLight.cpp b/Engine/Source/Runtime/Landscape/Private/LandscapeLight.cpp
--- a/Engine/Source/Runtime/Landscape/Private/LandscapeLight.cpp
+++ b/Engine/Source/Runtime/Landscape/Private/LandscapeLight.cpp
@@ -28,7 +28,6 @@ LandscapeLight.cpp: Static lighting for LandscapeComponents
 #define LANDSCAPE_LIGHTMAP_UV_INDEX 1
 
 TMap<FIntPoint, FColor> FLandscapeStaticLightingMesh::LandscapeUpscaleHeightDataCache;
-TMap<FIntPoint, FColor> FLandscapeStaticLightingMesh::LandscapeUpscaleXYOffsetDataCache;
 
 /** A texture mapping for landscapes */
 /** Initialization constructor. */
@@ -220,7 +219,7 @@ namespace
 		);
 	}
 
-	void GetLODData(ULandscapeComponent* LandscapeComponent, int32 X, int32 Y, int32 HeightmapOffsetX, int32 HeightmapOffsetY, int32 LODValue, int32 HeightmapStride, FColor& OutHeight, FColor& OutXYOffset)
+	void GetLODData(ULandscapeComponent* LandscapeComponent, int32 X, int32 Y, int32 HeightmapOffsetX, int32 HeightmapOffsetY, int32 LODValue, int32 HeightmapStride, FColor& OutHeight)
 	{
 		int32 ComponentSize = ((LandscapeComponent->SubsectionSizeQuads + 1) * LandscapeComponent->NumSubsections) >> LODValue;
 		int32 LODHeightmapSizeX = LandscapeComponent->GetHeightmap()->Source.GetSizeX() >> LODValue;
@@ -240,7 +239,6 @@ namespace
 
 		FLandscapeComponentDataInterface DataInterface(LandscapeComponent, LODValue);
 		FColor* HeightMipData = DataInterface.GetRawHeightData();
-		FColor* XYOffsetMipData = DataInterface.GetRawXYOffsetData();
 
 		FColor H1 = HeightMipData[XI + YI * LODHeightmapSizeX];
 		FColor H2 = HeightMipData[FMath::Min(XI + 1, LODHeightmapSizeX - 1) + YI * LODHeightmapSizeX];
@@ -252,22 +250,9 @@ namespace
 		uint8 A = MultiLerp<uint8, float>(H1.A, H2.A, XF, H3.A, H4.A, XF, YF);
 
 		OutHeight = FColor((Height >> 8), Height & 255, B, A);
-
-		if (LandscapeComponent->XYOffsetmapTexture)
-		{
-			FColor X1 = XYOffsetMipData[XI + YI * LODHeightmapSizeX];
-			FColor X2 = XYOffsetMipData[FMath::Min(XI + 1, LODHeightmapSizeX - 1) + YI * LODHeightmapSizeX];
-			FColor X3 = XYOffsetMipData[XI + FMath::Min(YI + 1, LODHeightmapSizeY - 1) * LODHeightmapSizeX];
-			FColor X4 = XYOffsetMipData[FMath::Min(XI + 1, LODHeightmapSizeX - 1) + FMath::Min(YI + 1, LODHeightmapSizeY - 1) * LODHeightmapSizeX];
-
-			uint16 XComp = MultiLerp<uint16, float>((X1.R << 8) + X1.G, (X2.R << 8) + X2.G, XF, (X3.R << 8) + X3.G, (X4.R << 8) + X4.G, XF, YF);
-			uint16 YComp = MultiLerp<uint16, float>((X1.B << 8) + X1.A, (X2.B << 8) + X2.A, XF, (X3.B << 8) + X3.A, (X4.B << 8) + X4.A, XF, YF);
-
-			OutXYOffset = FColor((XComp >> 8), XComp & 255, YComp >> 8, YComp & 255);
-		}
 	}
 
-	void InternalUpscaling(FLandscapeComponentDataInterface& DataInterface, ULandscapeComponent* LandscapeComponent, int32 InLOD, int32 GeometryLOD, TArray<FColor>& CompHeightData, TArray<FColor>& CompXYOffsetData)
+	void InternalUpscaling(FLandscapeComponentDataInterface& DataInterface, ULandscapeComponent* LandscapeComponent, int32 InLOD, int32 GeometryLOD, TArray<FColor>& CompHeightData)
 	{
 		// Upscaling using Landscape LOD system 
 		ULandscapeInfo* const Info = LandscapeComponent->GetLandscapeInfo();
@@ -331,9 +316,7 @@ namespace
 			int32 HeightDataSize = HeightmapStride * LandscapeComponent->GetHeightmap()->Source.GetSizeY() >> InLOD;
 
 			CompHeightData.Empty(HeightDataSize);
-			CompXYOffsetData.Empty(HeightDataSize);
 			CompHeightData.AddZeroed(HeightDataSize);
-			CompXYOffsetData.AddZeroed(HeightDataSize);
 
 			// Update for only component region for performance
 			int32 ComponentSize = ((LandscapeComponent->SubsectionSizeQuads + 1) * LandscapeComponent->NumSubsections) >> InLOD;
@@ -346,15 +329,10 @@ namespace
 					IXY += ComponentBase * (ComponentSize - 1);
 
 					FColor* CachedHeight = FLandscapeStaticLightingMesh::LandscapeUpscaleHeightDataCache.Find(IXY);
-					FColor* CachedXYOffset = FLandscapeStaticLightingMesh::LandscapeUpscaleXYOffsetDataCache.Find(IXY);
 
 					if (CachedHeight)
 					{
 						CompHeightData[X + Y * HeightmapStride] = *CachedHeight;
-						if (CachedXYOffset)
-						{
-							CompXYOffsetData[X + Y * HeightmapStride] = *CachedXYOffset;
-						}
 					}
 					else
 					{
@@ -393,15 +371,14 @@ namespace
 						float MorphAlpha = FMath::Fractional(RealLOD);
 
 						FColor Height[2];
-						FColor XYOffset[2];
 						::GetLODData(LandscapeComponent, X, Y, DataInterface.HeightmapComponentOffsetX, DataInterface.HeightmapComponentOffsetY,
-							FMath::Min(MaxLOD, LODValue), HeightmapStride, Height[0], XYOffset[0]);
+							FMath::Min(MaxLOD, LODValue), HeightmapStride, Height[0]);
 
 						// Interpolation between two LOD
 						if ((RealLOD > InLOD) && (LODValue + 1 <= MaxLOD) && MorphAlpha != 0.f)
 						{
 							::GetLODData(LandscapeComponent, X, Y, DataInterface.HeightmapComponentOffsetX, DataInterface.HeightmapComponentOffsetY,
-								FMath::Min(MaxLOD, LODValue + 1), HeightmapStride, Height[1], XYOffset[1]);
+								FMath::Min(MaxLOD, LODValue + 1), HeightmapStride, Height[1]);
 
 							// Need interpolation
 							const uint16 Height0 = (Height[0].R << 8) + Height[0].G;
@@ -412,41 +389,19 @@ namespace
 								FColor((LerpHeight >> 8), LerpHeight & 255,
 								static_cast<uint8>(FMath::RoundToInt32(FMath::Lerp<float>(Height[0].B, Height[1].B, MorphAlpha))),
 								static_cast<uint8>(FMath::RoundToInt32(FMath::Lerp<float>(Height[0].A, Height[1].A, MorphAlpha))));
-							if (LandscapeComponent->XYOffsetmapTexture)
-							{
-								uint16 XComp0 = (XYOffset[0].R << 8) + XYOffset[0].G;
-								uint16 XComp1 = (XYOffset[1].R << 8) + XYOffset[1].G;
-								uint16 LerpXComp = static_cast<uint16>(FMath::RoundToInt32(FMath::Lerp<float>(XComp0, XComp1, MorphAlpha)));
-
-								uint16 YComp0 = (XYOffset[0].B << 8) + XYOffset[0].A;
-								uint16 YComp1 = (XYOffset[1].B << 8) + XYOffset[1].A;
-								uint16 LerpYComp = static_cast<uint16>(FMath::RoundToInt32(FMath::Lerp<float>(YComp0, YComp1, MorphAlpha)));
-
-								CompXYOffsetData[X + Y * HeightmapStride] =
-									FColor(LerpXComp >> 8, LerpXComp & 255, LerpYComp >> 8, LerpYComp & 255);
-							}
 						}
 						else
 						{
 							CompHeightData[X + Y * HeightmapStride] = Height[0];
-							CompXYOffsetData[X + Y * HeightmapStride] = XYOffset[0];
 						}
 
 						// Caching current calculated value
 						FLandscapeStaticLightingMesh::LandscapeUpscaleHeightDataCache.Add(IXY, CompHeightData[X + Y * HeightmapStride]);
-						if (LandscapeComponent->XYOffsetmapTexture)
-						{
-							FLandscapeStaticLightingMesh::LandscapeUpscaleHeightDataCache.Add(IXY, CompXYOffsetData[X + Y * HeightmapStride]);
-						}
 					}
 				}
 			}
 
 			DataInterface.SetRawHeightData(&CompHeightData[0]);
-			if (LandscapeComponent->XYOffsetmapTexture)
-			{
-				DataInterface.SetRawXYOffsetData(&CompXYOffsetData[0]);
-			}
 		}
 	}
 };
@@ -475,9 +430,8 @@ void FLandscapeStaticLightingMesh::GetHeightmapData(int32 InLOD, int32 GeometryL
 	{
 		// Data array for upscaling case
 		TArray<FColor> CompHeightData;
-		TArray<FColor> CompXYOffsetData;
 		FLandscapeComponentDataInterface DataInterface(LandscapeComponent, InLOD);
-		::InternalUpscaling(DataInterface, LandscapeComponent, InLOD, GeometryLOD, CompHeightData, CompXYOffsetData);
+		::InternalUpscaling(DataInterface, LandscapeComponent, InLOD, GeometryLOD, CompHeightData);
 
 		TArray<uint16> RenderedWPOData;
 		if (bUseRenderedWPO)
@@ -541,10 +495,9 @@ void FLandscapeStaticLightingMesh::GetHeightmapData(int32 InLOD, int32 GeometryL
 			{
 				// Data array for upscaling case
 				TArray<FColor> CompHeightData;
-				TArray<FColor> CompXYOffsetData;
 				int32 NeighborGeometricLOD = ::GetLightingLOD(Neighbor);
 				FLandscapeComponentDataInterface DataInterface(Neighbor, InLOD);
-				::InternalUpscaling(DataInterface, Neighbor, InLOD, NeighborGeometricLOD, CompHeightData, CompXYOffsetData);
+				::InternalUpscaling(DataInterface, Neighbor, InLOD, NeighborGeometricLOD, CompHeightData);
 
 				TArray<uint16> RenderedWPOData;
 				if (bUseRenderedWPO)
diff --git a/Engine/Source/Runtime/Landscape/Private/LandscapeRender.cpp b/Engine/Source/Runtime/Landscape/Private/LandscapeRender.cpp
--- a/Engine/Source/Runtime/Landscape/Private/LandscapeRender.cpp
+++ b/Engine/Source/Runtime/Landscape/Private/LandscapeRender.cpp
@@ -1406,7 +1406,6 @@ FLandscapeComponentSceneProxy::FLandscapeComponentSceneProxy(ULandscapeComponent
 	, VisibilityWeightmapChannel(-1)
 	, HeightmapTexture(InComponent->GetHeightmap())
 	, HeightmapScaleBias(InComponent->HeightmapScaleBias)
-	, XYOffsetmapTexture(InComponent->XYOffsetmapTexture)
 	, SharedBuffersKey(0)
 	, SharedBuffers(nullptr)
 	, VertexFactory(nullptr)
@@ -1703,7 +1702,7 @@ FLandscapeComponentSceneProxy::FLandscapeComponentSceneProxy(ULandscapeComponent
 	Algo::Transform(AvailableMaterialInterfaces, MaterialRelevances, [ShaderPlatform](UMaterialInterface* InMaterialInterface) { check(InMaterialInterface != nullptr); return InMaterialInterface->GetRelevance_Concurrent(ShaderPlatform); });
 
 	const int8 SubsectionSizeLog2 = static_cast<int8>(FMath::CeilLogTwo(InComponent->SubsectionSizeQuads + 1));
-	SharedBuffersKey = (SubsectionSizeLog2 & 0xf) | ((NumSubsections & 0xf) << 4) | ((NumRayTracingSections & 0xf) << 8) | (XYOffsetmapTexture == nullptr ? 0 : 1 << 31);
+	SharedBuffersKey = (SubsectionSizeLog2 & 0xf) | ((NumSubsections & 0xf) << 4) | ((NumRayTracingSections & 0xf) << 8);
 
 	bSupportsHeightfieldRepresentation = true;
 
@@ -1772,20 +1771,10 @@ void FLandscapeComponentSceneProxy::CreateRenderThreadResources(FRHICommandListB
 
 		FLandscapeComponentSceneProxy::SharedBuffersMap.Add(SharedBuffersKey, SharedBuffers);
 
-		if (!XYOffsetmapTexture)
-		{
-			FLandscapeVertexFactory* LandscapeVertexFactory = new FLandscapeVertexFactory(FeatureLevel);
-			LandscapeVertexFactory->Data.PositionComponent = FVertexStreamComponent(SharedBuffers->VertexBuffer, 0, sizeof(FLandscapeVertex), VET_UByte4);
-			LandscapeVertexFactory->InitResource(RHICmdList);
-			SharedBuffers->VertexFactory = LandscapeVertexFactory;
-		}
-		else
-		{
-			FLandscapeXYOffsetVertexFactory* LandscapeXYOffsetVertexFactory = new FLandscapeXYOffsetVertexFactory(FeatureLevel);
-			LandscapeXYOffsetVertexFactory->Data.PositionComponent = FVertexStreamComponent(SharedBuffers->VertexBuffer, 0, sizeof(FLandscapeVertex), VET_UByte4);
-			LandscapeXYOffsetVertexFactory->InitResource(RHICmdList);
-			SharedBuffers->VertexFactory = LandscapeXYOffsetVertexFactory;
-		}
+		FLandscapeVertexFactory* LandscapeVertexFactory = new FLandscapeVertexFactory(FeatureLevel);
+		LandscapeVertexFactory->Data.PositionComponent = FVertexStreamComponent(SharedBuffers->VertexBuffer, 0, sizeof(FLandscapeVertex), VET_UByte4);
+		LandscapeVertexFactory->InitResource(RHICmdList);
+		SharedBuffers->VertexFactory = LandscapeVertexFactory;
 
 		bool bNeedsFixedGridVertexFactory = NeedsFixedGridVertexFactory(GetScene().GetShaderPlatform());
 
@@ -1794,11 +1783,10 @@ void FLandscapeComponentSceneProxy::CreateRenderThreadResources(FRHICommandListB
 
 		if (bNeedsFixedGridVertexFactory)
 		{
-			//todo[vt]: We will need a version of this to support XYOffsetmapTexture
-			FLandscapeFixedGridVertexFactory* LandscapeVertexFactory = new FLandscapeFixedGridVertexFactory(FeatureLevel);
-			LandscapeVertexFactory->Data.PositionComponent = FVertexStreamComponent(SharedBuffers->VertexBuffer, 0, sizeof(FLandscapeVertex), VET_UByte4);
-			LandscapeVertexFactory->InitResource(RHICmdList);
-			SharedBuffers->FixedGridVertexFactory = LandscapeVertexFactory;
+			FLandscapeFixedGridVertexFactory* FixedLandscapeVertexFactory = new FLandscapeFixedGridVertexFactory(FeatureLevel);
+			FixedLandscapeVertexFactory->Data.PositionComponent = FVertexStreamComponent(SharedBuffers->VertexBuffer, 0, sizeof(FLandscapeVertex), VET_UByte4);
+			FixedLandscapeVertexFactory->InitResource(RHICmdList);
+			SharedBuffers->FixedGridVertexFactory = FixedLandscapeVertexFactory;
 		}
 	}
 
@@ -1827,8 +1815,7 @@ void FLandscapeComponentSceneProxy::CreateRenderThreadResources(FRHICommandListB
 		GrassBatchParams.Empty(NumMips);
 		GrassBatchParams.AddDefaulted(NumMips);
 
-		// Grass is being generated using LOD0 material only
-		// It uses the fixed grid vertex factory so it doesn't support XY offsets
+		// Grass is being generated using LOD0 material only and uses the fixed grid vertex factory
 		FMaterialRenderProxy* RenderProxy = AvailableMaterials[LODIndexToMaterialIndex[0]];
 		GrassMeshBatch.VertexFactory = FixedGridVertexFactory;
 		GrassMeshBatch.MaterialRenderProxy = RenderProxy;
@@ -2391,18 +2378,6 @@ void FLandscapeComponentSceneProxy::OnTransformChanged(FRHICommandListBase& RHIC
 		LandscapeParams.NormalmapTextureSampler = GBlackTexture->SamplerStateRHI;
 	}
 
-	FTextureResource* XYOffsetmapResource = XYOffsetmapTexture ? XYOffsetmapTexture->GetResource() : nullptr;
-	if (XYOffsetmapResource)
-	{
-		LandscapeParams.XYOffsetmapTexture = XYOffsetmapTexture->TextureReference.TextureReferenceRHI;
-		LandscapeParams.XYOffsetmapTextureSampler = TStaticSamplerState<SF_Point>::GetRHI();
-	}
-	else
-	{
-		LandscapeParams.XYOffsetmapTexture = GBlackTexture->TextureRHI;
-		LandscapeParams.XYOffsetmapTextureSampler = GBlackTexture->SamplerStateRHI;
-	}
-
 	LandscapeUniformShaderParameters.SetContents(RHICmdList, LandscapeParams);
 
 	// Recache mesh draw commands for changed uniform buffers
@@ -3835,35 +3810,6 @@ void FLandscapeVertexFactory::Copy(const FLandscapeVertexFactory& Other)
 	BeginUpdateResourceRHI(this);
 }
 
-//
-// FLandscapeXYOffsetVertexFactory
-//
-
-void FLandscapeXYOffsetVertexFactory::ModifyCompilationEnvironment(const FVertexFactoryShaderPermutationParameters& Parameters, FShaderCompilerEnvironment& OutEnvironment)
-{
-	FLandscapeVertexFactory::ModifyCompilationEnvironment(Parameters, OutEnvironment);
-	OutEnvironment.SetDefine(TEXT("LANDSCAPE_XYOFFSET"), TEXT("1"));
-}
-
-IMPLEMENT_VERTEX_FACTORY_PARAMETER_TYPE(FLandscapeXYOffsetVertexFactory, SF_Vertex, FLandscapeVertexFactoryVertexShaderParameters);
-#if RHI_RAYTRACING
-IMPLEMENT_VERTEX_FACTORY_PARAMETER_TYPE(FLandscapeXYOffsetVertexFactory, SF_Compute, FLandscapeVertexFactoryVertexShaderParameters);
-IMPLEMENT_VERTEX_FACTORY_PARAMETER_TYPE(FLandscapeXYOffsetVertexFactory, SF_RayHitGroup, FLandscapeVertexFactoryVertexShaderParameters);
-#endif // RHI_RAYTRACING
-IMPLEMENT_VERTEX_FACTORY_PARAMETER_TYPE(FLandscapeXYOffsetVertexFactory, SF_Pixel, FLandscapeVertexFactoryPixelShaderParameters);
-
-IMPLEMENT_VERTEX_FACTORY_TYPE(FLandscapeXYOffsetVertexFactory, "/Engine/Private/LandscapeVertexFactory.ush",
-	  EVertexFactoryFlags::UsedWithMaterials
-	| EVertexFactoryFlags::SupportsStaticLighting
-	| EVertexFactoryFlags::SupportsDynamicLighting
-	| EVertexFactoryFlags::SupportsCachingMeshDrawCommands
-	| EVertexFactoryFlags::SupportsRayTracing
-	| EVertexFactoryFlags::SupportsRayTracingDynamicGeometry
-	| EVertexFactoryFlags::SupportsLightmapBaking
-	| EVertexFactoryFlags::SupportsPrimitiveIdStream
-	| EVertexFactoryFlags::SupportsLandscape
-);
-
 //
 // FLandscapeFixedGridVertexFactory
 //
@@ -4075,7 +4021,6 @@ public:
 			if (VertexFactoryType && VertexFactoryType->SupportsLandscape())
 			{
 				// For now only compile FLandscapeFixedGridVertexFactory for grass and runtime virtual texture page rendering (can change if we need for other cases)
-				// Todo: only compile LandscapeXYOffsetVertexFactory if we are using it
 				const bool bIsGrassShaderType = Algo::Find(GetGrassShaderTypes(), ShaderType->GetFName()) != nullptr;
 				const bool bIsGPULightmassShaderType = Algo::Find(GetGPULightmassShaderTypes(), ShaderType->GetFName()) != nullptr;
 				const bool bIsRuntimeVirtualTextureShaderType = Algo::Find(GetRuntimeVirtualTextureShaderTypes(), ShaderType->GetFName()) != nullptr;
@@ -4457,15 +4402,6 @@ void ULandscapeComponent::GetStreamingRenderAssetInfo(FStreamingTextureLevelCont
 		StreamingHeightmap.RenderAsset = HeightmapTexture;
 	}
 
-	// XYOffset has not been accounted for by GetUsedTextures on the material :
-	if (XYOffsetmapTexture)
-	{
-		FStreamingRenderAssetPrimitiveInfo& StreamingXYOffset = *new(OutStreamingRenderAssets)FStreamingRenderAssetPrimitiveInfo;
-		StreamingXYOffset.Bounds = BoundingSphere;
-		StreamingXYOffset.TexelFactor = TexelFactor;
-		StreamingXYOffset.RenderAsset = XYOffsetmapTexture;
-	}
-
 #if WITH_EDITOR
 	if (GIsEditor)
 	{
@@ -4613,9 +4549,6 @@ float FLandscapeComponentSceneProxy::ComputeLODBias() const
 		}
 	}
 
-	// TODO: support mipmap LOD bias of XY offset map
-	//XYOffsetmapTexture ? ((FTexture2DResource*)XYOffsetmapTexture->Resource)->GetCurrentFirstMip() : 0.0f);
-
 	return ComputedLODBias;
 }
 
diff --git a/Engine/Source/Runtime/Landscape/Public/LandscapeDataAccess.h b/Engine/Source/Runtime/Landscape/Public/LandscapeDataAccess.h
--- a/Engine/Source/Runtime/Landscape/Public/LandscapeDataAccess.h
+++ b/Engine/Source/Runtime/Landscape/Public/LandscapeDataAccess.h
@@ -13,8 +13,8 @@ LandscapeDataAccess.h: Classes for the editor to access to Landscape data
 #define LANDSCAPE_ZSCALE		(1.0f/128.0f)
 #define LANDSCAPE_INV_ZSCALE	128.0f
 
-#define LANDSCAPE_XYOFFSET_SCALE	(1.0f/256.f)
-#define LANDSCAPE_INV_XYOFFSET_SCALE	256.f
+#define LANDSCAPE_XYOFFSET_SCALE UE_DEPRECATED_MACRO(5.7, "RetopologizeTool/XYOffset deprecated with the removal of non-edit layer landscapes") (1.0f/256.f)
+#define LANDSCAPE_INV_XYOFFSET_SCALE UE_DEPRECATED_MACRO(5.7, "RetopologizeTool/XYOffset deprecated with the removal of non-edit layer landscapes") 256.f
 
 #define LANDSCAPE_VISIBILITY_THRESHOLD (2.0f/3.0f)
 
@@ -279,21 +279,11 @@ struct FLandscapeComponentDataInterface : public FLandscapeComponentDataInterfac
 		return HeightMipData;
 	}
 
-	FColor* GetRawXYOffsetData() const
-	{
-		return XYOffsetMipData;
-	}
-
 	void SetRawHeightData(FColor* NewHeightData)
 	{
 		HeightMipData = NewHeightData;
 	}
 
-	void SetRawXYOffsetData(FColor* NewXYOffsetData)
-	{
-		XYOffsetMipData = NewXYOffsetData;
-	}
-
 	LANDSCAPE_API int32 GetHeightmapSizeX(int32 MipIndex) const;
 	LANDSCAPE_API int32 GetHeightmapSizeY(int32 MipIndex) const;
 
@@ -317,6 +307,7 @@ struct FLandscapeComponentDataInterface : public FLandscapeComponentDataInterfac
 		return &HeightMipData[TexelX + HeightmapComponentOffsetX + (TexelY + HeightmapComponentOffsetY) * HeightmapStride];
 	}
 
+	UE_DEPRECATED(5.7, "RetopologizeTool/XYOffset deprecated with the removal of non-edit layer landscapes")
 	LANDSCAPE_API FColor* GetXYOffsetData(int32 LocalX, int32 LocalY) const;
 
 	uint16 GetHeight( int32 LocalX, int32 LocalY ) const
@@ -342,27 +333,6 @@ struct FLandscapeComponentDataInterface : public FLandscapeComponentDataInterfac
 		return LandscapeDataAccess::GetLocalHeight(GetHeight(VertexIndex));
 	}
 
-	void GetXYOffset( int32 X, int32 Y, float& XOffset, float& YOffset ) const
-	{
-		if (XYOffsetMipData)
-		{
-			FColor* Texel = GetXYOffsetData(X, Y);
-			XOffset = ((float)((Texel->R << 8) + Texel->G) - 32768.f) * LANDSCAPE_XYOFFSET_SCALE;
-			YOffset = ((float)((Texel->B << 8) + Texel->A) - 32768.f) * LANDSCAPE_XYOFFSET_SCALE;
-		}
-		else
-		{
-			XOffset = YOffset = 0.f;
-		}
-	}
-
-	void GetXYOffset( int32 VertexIndex, float& XOffset, float& YOffset ) const
-	{
-		int32 X, Y;
-		VertexIndexToXY( VertexIndex, X, Y );
-		GetXYOffset( X, Y, XOffset, YOffset );
-	}
-
 	LANDSCAPE_API FVector GetLocalVertex(int32 LocalX, int32 LocalY) const;
 
 	void GetLocalTangentVectors( int32 LocalX, int32 LocalY, FVector& LocalTangentX, FVector& LocalTangentY, FVector& LocalTangentZ ) const
@@ -421,8 +391,6 @@ private:
 	ULandscapeComponent* Component = nullptr;
 	bool bWorkOnEditingLayer = false;
 	FColor* HeightMipData = nullptr;
-	FColor* XYOffsetMipData = nullptr;
-	
 };
 
 // Helper functions
diff --git a/Engine/Source/Runtime/Landscape/Public/LandscapeEdit.h b/Engine/Source/Runtime/Landscape/Public/LandscapeEdit.h
--- a/Engine/Source/Runtime/Landscape/Public/LandscapeEdit.h
+++ b/Engine/Source/Runtime/Landscape/Public/LandscapeEdit.h
@@ -206,7 +206,7 @@ struct FLandscapeEditDataInterface : public FLandscapeTextureDataInterface
 	LANDSCAPE_API void GetHeightData(int32& X1, int32& Y1, int32& X2, int32& Y2, TMap<FIntPoint, uint16>& SparseData);
 	LANDSCAPE_API void GetHeightDataFast(const int32 X1, const int32 Y1, const int32 X2, const int32 Y2, TMap<FIntPoint, uint16>& SparseData, TMap<FIntPoint, uint16>* NormalData = NULL, UTexture2D* InHeightmap = nullptr);
 
-	// Recaclulate normals for the entire landscape.
+	UE_DEPRECATED(5.7, "XYOffset deprecated with the removal of non-edit layer landscapes")
 	LANDSCAPE_API void RecalculateNormals();
 
 	//
@@ -267,29 +267,31 @@ struct FLandscapeEditDataInterface : public FLandscapeTextureDataInterface
 	LANDSCAPE_API void GetDirtyData(const int32 X1, const int32 Y1, const int32 X2, const int32 Y2, TMap<FIntPoint, uint8>& SparseData);
 	LANDSCAPE_API void SetDirtyData(const int32 X1, const int32 Y1, const int32 X2, const int32 Y2, const uint8* Data, int32 Stride);
 
-	//
-	// XYOffsetmap access
-	//
-	template<typename T>
-	void SetXYOffsetDataTempl(int32 X1, int32 Y1, int32 X2, int32 Y2, const T* Data, int32 Stride);
+	UE_DEPRECATED(5.7, "XYOffset deprecated with the removal of non-edit layer landscapes")
 	LANDSCAPE_API void SetXYOffsetData(int32 X1, int32 Y1, int32 X2, int32 Y2, const FVector2D* Data, int32 Stride);
+	UE_DEPRECATED(5.7, "XYOffset deprecated with the removal of non-edit layer landscapes")
 	LANDSCAPE_API void SetXYOffsetData(int32 X1, int32 Y1, int32 X2, int32 Y2, const FVector* Data, int32 Stride);
 	// Helper accessor
+	UE_DEPRECATED(5.7, "XYOffset deprecated with the removal of non-edit layer landscapes")
 	LANDSCAPE_API FORCEINLINE FVector2D GetXYOffsetmapData(const ULandscapeComponent* Component, int32 TexU, int32 TexV, FColor* TextureData = NULL);
 
-	template<typename TStoreData>
-	void GetXYOffsetDataTempl(int32& X1, int32& Y1, int32& X2, int32& Y2, TStoreData& StoreData);
+	UE_DEPRECATED(5.7, "XYOffset deprecated with the removal of non-edit layer landscapes")
 	LANDSCAPE_API void GetXYOffsetData(int32& X1, int32& Y1, int32& X2, int32& Y2, FVector2D* Data, int32 Stride);
+	UE_DEPRECATED(5.7, "XYOffset deprecated with the removal of non-edit layer landscapes")
 	LANDSCAPE_API void GetXYOffsetData(int32& X1, int32& Y1, int32& X2, int32& Y2, TMap<FIntPoint, FVector2D>& SparseData);
+	UE_DEPRECATED(5.7, "XYOffset deprecated with the removal of non-edit layer landscapes")
 	LANDSCAPE_API void GetXYOffsetData(int32& X1, int32& Y1, int32& X2, int32& Y2, FVector* Data, int32 Stride);
+	UE_DEPRECATED(5.7, "XYOffset deprecated with the removal of non-edit layer landscapes")
 	LANDSCAPE_API void GetXYOffsetData(int32& X1, int32& Y1, int32& X2, int32& Y2, TMap<FIntPoint, FVector>& SparseData);
-	// Without data interpolation
-	template<typename TStoreData>
-	void GetXYOffsetDataTemplFast(const int32 X1, const int32 Y1, const int32 X2, const int32 Y2, TStoreData& StoreData);
+
 	// Without data interpolation, able to get normal data
+	UE_DEPRECATED(5.7, "XYOffset deprecated with the removal of non-edit layer landscapes")
 	LANDSCAPE_API void GetXYOffsetDataFast(const int32 X1, const int32 Y1, const int32 X2, const int32 Y2, FVector2D* Data, int32 Stride);
+	UE_DEPRECATED(5.7, "XYOffset deprecated with the removal of non-edit layer landscapes")
 	LANDSCAPE_API void GetXYOffsetDataFast(const int32 X1, const int32 Y1, const int32 X2, const int32 Y2, TMap<FIntPoint, FVector2D>& SparseData);
+	UE_DEPRECATED(5.7, "XYOffset deprecated with the removal of non-edit layer landscapes")
 	LANDSCAPE_API void GetXYOffsetDataFast(const int32 X1, const int32 Y1, const int32 X2, const int32 Y2, FVector* Data, int32 Stride);
+	UE_DEPRECATED(5.7, "XYOffset deprecated with the removal of non-edit layer landscapes")
 	LANDSCAPE_API void GetXYOffsetDataFast(const int32 X1, const int32 Y1, const int32 X2, const int32 Y2, TMap<FIntPoint, FVector>& SparseData);
 
 	template<typename T>
diff --git a/Engine/Source/Runtime/Landscape/Public/LandscapeLight.h b/Engine/Source/Runtime/Landscape/Public/LandscapeLight.h
--- a/Engine/Source/Runtime/Landscape/Public/LandscapeLight.h
+++ b/Engine/Source/Runtime/Landscape/Public/LandscapeLight.h
@@ -135,6 +135,7 @@ protected:
 public:
 	// Cache data for Landscape upscaling data
 	LANDSCAPE_API static TMap<FIntPoint, FColor> LandscapeUpscaleHeightDataCache;
+	UE_DEPRECATED(5.7, "RetopologizeTool/XYOffset deprecated with the removal of non-edit layer landscapes")
 	LANDSCAPE_API static TMap<FIntPoint, FColor> LandscapeUpscaleXYOffsetDataCache;
 #endif
 };
diff --git a/Engine/Source/Runtime/Landscape/Public/LandscapeRender.h b/Engine/Source/Runtime/Landscape/Public/LandscapeRender.h
--- a/Engine/Source/Runtime/Landscape/Public/LandscapeRender.h
+++ b/Engine/Source/Runtime/Landscape/Public/LandscapeRender.h
@@ -137,8 +137,6 @@ BEGIN_GLOBAL_SHADER_PARAMETER_STRUCT(FLandscapeUniformShaderParameters, LANDSCAP
 	SHADER_PARAMETER_SAMPLER(SamplerState, HeightmapTextureSampler)
 	SHADER_PARAMETER_TEXTURE(Texture2D, NormalmapTexture)
 	SHADER_PARAMETER_SAMPLER(SamplerState, NormalmapTextureSampler)
-	SHADER_PARAMETER_TEXTURE(Texture2D, XYOffsetmapTexture)
-	SHADER_PARAMETER_SAMPLER(SamplerState, XYOffsetmapTextureSampler)
 END_GLOBAL_SHADER_PARAMETER_STRUCT()
 
 BEGIN_GLOBAL_SHADER_PARAMETER_STRUCT(FLandscapeVertexFactoryMVFParameters, LANDSCAPE_API)
@@ -285,7 +283,7 @@ public:
 
 
 /** vertex factory for VTF-heightmap terrain  */
-class FLandscapeXYOffsetVertexFactory : public FLandscapeVertexFactory
+class UE_DEPRECATED(5.7, "XYOffset deprecated with the removal of non-edit layer landscapes") FLandscapeXYOffsetVertexFactory : public FLandscapeVertexFactory
 {
 	DECLARE_VERTEX_FACTORY_TYPE_API(FLandscapeXYOffsetVertexFactory, LANDSCAPE_API);
 
@@ -294,10 +292,6 @@ public:
 		: FLandscapeVertexFactory(InFeatureLevel)
 	{
 	}
-
-	virtual ~FLandscapeXYOffsetVertexFactory() {}
-
-	static void ModifyCompilationEnvironment(const FVertexFactoryShaderPermutationParameters& Parameters, FShaderCompilerEnvironment& OutEnvironment);
 };
 
 
@@ -829,8 +823,6 @@ protected:
 	float HeightmapSubsectionOffsetU;
 	float HeightmapSubsectionOffsetV;
 
-	UTexture2D* XYOffsetmapTexture;
-
 	uint32 SharedBuffersKey;
 	FLandscapeSharedBuffers* SharedBuffers;
 	FLandscapeVertexFactory* VertexFactory;
