Unshelved from pending changelist '45543337':

Landscape Editor - Fixes WP create new landscape bugs by protecting properties ULandscapeInfo::StreamingProxies and LandscapeProxy::LandscapeSectionOffset
* This fixes a regression introduced by 45383226. Before the change, we were simply lucky that PostRegisterAllComponents was called by the new actor label being set when the WP landscape creation process was spawning a new landscape streaming proxy. This would lead FixupProxiesTransform to be called, which would assign the proper transform to all the (new) proxies from the landscape. Without this, the components, when attached to the newly created proxy, would get their SectionBase refreshed with improper values (i.e. not aligned with the number of quads on the component).
* Anytime a proxy section offset is updated we re-sort all StreamingProxies

diff --git a/Engine/Config/BaseEngine.ini b/Engine/Config/BaseEngine.ini
--- a/Engine/Config/BaseEngine.ini
+++ b/Engine/Config/BaseEngine.ini
@@ -1363,6 +1363,7 @@ ClassName=/Script/OnlineSubsystemUtils.OnlineEngineInterfaceImpl
 
 ; 5.7
 +PropertyRedirects=(OldName="MeshNaniteSettings.bPreserveArea",NewName="ShapePreservation")
++PropertyRedirects=(OldName="LandscapeProxy.LandscapeSectionOffset",NewName="SectionBase")
 
 ; Movie Scene Text Track move to Movie Scene Tracks
 +StructRedirects=(OldName="/Script/MovieSceneTextTrack.MovieSceneTextChannel",NewName="/Script/MovieSceneTracks.MovieSceneTextChannel")
diff --git a/Engine/Plugins/Enterprise/GLTFExporter/Source/GLTFExporter/Private/Tasks/GLTFDelayedMeshTasks.cpp b/Engine/Plugins/Enterprise/GLTFExporter/Source/GLTFExporter/Private/Tasks/GLTFDelayedMeshTasks.cpp
--- a/Engine/Plugins/Enterprise/GLTFExporter/Source/GLTFExporter/Private/Tasks/GLTFDelayedMeshTasks.cpp
+++ b/Engine/Plugins/Enterprise/GLTFExporter/Source/GLTFExporter/Private/Tasks/GLTFDelayedMeshTasks.cpp
@@ -701,8 +701,8 @@ void FGLTFDelayedLandscapeTask::Process()
 	TArray<uint8> VisibilityData;
 	VisibilityData.SetNumZeroed(VertexCount);
 
-	int OffsetX = Landscape->LandscapeSectionOffset.X;
-	int OffsetY = Landscape->LandscapeSectionOffset.Y;
+	int OffsetX = Landscape->GetSectionBase().X;
+	int OffsetY = Landscape->GetSectionBase().Y;
 	
 	FGLTFLandscapeComponentDataInterface CDI(LandscapeComponent, ExportLOD);
 
diff --git a/Engine/Source/Developer/MaterialUtilities/Private/MaterialUtilities.cpp b/Engine/Source/Developer/MaterialUtilities/Private/MaterialUtilities.cpp
--- a/Engine/Source/Developer/MaterialUtilities/Private/MaterialUtilities.cpp
+++ b/Engine/Source/Developer/MaterialUtilities/Private/MaterialUtilities.cpp
@@ -1582,7 +1582,7 @@ bool FMaterialUtilities::ExportBaseColor(ULandscapeComponent* LandscapeComponent
 {
 	ALandscapeProxy* LandscapeProxy = LandscapeComponent->GetLandscapeProxy();
 
-	FIntPoint ComponentOrigin = LandscapeComponent->GetSectionBase() - LandscapeProxy->LandscapeSectionOffset;
+	FIntPoint ComponentOrigin = LandscapeComponent->GetSectionBase() - LandscapeProxy->GetSectionBase();
 	FIntPoint ComponentSize(LandscapeComponent->ComponentSizeQuads, LandscapeComponent->ComponentSizeQuads);
 	FVector MidPoint = FVector(ComponentOrigin, 0.f) + FVector(ComponentSize, 0.f)*0.5f;
 
diff --git a/Engine/Source/Editor/LandscapeEditor/Private/LandscapeEdMode.cpp b/Engine/Source/Editor/LandscapeEditor/Private/LandscapeEdMode.cpp
--- a/Engine/Source/Editor/LandscapeEditor/Private/LandscapeEdMode.cpp
+++ b/Engine/Source/Editor/LandscapeEditor/Private/LandscapeEdMode.cpp
@@ -4302,7 +4302,7 @@ void FEdModeLandscape::OnLevelActorRemoved(AActor* InActor)
 
 bool LandscapeEditorUtils::SetHeightmapData(ALandscapeProxy* Landscape, const TArray<uint16>& Data)
 {
-	FIntRect ComponentsRect = Landscape->GetBoundingRect() + Landscape->LandscapeSectionOffset;
+	FIntRect ComponentsRect = Landscape->GetBoundingRect() + Landscape->GetSectionBase();
 
 	if (Data.Num() == (1 + ComponentsRect.Width())*(1 + ComponentsRect.Height()))
 	{
@@ -4316,7 +4316,7 @@ bool LandscapeEditorUtils::SetHeightmapData(ALandscapeProxy* Landscape, const TA
 
 bool LandscapeEditorUtils::SetWeightmapData(ALandscapeProxy* Landscape, ULandscapeLayerInfoObject* LayerObject, const TArray<uint8>& Data)
 {
-	FIntRect ComponentsRect = Landscape->GetBoundingRect() + Landscape->LandscapeSectionOffset;
+	FIntRect ComponentsRect = Landscape->GetBoundingRect() + Landscape->GetSectionBase();
 
 	if (Data.Num() == (1 + ComponentsRect.Width())*(1 + ComponentsRect.Height()))
 	{
diff --git a/Engine/Source/Editor/LandscapeEditor/Private/LandscapeProxyUIDetails.cpp b/Engine/Source/Editor/LandscapeEditor/Private/LandscapeProxyUIDetails.cpp
--- a/Engine/Source/Editor/LandscapeEditor/Private/LandscapeProxyUIDetails.cpp
+++ b/Engine/Source/Editor/LandscapeEditor/Private/LandscapeProxyUIDetails.cpp
@@ -200,7 +200,7 @@ void FLandscapeProxyUIDetails::CustomizeDetails( IDetailLayoutBuilder& DetailBui
 				GenerateTextWidget(FText::Format(LOCTEXT("LandscapeResolutionValue", "{0} x {1}"), Size.X+1, Size.Y+1), true)
 			];
 
-			int32 LandscapeCount = LandscapeInfo->StreamingProxies.Num() + (LandscapeInfo->LandscapeActor.Get() ? 1 : 0);
+			int32 LandscapeCount = LandscapeInfo->GetSortedStreamingProxies().Num() + (LandscapeInfo->LandscapeActor.Get() ? 1 : 0);
 			RowDisplayText = LOCTEXT("LandscapeCount", "Landscape Proxy Count");
 			CategoryBuilder.AddCustomRow(RowDisplayText)
 			.RowTag(TEXT("LandscapeCount"))
@@ -365,7 +365,7 @@ void FLandscapeProxyUIDetails::CustomizeDetails( IDetailLayoutBuilder& DetailBui
 						ULandscapeInfo* LandscapeInfo = Landscape->GetLandscapeInfo();
 						if (LandscapeInfo != nullptr)
 						{
-							Algo::Transform(LandscapeInfo->StreamingProxies, ProxiesToBuild, [](const TWeakObjectPtr<ALandscapeStreamingProxy>& InStreamingProxy) { return InStreamingProxy; });
+							Algo::Transform(LandscapeInfo->GetSortedStreamingProxies(), ProxiesToBuild, [](const TWeakObjectPtr<ALandscapeStreamingProxy>& InStreamingProxy) { return InStreamingProxy; });
 						}
 					}
 				}
diff --git a/Engine/Source/Editor/WorldBrowser/Private/Tiles/WorldTileModel.cpp b/Engine/Source/Editor/WorldBrowser/Private/Tiles/WorldTileModel.cpp
--- a/Engine/Source/Editor/WorldBrowser/Private/Tiles/WorldTileModel.cpp
+++ b/Engine/Source/Editor/WorldBrowser/Private/Tiles/WorldTileModel.cpp
@@ -533,7 +533,7 @@ void FWorldTileModel::UpdateLandscapeSectionsOffset(FIntPoint LevelOffset)
 		FIntPoint QuadsSpaceOffset;
 		QuadsSpaceOffset.X = FMath::RoundToInt(LevelOffset.X / DrawScale.X);
 		QuadsSpaceOffset.Y = FMath::RoundToInt(LevelOffset.Y / DrawScale.Y);
-		LandscapeProxy->SetAbsoluteSectionBase(QuadsSpaceOffset + LandscapeProxy->LandscapeSectionOffset);
+		LandscapeProxy->SetSectionBase(QuadsSpaceOffset + LandscapeProxy->GetSectionBase());
 	}
 }
 
@@ -950,8 +950,8 @@ bool FWorldTileModel::CreateAdjacentLandscapeProxy(ALandscapeProxy* SourceLandsc
 	FVector SourceLandscapeScale = SourceLandscape->GetRootComponent()->GetComponentToWorld().GetScale3D();
 	FIntRect SourceLandscapeRect = SourceLandscape->GetBoundingRect();
 	FIntPoint SourceLandscapeSize = SourceLandscapeRect.Size();
-	FIntPoint LandscapeSectionOffset = FIntPoint((SourceLandscape->LandscapeSectionOffset.X+SourceLandscapeRect.Min.X)*SourceLandscapeScale.X, 
-												 (SourceLandscape->LandscapeSectionOffset.Y+SourceLandscapeRect.Min.Y)*SourceLandscapeScale.Y);
+	FIntPoint LandscapeSectionOffset = FIntPoint((SourceLandscape->GetSectionBase().X+SourceLandscapeRect.Min.X)*SourceLandscapeScale.X,
+												 (SourceLandscape->GetSectionBase().Y+SourceLandscapeRect.Min.Y)*SourceLandscapeScale.Y);
 
 	FLandscapeImportSettings ImportSettings = {};
 	ImportSettings.LandscapeGuid		= SourceLandscape->GetLandscapeGuid();
diff --git a/Engine/Source/Runtime/Engine/Private/WorldPartition/Landscape/LandscapeActorDesc.cpp b/Engine/Source/Runtime/Engine/Private/WorldPartition/Landscape/LandscapeActorDesc.cpp
--- a/Engine/Source/Runtime/Engine/Private/WorldPartition/Landscape/LandscapeActorDesc.cpp
+++ b/Engine/Source/Runtime/Engine/Private/WorldPartition/Landscape/LandscapeActorDesc.cpp
@@ -17,7 +17,7 @@ void FLandscapeActorDesc::Init(const AActor* InActor)
 		const ALandscapeProxy* LandscapeProxy = CastChecked<ALandscapeProxy>(InActor);
 		check(LandscapeProxy);
 
-		SetGridIndices(LandscapeProxy->LandscapeSectionOffset.X, LandscapeProxy->LandscapeSectionOffset.Y, 0);
+		SetGridIndices(LandscapeProxy->GetSectionBase().X, LandscapeProxy->GetSectionBase().Y, 0);
 
 		const ALandscape* LandscapeActor = LandscapeProxy->GetLandscapeActor();
 		if (LandscapeActor)
diff --git a/Engine/Source/Runtime/Landscape/Classes/LandscapeInfo.h b/Engine/Source/Runtime/Landscape/Classes/LandscapeInfo.h
--- a/Engine/Source/Runtime/Landscape/Classes/LandscapeInfo.h
+++ b/Engine/Source/Runtime/Landscape/Classes/LandscapeInfo.h
@@ -109,6 +109,8 @@ class ULandscapeInfo : public UObject
 {
 	GENERATED_UCLASS_BODY()
 
+	friend class ALandscapeProxy;
+
 	UPROPERTY()
 	TWeakObjectPtr<ALandscape> LandscapeActor;
 
@@ -151,8 +153,9 @@ public:
 	TMap<FIntPoint, FLandscapeAddCollision> XYtoAddCollisionMap;
 #endif // WITH_EDITORONLY_DATA
 
+	UE_DEPRECATED(5.7, "StreamingProxies rely on a sorted state and should not be modified outside of ULandscapeInfo. Use GetSortedStreamingProxies instead.")
 	UPROPERTY()
-	TArray<TWeakObjectPtr<ALandscapeStreamingProxy>> StreamingProxies;
+	TArray<TWeakObjectPtr<ALandscapeStreamingProxy>> StreamingProxies_DEPRECATED;
 
 private:
 #if WITH_EDITORONLY_DATA
@@ -379,8 +382,11 @@ public:
 	/** Called after creating object so that it can initialize its state */
 	void Initialize(UWorld* InWorld, const FGuid& InLandscapeGuid);
 
+	/** Returns a const reference to the sorted registered streaming proxies */
+	LANDSCAPE_API const TArray<TWeakObjectPtr<ALandscapeStreamingProxy>>& GetSortedStreamingProxies() const;
+
 	/**
-	 * Runs the given function on the root landscape actor and all streaming proxies, with the posibility of early exit
+	 * Runs the given function on the root landscape actor and all streaming proxies, with the possibility of early exit
 	 * Most easily used with a lambda as follows:
 	 * ForEachLandscapeProxy([](ALandscapeProxy* Proxy) -> bool
 	 * {
@@ -464,4 +470,10 @@ private:
 
 private:
 	void RegisterLandscapeActorWithProxyInternal(ALandscapeProxy* Proxy, bool bMapCheck);
+
+	/** Sorts registered streaming proxies by their section base offset */
+	void SortStreamingProxies();
+
+	UPROPERTY()
+	TArray<TWeakObjectPtr<ALandscapeStreamingProxy>> SortedStreamingProxies;
 };
diff --git a/Engine/Source/Runtime/Landscape/Classes/LandscapeProxy.h b/Engine/Source/Runtime/Landscape/Classes/LandscapeProxy.h
--- a/Engine/Source/Runtime/Landscape/Classes/LandscapeProxy.h
+++ b/Engine/Source/Runtime/Landscape/Classes/LandscapeProxy.h
@@ -499,10 +499,10 @@ public:
 #if WITH_EDITORONLY_DATA
 	UPROPERTY()
 	TMap<FName, TObjectPtr<ULandscapeLayerInfoObject>> TargetLayersForFixup;
-#endif //WITH_EDITORONLY_DATA
-	/** Offset in quads from global components grid origin (in quads) **/
-	UPROPERTY()
+
+	UE_DEPRECATED(5.7, "Use private SectionBase with GetSectionBase instead.")
 	FIntPoint LandscapeSectionOffset;
+#endif //WITH_EDITORONLY_DATA
 
 	/** Max LOD level to use when rendering, -1 means the max available */
 	UPROPERTY(EditAnywhere, Category=LOD, meta = (LandscapeInherited))
@@ -1247,9 +1247,12 @@ PRAGMA_ENABLE_DEPRECATION_WARNINGS
 	UFUNCTION(BlueprintCallable, meta = (DisplayName = "Landscape Export Heightmap to RenderTarget", Keywords = "Push Landscape Heightmap to RenderTarget", UnsafeDuringActorConstruction = "true"), Category = Rendering)
 	bool LandscapeExportHeightmapToRenderTarget(UTextureRenderTarget2D* InRenderTarget, bool InExportHeightIntoRGChannel = false, bool InExportLandscapeProxies = true);
 
-	/** Get landscape position in section space */
+	UE_DEPRECATED(5.7, "Use GetSectionBase instead.")
 	LANDSCAPE_API FIntPoint GetSectionBaseOffset() const;
 
+	/** Get landscape position in section space */
+	LANDSCAPE_API FIntPoint GetSectionBase() const;
+
 	// ILandscapeSplineInterface
 	virtual ULandscapeSplinesComponent* GetSplinesComponent() const override { return SplineComponent; }
 	LANDSCAPE_API virtual void UpdateSharedProperties(ULandscapeInfo* InLandscapeInfo) override;
@@ -1370,9 +1373,12 @@ PRAGMA_ENABLE_DEPRECATION_WARNINGS
 	// Assign only mismatching data and mark proxy package dirty
 	LANDSCAPE_API void FixupSharedData(ALandscape* Landscape, bool bMapCheck = false);
 
-	/** Set landscape absolute location in section space */
+	UE_DEPRECATED(5.7, "Use SetSectionBase instead.")
 	LANDSCAPE_API void SetAbsoluteSectionBase(FIntPoint SectionOffset);
 
+	/** Set landscape location in section space */
+	LANDSCAPE_API void SetSectionBase(FIntPoint SectionOffset);
+
 	/** Recreate all components rendering and collision states */
 	LANDSCAPE_API void RecreateComponentsState();
 
@@ -1701,10 +1707,13 @@ public:
 private:
 	/** Maintain list of Proxies for faster iteration */
 	static TArray<ALandscapeProxy*> LandscapeProxies;
-
 	FAsyncWorkMonitor AsyncWorkMonitor;
 #endif // WITH_EDITORONLY_DATA
 
+	/** Offset in quads from global components grid origin (in quads) **/
+	UPROPERTY()
+	FIntPoint SectionBase = FIntPoint::ZeroValue;
+
 #if WITH_EDITOR
 	static constexpr TCHAR const* LandscapeInheritedTag = TEXT("LandscapeInherited");
 	static constexpr TCHAR const* LandscapeOverridableTag = TEXT("LandscapeOverridable");
diff --git a/Engine/Source/Runtime/Landscape/Private/Landscape.cpp b/Engine/Source/Runtime/Landscape/Private/Landscape.cpp
--- a/Engine/Source/Runtime/Landscape/Private/Landscape.cpp
+++ b/Engine/Source/Runtime/Landscape/Private/Landscape.cpp
@@ -6,6 +6,7 @@ Landscape.cpp: Terrain rendering
 
 #include "Landscape.h"
 
+#include "Algo/IsSorted.h"
 #include "Serialization/MemoryWriter.h"
 #include "Serialization/BufferArchive.h"
 #include "Serialization/MemoryReader.h"
@@ -1282,8 +1283,8 @@ void ULandscapeComponent::PostLoad()
 
 		// check SectionBaseX/Y are correct
 		const FVector LocalRelativeLocation = GetRelativeLocation();
-		int32 CheckSectionBaseX = FMath::RoundToInt32(LocalRelativeLocation.X) + LandscapeProxy->LandscapeSectionOffset.X;
-		int32 CheckSectionBaseY = FMath::RoundToInt32(LocalRelativeLocation.Y) + LandscapeProxy->LandscapeSectionOffset.Y;
+		int32 CheckSectionBaseX = FMath::RoundToInt32(LocalRelativeLocation.X) + LandscapeProxy->GetSectionBase().X;
+		int32 CheckSectionBaseY = FMath::RoundToInt32(LocalRelativeLocation.Y) + LandscapeProxy->GetSectionBase().Y;
 		if (CheckSectionBaseX != SectionBaseX ||
 			CheckSectionBaseY != SectionBaseY)
 		{
@@ -1294,12 +1295,12 @@ void ULandscapeComponent::PostLoad()
 		}
 	}
 
-	if (GIsEditor && !HasAnyFlags(RF_ClassDefaultObject))
+	if (GIsEditor && !HasAnyFlags(RF_ClassDefaultObject) && ensure(LandscapeProxy))
 	{
 		// This is to ensure that component relative location is exact section base offset value
 		FVector LocalRelativeLocation = GetRelativeLocation();
-		float CheckRelativeLocationX = float(SectionBaseX - LandscapeProxy->LandscapeSectionOffset.X);
-		float CheckRelativeLocationY = float(SectionBaseY - LandscapeProxy->LandscapeSectionOffset.Y);
+		float CheckRelativeLocationX = float(SectionBaseX - LandscapeProxy->GetSectionBase().X);
+		float CheckRelativeLocationY = float(SectionBaseY - LandscapeProxy->GetSectionBase().Y);
 		if (!FMath::IsNearlyEqual(CheckRelativeLocationX, LocalRelativeLocation.X, UE_DOUBLE_KINDA_SMALL_NUMBER) ||
 			!FMath::IsNearlyEqual(CheckRelativeLocationY, LocalRelativeLocation.Y, UE_DOUBLE_KINDA_SMALL_NUMBER))
 		{
@@ -1630,7 +1631,6 @@ ALandscapeProxy::ALandscapeProxy(const FObjectInitializer& ObjectInitializer)
 
 	RootComponent->SetRelativeScale3D(FVector(128.0f, 128.0f, 256.0f)); // Old default scale, preserved for compatibility. See ULandscapeEditorObject::NewLandscape_Scale
 	RootComponent->Mobility = EComponentMobility::Static;
-	LandscapeSectionOffset = FIntPoint::ZeroValue;
 
 	StaticLightingResolution = 1.0f;
 	StreamingDistanceMultiplier = 1.0f;
@@ -1809,7 +1809,7 @@ FTransform ALandscapeProxy::LandscapeActorToWorld() const
 {
 	FTransform TM = ActorToWorld();
 	// Add this proxy landscape section offset to obtain landscape actor transform
-	TM.AddToTranslation(TM.TransformVector(-FVector(LandscapeSectionOffset)));
+	TM.AddToTranslation(TM.TransformVector(-FVector(SectionBase)));
 	return TM;
 }
 
@@ -4686,12 +4686,44 @@ PRAGMA_ENABLE_DEPRECATION_WARNINGS
 #endif // WITH_EDITOR
 }
 
+// Deprecated
 FIntPoint ALandscapeProxy::GetSectionBaseOffset() const
 {
-	return LandscapeSectionOffset;
+	return SectionBase;
+}
+
+FIntPoint ALandscapeProxy::GetSectionBase() const
+{
+	return SectionBase;
 }
 
 #if WITH_EDITOR
+void ALandscapeProxy::SetSectionBase(FIntPoint InSectionBase)
+{
+	FIntPoint Difference = InSectionBase - SectionBase;
+	SectionBase = InSectionBase;
+
+	RecreateComponentsRenderState([Difference](ULandscapeComponent* Comp)
+		{
+			FIntPoint AbsoluteSectionBase = Comp->GetSectionBase() + Difference;
+			Comp->SetSectionBase(AbsoluteSectionBase);
+		});
+
+	for (int32 CompIdx = 0; CompIdx < CollisionComponents.Num(); CompIdx++)
+	{
+		ULandscapeHeightfieldCollisionComponent* Comp = CollisionComponents[CompIdx];
+		if (Comp)
+		{
+			FIntPoint AbsoluteSectionBase = Comp->GetSectionBase() + Difference;
+			Comp->SetSectionBase(AbsoluteSectionBase);
+		}
+	}
+
+	// LandscapeInfo sorts on section base
+	ULandscapeInfo* LandscapeInfo = GetLandscapeInfo();
+	check(LandscapeInfo);
+	LandscapeInfo->SortStreamingProxies();
+}
 
 void ALandscapeProxy::Destroyed()
 {
@@ -5141,26 +5173,9 @@ void ALandscapeProxy::FixupSharedData(ALandscape* Landscape, const bool bMapChec
 	OnLandscapeProxyFixupSharedDataDelegate.Broadcast(/*Proxy = */this, FOnLandscapeProxyFixupSharedDataParams { .Landscape = Landscape, .bUpgradeSharedPropertiesPerformed = bUpgradeSharedPropertiesPerformedBefore });
 }
 
+// Deprecated
 void ALandscapeProxy::SetAbsoluteSectionBase(FIntPoint InSectionBase)
 {
-	FIntPoint Difference = InSectionBase - LandscapeSectionOffset;
-	LandscapeSectionOffset = InSectionBase;
-
-	RecreateComponentsRenderState([Difference](ULandscapeComponent* Comp)
-	{
-		FIntPoint AbsoluteSectionBase = Comp->GetSectionBase() + Difference;
-		Comp->SetSectionBase(AbsoluteSectionBase);
-	});
-	
-	for (int32 CompIdx = 0; CompIdx < CollisionComponents.Num(); CompIdx++)
-	{
-		ULandscapeHeightfieldCollisionComponent* Comp = CollisionComponents[CompIdx];
-		if (Comp)
-		{
-			FIntPoint AbsoluteSectionBase = Comp->GetSectionBase() + Difference;
-			Comp->SetSectionBase(AbsoluteSectionBase);
-		}
-	}
 }
 
 void ALandscapeProxy::RecreateComponentsState()
@@ -5549,7 +5564,7 @@ ALandscapeProxy* ULandscapeInfo::GetLandscapeProxy() const
 	}
 
 	// any proxy in the world
-	for (TWeakObjectPtr<ALandscapeStreamingProxy> ProxyPtr : StreamingProxies)
+	for (TWeakObjectPtr<ALandscapeStreamingProxy> ProxyPtr : SortedStreamingProxies)
 	{
 		ALandscapeStreamingProxy* Proxy = ProxyPtr.Get();
 		USceneComponent* ProxyRootComponent = (Proxy != nullptr) ? Proxy->GetRootComponent() : nullptr;
@@ -5569,7 +5584,7 @@ void ULandscapeInfo::Reset()
 {
 	LandscapeActor.Reset();
 
-	StreamingProxies.Empty();
+	SortedStreamingProxies.Empty();
 	XYtoComponentMap.Empty();
 	XYtoAddCollisionMap.Empty();
 
@@ -5590,7 +5605,7 @@ void ULandscapeInfo::FixupProxiesTransform(bool bDirty)
 	}
 
 	// Make sure section offset of all proxies is multiple of ALandscapeProxy::ComponentSizeQuads
-	for (TWeakObjectPtr<ALandscapeStreamingProxy> ProxyPtr : StreamingProxies)
+	for (TWeakObjectPtr<ALandscapeStreamingProxy> ProxyPtr : SortedStreamingProxies)
 	{
 		ALandscapeProxy* Proxy = ProxyPtr.Get();
 		if (!Proxy)
@@ -5603,25 +5618,25 @@ void ULandscapeInfo::FixupProxiesTransform(bool bDirty)
 			Proxy->Modify();
 		}
 
-		FIntPoint LandscapeSectionOffset = Proxy->LandscapeSectionOffset - Landscape->LandscapeSectionOffset;
+		FIntPoint LandscapeSectionOffset = Proxy->GetSectionBase() - Landscape->GetSectionBase();
 		FIntPoint LandscapeSectionOffsetRem(
 			LandscapeSectionOffset.X % Proxy->ComponentSizeQuads,
 			LandscapeSectionOffset.Y % Proxy->ComponentSizeQuads);
 
 		if (LandscapeSectionOffsetRem.X != 0 || LandscapeSectionOffsetRem.Y != 0)
 		{
-			FIntPoint NewLandscapeSectionOffset = Proxy->LandscapeSectionOffset - LandscapeSectionOffsetRem;
+			FIntPoint NewLandscapeSectionOffset = Proxy->GetSectionBase() - LandscapeSectionOffsetRem;
 
 			UE_LOG(LogLandscape, Warning, TEXT("Landscape section base is not multiple of component size, attempted automated fix: '%s', %d,%d vs %d,%d."),
-				*Proxy->GetFullName(), Proxy->LandscapeSectionOffset.X, Proxy->LandscapeSectionOffset.Y, NewLandscapeSectionOffset.X, NewLandscapeSectionOffset.Y);
+				*Proxy->GetFullName(), Proxy->GetSectionBase().X, Proxy->GetSectionBase().Y, NewLandscapeSectionOffset.X, NewLandscapeSectionOffset.Y);
 
-			Proxy->SetAbsoluteSectionBase(NewLandscapeSectionOffset);
+			Proxy->SetSectionBase(NewLandscapeSectionOffset);
 		}
 	}
 
 	FTransform LandscapeTM = Landscape->LandscapeActorToWorld();
 	// Update transformations of all linked landscape proxies
-	for (TWeakObjectPtr<ALandscapeStreamingProxy> ProxyPtr : StreamingProxies)
+	for (TWeakObjectPtr<ALandscapeStreamingProxy> ProxyPtr : SortedStreamingProxies)
 	{
 		ALandscapeProxy* Proxy = ProxyPtr.Get();
 		if (!Proxy)
@@ -5629,7 +5644,7 @@ void ULandscapeInfo::FixupProxiesTransform(bool bDirty)
 			continue;
 		}
 
-		FTransform ProxyRelativeTM(FVector(Proxy->LandscapeSectionOffset));
+		FTransform ProxyRelativeTM(FVector(Proxy->GetSectionBase()));
 		FTransform ProxyTransform = ProxyRelativeTM * LandscapeTM;
 
 		if (!Proxy->GetTransform().Equals(ProxyTransform))
@@ -5771,6 +5786,17 @@ void ULandscapeInfo::Initialize(UWorld* InWorld, const FGuid& InLandscapeGuid)
 	LandscapeGuid = InLandscapeGuid;
 }
 
+const TArray<TWeakObjectPtr<ALandscapeStreamingProxy>>& ULandscapeInfo::GetSortedStreamingProxies() const
+{
+	check(Algo::IsSorted(SortedStreamingProxies, UE::Landscape::LandscapeProxySortPredicate));
+	return SortedStreamingProxies;
+}
+
+void ULandscapeInfo::SortStreamingProxies()
+{
+	SortedStreamingProxies.Sort(UE::Landscape::LandscapeProxySortPredicate);
+}
+
 void ULandscapeInfo::ForEachLandscapeProxy(TFunctionRef<bool(ALandscapeProxy*)> Fn) const
 {
 	if (ALandscape* Landscape = LandscapeActor.Get())
@@ -5785,7 +5811,9 @@ void ULandscapeInfo::ForEachLandscapeProxy(TFunctionRef<bool(ALandscapeProxy*)>
 		
 	}
 
-	for (TWeakObjectPtr<ALandscapeStreamingProxy> StreamingProxyPtr : StreamingProxies)
+	// Gather the proxies in a local list so it won't be updated through a re-register while iterating 
+	const TArray <TWeakObjectPtr<ALandscapeStreamingProxy>> SortedStreamingProxiesCopy(SortedStreamingProxies);
+	for (TWeakObjectPtr<ALandscapeStreamingProxy> StreamingProxyPtr : SortedStreamingProxiesCopy)
 	{
 		if (ALandscapeProxy* LandscapeProxy = StreamingProxyPtr.Get())
 		{
@@ -5850,7 +5878,7 @@ bool ULandscapeInfo::IsRegistered(const ALandscapeProxy* Proxy) const
 	else if (const ALandscapeStreamingProxy *StreamingProxy = Cast<ALandscapeStreamingProxy>(Proxy))
 	{
 		TWeakObjectPtr<const ALandscapeStreamingProxy> StreamingProxyPtr = StreamingProxy;
-		bResult = StreamingProxies.Contains(StreamingProxyPtr);
+		bResult = SortedStreamingProxies.Contains(StreamingProxyPtr);
 	}
 
 #if WITH_EDITORONLY_DATA
@@ -5986,6 +6014,8 @@ void ULandscapeInfo::RegisterLandscapeActorWithProxyInternal(ALandscapeProxy* Pr
 			}
 		}
 	}
+
+	Proxy->bIsLandscapeActorRegisteredWithLandscapeInfo = true;
 #endif // WITH_EDITOR
 }
 
@@ -6052,7 +6082,7 @@ void ULandscapeInfo::RegisterActor(ALandscapeProxy* Proxy, bool bMapCheck, bool
 
 			// run post-landscape actor registration on the LandscapeActor first, then on each streaming proxy
 			RegisterLandscapeActorWithProxyInternal(Landscape, bMapCheck);
-			for (TWeakObjectPtr<ALandscapeStreamingProxy> StreamingProxyPtr : StreamingProxies)
+			for (TWeakObjectPtr<ALandscapeStreamingProxy> StreamingProxyPtr : SortedStreamingProxies)
 			{
 				if (ALandscapeStreamingProxy* StreamingProxy = StreamingProxyPtr.Get())
 				{
@@ -6081,37 +6111,15 @@ void ULandscapeInfo::RegisterActor(ALandscapeProxy* Proxy, bool bMapCheck, bool
 		}
 #endif // WITH_EDITORONLY_DATA
 
-		auto LamdbdaLowerBound = [](TWeakObjectPtr<ALandscapeProxy> APtr, TWeakObjectPtr<ALandscapeProxy> BPtr)
-		{
-			ALandscapeProxy *A = APtr.Get();
-			ALandscapeProxy *B = BPtr.Get();
-
-			// sort nulls, assuming null < !null
-			if (!A || !B)
-			{
-				return !!B;
-			}
-			
-			FIntPoint SectionBaseA = A->GetSectionBaseOffset();
-			FIntPoint SectionBaseB = B->GetSectionBaseOffset();
-
-			if (SectionBaseA.X != SectionBaseB.X)
-			{
-				return SectionBaseA.X < SectionBaseB.X;
-			}
-
-			return SectionBaseA.Y < SectionBaseB.Y;
-		};
-
 		// Insert Proxies in a sorted fashion into the landscape info Proxies list, for generating deterministic results in the Layer system
 		ALandscapeStreamingProxy* StreamingProxy = CastChecked<ALandscapeStreamingProxy>(Proxy);
 		TWeakObjectPtr<ALandscapeStreamingProxy> StreamingProxyPtr = StreamingProxy;
-		if (!StreamingProxies.Contains(StreamingProxyPtr))
+		if (!SortedStreamingProxies.Contains(StreamingProxyPtr))
 		{
 			// NOTE: if a streaming proxy somehow gets garbage collected without de-registering from the Proxies list, then
 			// this search may return a non-deterministic index because the Proxies list will contain a null
-			uint32 InsertIndex = Algo::LowerBound(StreamingProxies, StreamingProxyPtr, LamdbdaLowerBound);
-			StreamingProxies.Insert(StreamingProxyPtr, InsertIndex);
+			uint32 InsertIndex = Algo::LowerBound(SortedStreamingProxies, StreamingProxyPtr, UE::Landscape::LandscapeProxySortPredicate);
+			SortedStreamingProxies.Insert(StreamingProxyPtr, InsertIndex);
 		}
 
 #if WITH_EDITORONLY_DATA
@@ -6163,7 +6171,7 @@ void ULandscapeInfo::UnregisterActor(ALandscapeProxy* Proxy)
 		}
 
 		// update proxies reference to landscape actor
-		for (TWeakObjectPtr<ALandscapeStreamingProxy> StreamingProxyPtr : StreamingProxies)
+		for (TWeakObjectPtr<ALandscapeStreamingProxy> StreamingProxyPtr : SortedStreamingProxies)
 		{
 			if (ALandscapeStreamingProxy* StreamingProxy = StreamingProxyPtr.Get())
 			{
@@ -6175,7 +6183,7 @@ void ULandscapeInfo::UnregisterActor(ALandscapeProxy* Proxy)
 	{
 		ALandscapeStreamingProxy* StreamingProxy = CastChecked<ALandscapeStreamingProxy>(Proxy);
 		TWeakObjectPtr<ALandscapeStreamingProxy> StreamingProxyPtr = StreamingProxy;
-		StreamingProxies.Remove(StreamingProxyPtr);
+		SortedStreamingProxies.Remove(StreamingProxyPtr);
 
 		FLandscapeGroup::UnregisterAllComponentsOnStreamingProxy(StreamingProxy);
 	}
diff --git a/Engine/Source/Runtime/Landscape/Private/LandscapeCollision.cpp b/Engine/Source/Runtime/Landscape/Private/LandscapeCollision.cpp
--- a/Engine/Source/Runtime/Landscape/Private/LandscapeCollision.cpp
+++ b/Engine/Source/Runtime/Landscape/Private/LandscapeCollision.cpp
@@ -1952,8 +1952,8 @@ void ULandscapeHeightfieldCollisionComponent::PostLoad()
 		{
 			// This is to ensure that component relative location is exact section base offset value
 			FVector LocalRelativeLocation = GetRelativeLocation();
-			float CheckRelativeLocationX = float(SectionBaseX - LandscapeProxy->LandscapeSectionOffset.X);
-			float CheckRelativeLocationY = float(SectionBaseY - LandscapeProxy->LandscapeSectionOffset.Y);
+			float CheckRelativeLocationX = float(SectionBaseX - LandscapeProxy->GetSectionBase().X);
+			float CheckRelativeLocationY = float(SectionBaseY - LandscapeProxy->GetSectionBase().Y);
 			if (!FMath::IsNearlyEqual(CheckRelativeLocationX, LocalRelativeLocation.X, UE_DOUBLE_KINDA_SMALL_NUMBER) ||
 				!FMath::IsNearlyEqual(CheckRelativeLocationY, LocalRelativeLocation.Y, UE_DOUBLE_KINDA_SMALL_NUMBER))
 			{
diff --git a/Engine/Source/Runtime/Landscape/Private/LandscapeConfigHelper.cpp b/Engine/Source/Runtime/Landscape/Private/LandscapeConfigHelper.cpp
--- a/Engine/Source/Runtime/Landscape/Private/LandscapeConfigHelper.cpp
+++ b/Engine/Source/Runtime/Landscape/Private/LandscapeConfigHelper.cpp
@@ -78,10 +78,9 @@ ALandscapeProxy* FLandscapeConfigHelper::FindOrAddLandscapeStreamingProxy(UActor
 
 		LandscapeProxy->CreateLandscapeInfo();
 		LandscapeProxy->SetActorLocationAndRotation(ProxyLocation, Landscape->GetActorRotation());
-		LandscapeProxy->LandscapeSectionOffset = FIntPoint(CellLocation.X, CellLocation.Y);
+		LandscapeProxy->SetSectionBase(FIntPoint(CellLocation.X, CellLocation.Y));
 		LandscapeProxy->SetIsSpatiallyLoaded(LandscapeProxy->GetLandscapeInfo()->AreNewLandscapeActorsSpatiallyLoaded());
-		// TODO [jonathan.bard] temp fix for the fact that PostRegisterAllComponents doesn't get called anymore, to prevent bad data from creeping in. 
-		//  A proper fix is coming soon
+		// Fixup the transform based on the section base offset
 		InLandscapeInfo->FixupProxiesTransform();
 	};
 
@@ -251,16 +250,20 @@ bool FLandscapeConfigHelper::PartitionLandscape(UWorld* InWorld, ULandscapeInfo*
 	// Handle Landscapes with missing LandscapeActor(s)
 	if (!InLandscapeInfo->LandscapeActor.Get())
 	{
-		// Use the first proxy as the landscape template
-		if (ALandscapeProxy* FirstProxy = InLandscapeInfo->StreamingProxies[0].Get())
+		const TArray<TWeakObjectPtr<ALandscapeStreamingProxy>>& StreamingProxies = InLandscapeInfo->GetSortedStreamingProxies();
+		if (!StreamingProxies.IsEmpty())
 		{
-			FActorSpawnParameters SpawnParams;
-			FTransform LandscapeTransform = FirstProxy->LandscapeActorToWorld();
-			ALandscape* NewLandscape = InWorld->SpawnActor<ALandscape>(ALandscape::StaticClass(), LandscapeTransform, SpawnParams);
+			// Use the first proxy as the landscape template
+			if (ALandscapeProxy* FirstProxy = StreamingProxies[0].Get())
+			{
+				FActorSpawnParameters SpawnParams;
+				FTransform LandscapeTransform = FirstProxy->LandscapeActorToWorld();
+				ALandscape* NewLandscape = InWorld->SpawnActor<ALandscape>(ALandscape::StaticClass(), LandscapeTransform, SpawnParams);
 
-			NewLandscape->CopySharedProperties(FirstProxy);
+				NewLandscape->CopySharedProperties(FirstProxy);
 
-			InLandscapeInfo->RegisterActor(NewLandscape);
+				InLandscapeInfo->RegisterActor(NewLandscape);
+			}
 		}
 	}
 
diff --git a/Engine/Source/Runtime/Landscape/Private/LandscapeEdit.cpp b/Engine/Source/Runtime/Landscape/Private/LandscapeEdit.cpp
--- a/Engine/Source/Runtime/Landscape/Private/LandscapeEdit.cpp
+++ b/Engine/Source/Runtime/Landscape/Private/LandscapeEdit.cpp
@@ -143,7 +143,7 @@ void ULandscapeComponent::Init(int32 InBaseX, int32 InBaseY, int32 InComponentSi
 	LandscapeProxy->LandscapeComponents.Add(this);
 
 	SetSectionBase(FIntPoint(InBaseX, InBaseY));
-	SetRelativeLocation(FVector(GetSectionBase() - GetLandscapeProxy()->LandscapeSectionOffset));
+	SetRelativeLocation(FVector(GetSectionBase() - GetLandscapeProxy()->GetSectionBase()));
 	ComponentSizeQuads = InComponentSizeQuads;
 	NumSubsections = InNumSubsections;
 	SubsectionSizeQuads = InSubsectionSizeQuads;
@@ -3799,13 +3799,8 @@ void ALandscapeProxy::Import(const FGuid& InGuid, int32 InMinX, int32 InMinY, in
 	// Must be after the FMaterialUpdateContext is destroyed
 	RecreateRenderStateContexts.Reset();
 
-	// Create the default layer first
 	ALandscape* LandscapeActor = GetLandscapeActor();
 	check(LandscapeActor != nullptr);
-	if (LandscapeActor->GetEditLayersConst().IsEmpty() && InImportLayers.IsEmpty())
-	{
-		LandscapeActor->CreateDefaultLayer();
-	}
 
 	// Create and initialize landscape info object after default layer exists
 	ULandscapeInfo* LandscapeInfo = CreateLandscapeInfo();
@@ -3814,6 +3809,12 @@ void ALandscapeProxy::Import(const FGuid& InGuid, int32 InMinX, int32 InMinY, in
 	bool ShouldComponentBeRegistered = GetLevel()->bIsVisible;
 	RegisterAllComponents();
 
+	// Create the default layer after components are registered
+	if (LandscapeActor->GetEditLayersConst().IsEmpty() && InImportLayers.IsEmpty())
+	{
+		LandscapeActor->CreateDefaultLayer();
+	}
+
 	TSet<ULandscapeComponent*> ComponentsToProcess;
 
 	struct FLayerImportSettings
@@ -4263,10 +4264,8 @@ bool ALandscapeProxy::ExportToRawMeshDataCopy(const FRawMeshExportParams& InExpo
 	{
 		check(Component->ComponentSizeQuads == ComponentSizeQuads);		// must all be the same number of quads
 
-		FIntPoint SectionBase{ Component->SectionBaseX, Component->SectionBaseY };
-
-		MinSectionBase = MinSectionBase.ComponentMin(SectionBase);
-		MaxSectionBase = MaxSectionBase.ComponentMax(SectionBase);
+		MinSectionBase = MinSectionBase.ComponentMin(Component->SectionBaseX);
+		MaxSectionBase = MaxSectionBase.ComponentMax(Component->SectionBaseY);
 	}
 
 	FVector2f ExportedSectionMinQuadCoord(MinSectionBase.X, MinSectionBase.Y);
@@ -4352,7 +4351,7 @@ bool ALandscapeProxy::ExportToRawMeshDataCopy(const FRawMeshExportParams& InExpo
 
 		const FLandscapeComponentDataInterfaceBase& CDI = *AsyncData.ComponentData.Find(Component)->ComponentDataInterface;
 
-		const FIntPoint ComponentOffsetRelativeToProxyBoundsQuads = Component->GetSectionBase() - LandscapeSectionOffset - LandscapeProxyBoundsRect.Min;
+		const FIntPoint ComponentOffsetRelativeToProxyBoundsQuads = Component->GetSectionBase() - SectionBase - LandscapeProxyBoundsRect.Min;
 		const FVector2f ComponentOffsetRelativeToProxyBoundsQuadsLOD = FVector2f(ComponentOffsetRelativeToProxyBoundsQuads) * LODScale;
 
 		const FVector2f GlobalQuadCoordsOffset = FVector2f(Component->GetSectionBase());
@@ -4736,7 +4735,7 @@ FIntRect ALandscapeProxy::GetBoundingRect() const
 			}
 		}
 		Rect.Max += FIntPoint(ComponentSizeQuads, ComponentSizeQuads);
-		Rect -= LandscapeSectionOffset;
+		Rect -= SectionBase;
 		return Rect;
 	}
 
@@ -5649,7 +5648,7 @@ ALandscapeProxy* ULandscapeInfo::MoveComponentsToProxy(const TArray<ULandscapeCo
 		// by default first component location
 		ULandscapeComponent* FirstComponent = *TargetSelectedComponents.CreateConstIterator();
 		LandscapeProxy->GetRootComponent()->SetWorldLocationAndRotation(FirstComponent->GetComponentLocation(), FirstComponent->GetComponentRotation());
-		LandscapeProxy->LandscapeSectionOffset = FirstComponent->GetSectionBase();
+		LandscapeProxy->SetSectionBase(FirstComponent->GetSectionBase());
 	}
 
 	// Hide(unregister) the new landscape if owning level currently in hidden state
@@ -6506,7 +6505,7 @@ bool ULandscapeInfo::CanDeleteLandscape(FText& OutReason) const
 	int32 UndeletedSplineCount = 0;
 
 	// Check Registered Proxies
-	for (TWeakObjectPtr<ALandscapeStreamingProxy> ProxyPtr : StreamingProxies)
+	for (TWeakObjectPtr<ALandscapeStreamingProxy> ProxyPtr : SortedStreamingProxies)
 	{
 		ALandscapeProxy* RegisteredProxy = ProxyPtr.Get();
 		if (!RegisteredProxy || RegisteredProxy == LandscapeActor)
@@ -6536,7 +6535,7 @@ bool ULandscapeInfo::CanDeleteLandscape(FText& OutReason) const
 						const bool bIsLoaded = ActorDescInstance->IsLoaded();
 						const bool bIsAllocated = (LandscapeProxy != nullptr);
 						const bool bIsGarbage = bIsAllocated && !IsValidChecked(LandscapeProxy);
-						const bool bIsRegistered = bIsAllocated && StreamingProxies.Contains(CastChecked<ALandscapeStreamingProxy>(LandscapeProxy));
+						const bool bIsRegistered = bIsAllocated && SortedStreamingProxies.Contains(CastChecked<ALandscapeStreamingProxy>(LandscapeProxy));
 						const bool bRegisteredComps = bIsAllocated && LandscapeProxy->HasActorRegisteredAllComponents();
 						check(bIsRegistered == bRegisteredComps);
 
@@ -6765,17 +6764,16 @@ void ALandscape::PostEditChangeProperty(FPropertyChangedEvent& PropertyChangedEv
 		if (bPropagateToProxies)
 		{
 			// Propagate Event to Proxies...
-			// First, gather the proxies in a local list so that we don't get screwed by PostEditChangeProperty re-registering streaming proxies while iterating :
-			// TODO [jonathan.bard] : replace by ForEachLandscapeProxy (and make ForEachLandscapeProxy also iterate on the list's local copy)
-			TArray<TWeakObjectPtr<ALandscapeStreamingProxy>> StreamingProxies(Info->StreamingProxies);
-			for (TWeakObjectPtr<ALandscapeStreamingProxy> ProxyPtr : StreamingProxies)
+			Info->ForEachLandscapeProxy([this, &PropertyChangedEvent](ALandscapeProxy* Proxy)
 			{
-				if (ALandscapeProxy* Proxy = ProxyPtr.Get())
+				// The ALandscape is included in ForEachLandscapeProxy, avoid getting stuck in a PostEditChange loop
+				if (Proxy != this)
 				{
 					Proxy->SynchronizeSharedProperties(this);
 					Proxy->PostEditChangeProperty(PropertyChangedEvent);
 				}
-			}
+				return true;
+			});
 		}
 
 		if (bMarkAllLandscapeRenderStateDirty)
@@ -8228,13 +8226,11 @@ bool ALandscapeProxy::LandscapeExportHeightmapToRenderTarget(UTextureRenderTarge
 	{
 		ULandscapeInfo* LandscapeInfo = GetLandscapeInfo();
 
-		for (TWeakObjectPtr<ALandscapeStreamingProxy> ProxyPtr : LandscapeInfo->StreamingProxies)
+		LandscapeInfo->ForEachLandscapeProxy([&](ALandscapeProxy* Proxy)
 		{
-			if (ALandscapeProxy* Proxy = ProxyPtr.Get())
-			{
-				LandscapeComponentsToExport.Append(Proxy->LandscapeComponents);
-			}
-		}
+			LandscapeComponentsToExport.Append(Proxy->LandscapeComponents);
+			return true;
+		});
 	}
 
 	if (LandscapeComponentsToExport.Num() == 0)
diff --git a/Engine/Source/Runtime/Landscape/Private/LandscapeEditLayers.cpp b/Engine/Source/Runtime/Landscape/Private/LandscapeEditLayers.cpp
--- a/Engine/Source/Runtime/Landscape/Private/LandscapeEditLayers.cpp
+++ b/Engine/Source/Runtime/Landscape/Private/LandscapeEditLayers.cpp
@@ -4316,15 +4316,14 @@ void ALandscape::DeleteUnusedLayers()
 		return;
 	}
 
-	for (const TWeakObjectPtr<ALandscapeStreamingProxy>& Proxy : LandscapeInfo->StreamingProxies)
+	LandscapeInfo->ForEachLandscapeProxy([&](ALandscapeProxy* Proxy)
 	{
-		if (!Proxy.IsValid())
+		if (Proxy != this)
 		{
-			continue;
+			Proxy->DeleteUnusedLayers();
 		}
-
-		Proxy->DeleteUnusedLayers();
-	}
+		return true;
+	});
 
 	ALandscapeProxy::DeleteUnusedLayers();
 }
diff --git a/Engine/Source/Runtime/Landscape/Private/LandscapeGrass.cpp b/Engine/Source/Runtime/Landscape/Private/LandscapeGrass.cpp
--- a/Engine/Source/Runtime/Landscape/Private/LandscapeGrass.cpp
+++ b/Engine/Source/Runtime/Landscape/Private/LandscapeGrass.cpp
@@ -1737,7 +1737,7 @@ struct FGrassBuilderBase
 
 		DrawScale = Landscape->GetRootComponent()->GetRelativeScale3D();
 		DrawLoc = Landscape->GetActorLocation();
-		LandscapeSectionOffset = Landscape->LandscapeSectionOffset;
+		LandscapeSectionOffset = Landscape->GetSectionBase();
 
 		SectionBase = Component->GetSectionBase();
 		ComponentSizeQuads = Component->ComponentSizeQuads;
diff --git a/Engine/Source/Runtime/Landscape/Private/LandscapeGrassWeightExporter.cpp b/Engine/Source/Runtime/Landscape/Private/LandscapeGrassWeightExporter.cpp
--- a/Engine/Source/Runtime/Landscape/Private/LandscapeGrassWeightExporter.cpp
+++ b/Engine/Source/Runtime/Landscape/Private/LandscapeGrassWeightExporter.cpp
@@ -523,7 +523,7 @@ FLandscapeGrassWeightExporter::FLandscapeGrassWeightExporter(ALandscapeProxy* In
 	// Then compute FComponentInfo's ViewOffset with the knowledge of the total render target size : 
 	for (FComponentInfo& ComponentInfo : ComponentInfos)
 	{
-		FIntPoint ComponentOffset = (ComponentInfo.Component->GetSectionBase() - LandscapeProxy->LandscapeSectionOffset);
+		FIntPoint ComponentOffset = (ComponentInfo.Component->GetSectionBase() - LandscapeProxy->GetSectionBase());
 		FVector2D ViewOffset(-ComponentOffset.X, ComponentOffset.Y);
 		ViewOffset.X += ComponentInfo.PixelOffsetX;
 		ViewOffset /= (FVector2D(TargetSize) * 0.5f);
diff --git a/Engine/Source/Runtime/Landscape/Private/LandscapeGroup.cpp b/Engine/Source/Runtime/Landscape/Private/LandscapeGroup.cpp
--- a/Engine/Source/Runtime/Landscape/Private/LandscapeGroup.cpp
+++ b/Engine/Source/Runtime/Landscape/Private/LandscapeGroup.cpp
@@ -331,25 +331,25 @@ namespace UE::Landscape::Private
 
 		// Make sure section offset of all proxies is multiple of ALandscapeProxy::ComponentSizeQuads
 		{
-			FIntPoint LandscapeSectionOffset = Proxy->LandscapeSectionOffset - Landscape->LandscapeSectionOffset;
+			FIntPoint LandscapeSectionOffset = Proxy->GetSectionBase() - Landscape->GetSectionBase();
 			FIntPoint LandscapeSectionOffsetRem(
 				LandscapeSectionOffset.X % Proxy->ComponentSizeQuads,
 				LandscapeSectionOffset.Y % Proxy->ComponentSizeQuads);
 			if (LandscapeSectionOffsetRem.X != 0 || LandscapeSectionOffsetRem.Y != 0)
 			{
-				FIntPoint NewLandscapeSectionOffset = Proxy->LandscapeSectionOffset - LandscapeSectionOffsetRem;
+				FIntPoint NewLandscapeSectionOffset = Proxy->GetSectionBase() - LandscapeSectionOffsetRem;
 
 				UE_LOG(LogLandscape, Warning, TEXT("Landscape section base is not multiple of component size, attempted automated fix: '%s', %d,%d vs %d,%d."),
-					*Proxy->GetFullName(), Proxy->LandscapeSectionOffset.X, Proxy->LandscapeSectionOffset.Y, NewLandscapeSectionOffset.X, NewLandscapeSectionOffset.Y);
+					*Proxy->GetFullName(), Proxy->GetSectionBase().X, Proxy->GetSectionBase().Y, NewLandscapeSectionOffset.X, NewLandscapeSectionOffset.Y);
 
-				Proxy->SetAbsoluteSectionBase(NewLandscapeSectionOffset);
+				Proxy->SetSectionBase(NewLandscapeSectionOffset);
 			}
 		}
 
 		// Update transformations of all linked landscape proxies
 		FTransform LandscapeTM = Landscape->LandscapeActorToWorld();
 		{
-			FTransform ProxyRelativeTM(FVector(Proxy->LandscapeSectionOffset));
+			FTransform ProxyRelativeTM(FVector(Proxy->GetSectionBase()));
 			FTransform ProxyTransform = ProxyRelativeTM * LandscapeTM;
 
 			if (!Proxy->GetTransform().Equals(ProxyTransform))
diff --git a/Engine/Source/Runtime/Landscape/Private/LandscapeSubsystem.cpp b/Engine/Source/Runtime/Landscape/Private/LandscapeSubsystem.cpp
--- a/Engine/Source/Runtime/Landscape/Private/LandscapeSubsystem.cpp
+++ b/Engine/Source/Runtime/Landscape/Private/LandscapeSubsystem.cpp
@@ -1182,7 +1182,7 @@ void ULandscapeSubsystem::BuildNanite(UE::Landscape::EBuildFlags InBuildFlags, T
 				ULandscapeInfo* LandscapeInfo = Landscape->GetLandscapeInfo();
 				if (LandscapeInfo != nullptr)
 				{
-					Algo::Transform(LandscapeInfo->StreamingProxies, FinalProxiesToBuild, [](const TWeakObjectPtr<ALandscapeStreamingProxy>& InStreamingProxy) { return InStreamingProxy.Get(); });
+					Algo::Transform(LandscapeInfo->GetSortedStreamingProxies(), FinalProxiesToBuild, [](const TWeakObjectPtr<ALandscapeStreamingProxy>& InStreamingProxy) { return InStreamingProxy.Get(); });
 				}
 			}
 		}
diff --git a/Engine/Source/Runtime/Landscape/Private/LandscapeUtils.cpp b/Engine/Source/Runtime/Landscape/Private/LandscapeUtils.cpp
--- a/Engine/Source/Runtime/Landscape/Private/LandscapeUtils.cpp
+++ b/Engine/Source/Runtime/Landscape/Private/LandscapeUtils.cpp
@@ -591,4 +591,26 @@ bool DeleteActors(const TArray<AActor*>& InActorsToDelete, UWorld* InWorld, bool
 
 #endif // WITH_EDITOR
 
+bool LandscapeProxySortPredicate(const TWeakObjectPtr<ALandscapeProxy> APtr, const TWeakObjectPtr<ALandscapeProxy> BPtr)
+{
+	ALandscapeProxy* A = APtr.Get();
+	ALandscapeProxy* B = BPtr.Get();
+
+	// sort nulls, assuming null < !null
+	if (!A || !B)
+	{
+		return !!B;
+	}
+
+	FIntPoint SectionBaseA = A->GetSectionBase();
+	FIntPoint SectionBaseB = B->GetSectionBase();
+
+	if (SectionBaseA.X != SectionBaseB.X)
+	{
+		return SectionBaseA.X < SectionBaseB.X;
+	}
+
+	return SectionBaseA.Y < SectionBaseB.Y;
+}
+
 } // end namespace UE::Landscape
\ No newline at end of file
diff --git a/Engine/Source/Runtime/Landscape/Public/LandscapeUtils.h b/Engine/Source/Runtime/Landscape/Public/LandscapeUtils.h
--- a/Engine/Source/Runtime/Landscape/Public/LandscapeUtils.h
+++ b/Engine/Source/Runtime/Landscape/Public/LandscapeUtils.h
@@ -14,6 +14,7 @@
 
 enum EShaderPlatform : uint16;
 class AActor;
+class ALandscapeProxy;
 class FRDGBuilder;
 class FRDGEventName;
 class ULandscapeComponent;
@@ -371,4 +372,7 @@ LANDSCAPE_API bool DeleteActors(const TArray<AActor*>& InActorsToDelete, UWorld*
 
 #endif // WITH_EDITOR
 
+/** Helper to sort ALandscapeProxy actors by their section base. Compares X axis values by default and Y values if needed */
+LANDSCAPE_API bool LandscapeProxySortPredicate(const TWeakObjectPtr<ALandscapeProxy> APtr, const TWeakObjectPtr<ALandscapeProxy> BPtr);
+
 } // end namespace UE::Landscape
