Landscape ray tracing crash fix.  The crash is in linked list related code tracking ray tracing state for views, and the fix removes the linked list entirely, instead iterating over the global list of Landscape Render Systems to clean up view specific ray tracing state on view deletion.  Ray tracing updates happening for different render proxies in parallel could corrupt the linked list, since it's not a thread safe structure.  The global list is only modified on the main render thread, and won't suffer from the same issue.

#rnx

diff --git a/Engine/Source/Runtime/Landscape/Private/LandscapeRender.cpp b/Engine/Source/Runtime/Landscape/Private/LandscapeRender.cpp
--- a/Engine/Source/Runtime/Landscape/Private/LandscapeRender.cpp
+++ b/Engine/Source/Runtime/Landscape/Private/LandscapeRender.cpp
@@ -587,13 +587,6 @@ struct FLandscapeRayTracingState
 	FLandscapeRayTracingState() = default;
 	~FLandscapeRayTracingState();
 
-	// Parent structure that holds this ray tracing state, needed for deletion of this item if its view is deleted
-	FLandscapeRayTracingImpl* Pimpl = nullptr;
-
-	// Linked list pointers for FLandscapeRayTracingStateList, referenced from FSceneViewState, iterated over if view gets deleted
-	FLandscapeRayTracingState* ViewStateNext = nullptr;
-	FLandscapeRayTracingState** ViewStatePrev = nullptr;
-
 	// View state key from FSceneViewState.  Zero is used if FSceneViewState is null (view state keys start at 1, so 0 is invalid for an actual view).
 	uint32 ViewKey = -1;
 
@@ -608,8 +601,7 @@ struct FLandscapeRayTracingState
 // This wrapper holds the ray tracing state for a single scene proxy for all views
 struct FLandscapeRayTracingImpl
 {
-	// Needs to be indirect array, because elements are added to a linked list 
-	TIndirectArray<FLandscapeRayTracingState> PerViewRayTracingState;
+	TArray<FLandscapeRayTracingState> PerViewRayTracingState;
 
 	// ViewStateInterface pointer can be NULL
 	FLandscapeRayTracingState* FindOrCreateRayTracingState(
@@ -620,36 +612,24 @@ struct FLandscapeRayTracingImpl
 	    int32 RayTracingSectionVertCount);
 };
 
-// When views get deleted, we need to clean up the per view ray tracing data, which is handled by this class.  This is just a doubly linked list
-// of all the ray tracing data associated with a given view, pointed to by the FSceneViewState, with the destructor emptying the items from the list.
+// When views get deleted, we need to clean up the per view ray tracing state, which is handled by this class.  On deletion of the view,
+// iterate over active landscape systems and their sections (each is a render proxy), and notify them that the view was deleted.  Note that
+// originally this class included an actual list of the render proxies, but it was removed in favor of iterating over the global list, and
+// now just includes the view's key.
 class FLandscapeRayTracingStateList
 {
 public:
-	FLandscapeRayTracingState* ListHead;
-
-	FLandscapeRayTracingStateList() : ListHead(nullptr) {}
+	uint32 ViewKey;
+	FLandscapeRayTracingStateList(uint32 InViewKey) : ViewKey(InViewKey) {}
 
 	~FLandscapeRayTracingStateList()
 	{
-		// Pop items from the list head until the list is empty
-		while (ListHead)
+		for (auto System : LandscapeRenderSystems)
 		{
-			// Find the item in its parent array
-			FLandscapeRayTracingState* ToRemove = ListHead;
-			FLandscapeRayTracingImpl* Pimpl = ToRemove->Pimpl;
-
-			for (int32 RemoveIndex = 0; RemoveIndex < Pimpl->PerViewRayTracingState.Num(); RemoveIndex++)
+			for (FLandscapeSectionInfo* SectionInfo : System.Value->SectionInfos)
 			{
-				if (&Pimpl->PerViewRayTracingState[RemoveIndex] == ToRemove)
-				{
-					// Remove it -- the destructor will also unlink it from the list, updating ListHead
-					Pimpl->PerViewRayTracingState.RemoveAtSwap(RemoveIndex);
-					break;
-				}
+				SectionInfo->RayTracingNotifyViewDeleted(ViewKey);
 			}
-
-			// Make sure the item was successfully removed (ListHead updated)
-			check(ListHead != ToRemove);
 		}
 	}
 };
@@ -1217,7 +1197,7 @@ void FLandscapeSceneViewExtension::PreRenderView_RenderThread(FRDGBuilder& Graph
 		FSceneViewStateInterface* ViewState = InView.State;
 		if (ViewState->GetLandscapeRayTracingStates() == nullptr)
 		{
-			ViewState->SetLandscapeRayTracingStates( MakePimpl<FLandscapeRayTracingStateList>() );
+			ViewState->SetLandscapeRayTracingStates( MakePimpl<FLandscapeRayTracingStateList>(ViewState->GetViewKey()) );
 		}
 	}
 #endif	// RHI_RAYTRACING
@@ -1888,28 +1868,10 @@ FLandscapeRayTracingState* FLandscapeRayTracingImpl::FindOrCreateRayTracingState
 	}
 
 	// Need to create a new one
-	FLandscapeRayTracingState* RayTracingState = new FLandscapeRayTracingState();
+	FLandscapeRayTracingState* RayTracingState = &PerViewRayTracingState.AddDefaulted_GetRef();
 
-	PerViewRayTracingState.Add(RayTracingState);
-
-	RayTracingState->Pimpl = this;
 	RayTracingState->ViewKey = ViewKey;
 
-	if (ViewState)
-	{
-		// Link into the view state's linked list, so it can be cleaned up if the view gets deleted
-		FLandscapeRayTracingStateList* StateList = ViewState->GetLandscapeRayTracingStates();
-		check(StateList);
-
-		if (StateList->ListHead)
-		{
-			StateList->ListHead->ViewStatePrev = &RayTracingState->ViewStateNext;
-		}
-		RayTracingState->ViewStateNext = StateList->ListHead;
-		RayTracingState->ViewStatePrev = &StateList->ListHead;
-		StateList->ListHead = RayTracingState;
-	}
-
 	// Initialize rendering data
 	uint32 TotalSections = FMath::Square(NumRayTracingSections * NumSubsections);
 	RayTracingState->Sections.SetNum(TotalSections);
@@ -1960,17 +1922,7 @@ FLandscapeRayTracingState* FLandscapeRayTracingImpl::FindOrCreateRayTracingState
 
 FLandscapeRayTracingState::~FLandscapeRayTracingState()
 {
-	// Unlink this from the view state linked list
-	if (ViewStatePrev)
-	{
-		(*ViewStatePrev) = ViewStateNext;
-	}
-	if (ViewStateNext)
-	{
-		ViewStateNext->ViewStatePrev = ViewStatePrev;
-	}
-
-	// And clean up the contents
+	// Clean up the contents
 	for (FLandscapeSectionRayTracingState& SectionRayTracingState : Sections)
 	{
 		SectionRayTracingState.Geometry.ReleaseResource();
@@ -4674,6 +4626,24 @@ bool FLandscapeComponentSceneProxy::ShouldInvalidateShadows(const FSceneView& In
 	return bShouldInvalidateShadow;
 }
 
+#if RHI_RAYTRACING
+void FLandscapeComponentSceneProxy::RayTracingNotifyViewDeleted(uint32 ViewKey)
+{
+	if (RayTracingImpl)
+	{
+		for (int32 ViewStateIndex = 0; ViewStateIndex < RayTracingImpl->PerViewRayTracingState.Num(); ViewStateIndex++)
+		{
+			if (RayTracingImpl->PerViewRayTracingState[ViewStateIndex].ViewKey == ViewKey)
+			{
+				RayTracingImpl->PerViewRayTracingState.RemoveAtSwap(ViewStateIndex);
+				break;
+			}
+		}
+	}
+}
+#endif
+
+
 //
 // FLandscapeSectionInfo
 //
@@ -4745,6 +4715,12 @@ public:
 		return nullptr;
 	}
 
+#if RHI_RAYTRACING
+	virtual void RayTracingNotifyViewDeleted(uint32 ViewKey) override
+	{
+	}
+#endif
+
 private:
 	int8 ProxyLOD;
 	int32 ComponentResolution;
diff --git a/Engine/Source/Runtime/Landscape/Public/LandscapeRender.h b/Engine/Source/Runtime/Landscape/Public/LandscapeRender.h
--- a/Engine/Source/Runtime/Landscape/Public/LandscapeRender.h
+++ b/Engine/Source/Runtime/Landscape/Public/LandscapeRender.h
@@ -440,6 +440,10 @@ public:
 
 	virtual bool ShouldInvalidateShadows(const FSceneView& InView, float InLODValue, float InLastShadowInvalidationLODValue) const { return false; }
 
+#if RHI_RAYTRACING
+	virtual void RayTracingNotifyViewDeleted(uint32 ViewKey) = 0;
+#endif
+
 public:
 	// A hash of the world and (LandscapeGUID or LOD Group Key)
 	uint32 LandscapeKey = 0;
@@ -942,6 +946,11 @@ public:
 	LANDSCAPE_API virtual void GetSectionCenterAndVectors(FVector& OutSectionCenterWorldSpace, FVector& OutSectionXVectorWorldSpace, FVector& OutSectionYVectorWorldSpace) const override;
 	LANDSCAPE_API virtual const FPrimitiveSceneInfo* GetPrimitiveSceneInfo() const override;
 	LANDSCAPE_API virtual bool ShouldInvalidateShadows(const FSceneView& InView, float InLODValue, float InLastShadowInvalidationLODValue) const override;
+
+#if RHI_RAYTRACING
+	LANDSCAPE_API virtual void RayTracingNotifyViewDeleted(uint32 ViewKey) override;
+#endif
+
 };
 
 class FLandscapeDebugMaterialRenderProxy : public FMaterialRenderProxy
