Fix crash in LandscapeRender introduced in CL 45810305
#rnx

diff --git a/Engine/Source/Runtime/Landscape/Private/LandscapeRender.cpp b/Engine/Source/Runtime/Landscape/Private/LandscapeRender.cpp
--- a/Engine/Source/Runtime/Landscape/Private/LandscapeRender.cpp
+++ b/Engine/Source/Runtime/Landscape/Private/LandscapeRender.cpp
@@ -628,7 +628,10 @@ public:
 		{
 			for (FLandscapeSectionInfo* SectionInfo : System.Value->SectionInfos)
 			{
-				SectionInfo->RayTracingNotifyViewDeleted(ViewKey);
+				if (SectionInfo)
+				{
+					SectionInfo->RayTracingNotifyViewDeleted(ViewKey);
+				}
 			}
 		}
 	}
