Fix crash on undo+redo of landscape spline deletion.

diff --git a/Engine/Source/Runtime/Landscape/Private/LandscapeSplines.cpp b/Engine/Source/Runtime/Landscape/Private/LandscapeSplines.cpp
--- a/Engine/Source/Runtime/Landscape/Private/LandscapeSplines.cpp
+++ b/Engine/Source/Runtime/Landscape/Private/LandscapeSplines.cpp
@@ -1181,7 +1181,11 @@ void ULandscapeSplinesComponent::RequestSplineLayerUpdate()
 	if (IsValidChecked(this))
 	{
 		ILandscapeSplineInterface* SplineOwner = GetSplineOwner();
-		SplineOwner->GetLandscapeInfo()->RequestSplineLayerUpdate();
+		//  This can get called on a deleted spline, when redoing a deletion.
+		if (SplineOwner)
+		{
+			SplineOwner->GetLandscapeInfo()->RequestSplineLayerUpdate();
+		}
 	}
 }
 
@@ -3551,8 +3555,11 @@ void ULandscapeSplineSegment::UpdateSplineMeshSelectionState()
 {
 	for (auto& LocalMeshComponent : LocalMeshComponents)
 	{
-		LocalMeshComponent->bSelected = bSelected;
-		LocalMeshComponent->PushSelectionToProxy();
+		if (LocalMeshComponent)
+		{
+			LocalMeshComponent->bSelected = bSelected;
+			LocalMeshComponent->PushSelectionToProxy();
+		}
 	}
 
 	auto ForeignMeshComponentsMap = GetForeignMeshComponents();
