Landscape Editor - Do not free allocations during non-edit layer auto-conversion in CopyOldDataToDefaultLayer
* There is no reason to clear old data, the next UpdateLayersContent will reallocate restoring the freed allocations. This allows the conversion to happen in PIE/Cook since edit layers can be enabled and allocations will still be valid, even when UpdateLayersContent does not run

#FailedValidations Preflight

diff --git a/Engine/Source/Runtime/Landscape/Private/LandscapeEditLayers.cpp b/Engine/Source/Runtime/Landscape/Private/LandscapeEditLayers.cpp
--- a/Engine/Source/Runtime/Landscape/Private/LandscapeEditLayers.cpp
+++ b/Engine/Source/Runtime/Landscape/Private/LandscapeEditLayers.cpp
@@ -2892,7 +2892,6 @@ void ALandscape::CopyOldDataToDefaultLayer(ALandscapeProxy* InProxy)
 
 	TMap<UTexture2D*, FWeightmapTextureData> ProcessedWeightmaps;
 	TArray<UTexture2D*> ProcessedHeightmaps;
-	TArray<ULandscapeComponent*> WeightmapsComponentsToCleanup;
 
 	for (ULandscapeComponent* Component : InProxy->LandscapeComponents)
 	{
@@ -2926,8 +2925,10 @@ void ALandscape::CopyOldDataToDefaultLayer(ALandscapeProxy* InProxy)
 				DefaultLayerHeightmap->UpdateResource();
 			}
 
-			// Weightmaps
-			WeightmapsComponentsToCleanup.Add(Component);
+			// Weightmaps:
+			// Existing weightmap allocations are not freed after being copied
+			// Processes like PIE/Cook can run the non-edit layers conversion but will not run UpdateLayersContent/ReallocateWeightmaps
+			// Instead of having freed allocations, leave the data since the next ReallocateWeightmaps will restore a similar state anyways
 
 			const TArray<UTexture2D*>& ComponentWeightmapTextures = Component->GetWeightmapTextures();
 			const TArray<FWeightmapLayerAllocationInfo>& ComponentLayerAllocations = Component->GetWeightmapLayerAllocations();
@@ -3010,16 +3011,6 @@ void ALandscape::CopyOldDataToDefaultLayer(ALandscapeProxy* InProxy)
 		}
 	}
 
-	for (ULandscapeComponent* Component : WeightmapsComponentsToCleanup)
-	{
-		TArray<FWeightmapLayerAllocationInfo>& ComponentLayerAllocations = Component->GetWeightmapLayerAllocations();
-
-		for (FWeightmapLayerAllocationInfo& Allocation : ComponentLayerAllocations)
-		{
-			Allocation.Free();
-		}
-	}
-
 	// Data is only copied for proxies with valid components
 	// After copying, ensure the InProxy's components are valid and have an edit layer in their LayersData map (Default edit layer)
 	check(!InProxy->LandscapeComponents.IsEmpty() && (InProxy->LandscapeComponents[0] != nullptr) && InProxy->LandscapeComponents[0]->HasLayersData());
