Landscape : removed global and local merge code and deprecated anything that was only available for this.

diff --git a/Engine/Source/Runtime/Landscape/Private/LandscapeEditLayers.cpp b/Engine/Source/Runtime/Landscape/Private/LandscapeEditLayers.cpp
--- a/Engine/Source/Runtime/Landscape/Private/LandscapeEditLayers.cpp
+++ b/Engine/Source/Runtime/Landscape/Private/LandscapeEditLayers.cpp
@@ -15,7 +15,6 @@
 #include "LandscapeDataAccess.h"
 #include "LandscapePrivate.h"
 #include "LandscapeEditReadback.h"
-#include "LandscapeEditResources.h"
 #include "LandscapeEditResourcesSubsystem.h"
 #include "LandscapeEditLayerRenderer.h"
 #include "LandscapeEditLayerMergeRenderContext.h"
@@ -71,7 +70,6 @@
 #include "LandscapeRender.h"
 #include "LandscapeInfoMap.h"
 #include "Misc/MessageDialog.h"
-#include "GlobalMergeLegacySupportUtil.h" // ILandscapeBrushRenderCallAdapter_GlobalMergeLegacySupport
 #include "GameFramework/WorldSettings.h"
 #include "UObject/UObjectThreadContext.h"
 #include "LandscapeSplinesComponent.h"
@@ -132,15 +130,23 @@
 	RenderCaptureLayersNextWeightmapDraws,
 	TEXT("Trigger N render captures during the next weightmap draw calls."));
 
-static TAutoConsoleVariable<int32> CVarOutputLayersRTContent(
-	TEXT("landscape.OutputLayersRTContent"),
-	0,
-	TEXT("This will output the content of render target. This is used for debugging only."));
+UE_DEPRECATED(5.7, "The CVar has been fully deprecated.")
+FAutoConsoleVariableDeprecated CVarOutputLayersRTContent(
+	/*Name = */TEXT("landscape.OutputLayersRTContent"),
+	/*ShadowName = */TEXT("landscape.OutputLayersRTContent.Deprecated"),
+	/*DeprecatedAtVersion = */TEXT("5.7"),
+	/*UsageBehavior = */EShadowCVarBehavior::Warn,
+	/*LookupFailureBehavior = */EShadowCVarBehavior::Warn,
+	TEXT("landscape.OutputLayersRTContent is deprecated. This was only used in Global merge. Batched merge (landscape.EditLayersLocalMerge.Enable 2) is the only merge technique available now."));
 
-static TAutoConsoleVariable<int32> CVarOutputLayersWeightmapsRTContent(
-	TEXT("landscape.OutputLayersWeightmapsRTContent"),
-	0,
-	TEXT("This will output the content of render target used for weightmap. This is used for debugging only."));
+UE_DEPRECATED(5.7, "The CVar has been fully deprecated.")
+FAutoConsoleVariableDeprecated CVarOutputLayersWeightmapsRTContent(
+	/*Name = */TEXT("landscape.OutputLayersWeightmapsRTContent"),
+	/*ShadowName = */TEXT("landscape.OutputLayersWeightmapsRTContent.Deprecated"),
+	/*DeprecatedAtVersion = */TEXT("5.7"),
+	/*UsageBehavior = */EShadowCVarBehavior::Warn,
+	/*LookupFailureBehavior = */EShadowCVarBehavior::Warn,
+	TEXT("landscape.OutputLayersWeightmapsRTContent is deprecated. This was only used in Global merge. Batched merge (landscape.EditLayersLocalMerge.Enable 2) is the only merge technique available now."));
 
 static TAutoConsoleVariable<int32> CVarLandscapeSimulatePhysics(
 	TEXT("landscape.SimulatePhysics"),
@@ -267,76 +273,74 @@
 TAutoConsoleVariable<bool> CVarSilenceMergeBatchResolutionWarning(
 	TEXT("landscape.BatchedMerge.SilenceResolutionWarning"),
 	false,
-	TEXT("When true, don't warn about about exceeding batch merge resolution from landscape.EditLayersLocalMerge.MaxResolutionPerRenderBatch"));
+	TEXT("When true, don't warn about about exceeding batch merge resolution from landscape.BatchedMerge.MaxResolutionPerRenderBatch"));
 
-// TODO:  remove with IntermediateRender.
-TAutoConsoleVariable<bool> CVarLandscapeSelectiveRenderEditLayers(
-	TEXT("landscape.BatchedMerge.SelectiveRenderEditLayers"),
-	true,
-	TEXT("When true, use SelectiveRenderEditLayers codepath instead of intermediate renders"));
+UE_DEPRECATED(5.7, "The CVar has been fully deprecated.")
+FAutoConsoleVariableDeprecated CVarLandscapeEditLayersLocalMerge_Deprecated(
+	/*Name = */TEXT("landscape.EditLayersLocalMerge.Enable"), 
+	/*ShadowName = */TEXT("landscape.EditLayersLocalMerge.Enable.Deprecated"),
+	/*DeprecatedAtVersion = */TEXT("5.7"), 
+	/*UsageBehavior = */EShadowCVarBehavior::Warn, 
+	/*LookupFailureBehavior = */EShadowCVarBehavior::Warn, 
+	TEXT("landscape.EditLayersLocalMerge.Enable is deprecated. Batched merge (landscape.EditLayersLocalMerge.Enable 2) is the only merge technique available now."));
 
+UE_DEPRECATED(5.7, "The CVar has been fully deprecated.")
+FAutoConsoleVariableDeprecated CVarLandscapeEditLayersMaxComponentsPerHeightmapResolveBatch_Deprecated(
+	/*Name = */TEXT("landscape.EditLayersLocalMerge.MaxComponentsPerHeightmapResolveBatch"),
+	/*ShadowName = */TEXT("landscape.EditLayersLocalMerge.MaxComponentsPerHeightmapResolveBatch.Deprecated"),
+	/*DeprecatedAtVersion = */TEXT("5.7"),
+	/*UsageBehavior = */EShadowCVarBehavior::Warn,
+	/*LookupFailureBehavior = */EShadowCVarBehavior::Warn,
+	TEXT("landscape.EditLayersLocalMerge.MaxComponentsPerHeightmapResolveBatch is deprecated. This was only used in Local merge. Batched merge (landscape.EditLayersLocalMerge.Enable 2) is the only merge technique available now."));
 
-void OnLandscapeEditLayersLocalMergeChanged(IConsoleVariable* CVar)
-{
-	for (TObjectIterator<UWorld> It(/*AdditionalExclusionFlags = */RF_ClassDefaultObject, /*bIncludeDerivedClasses = */true, /*InInternalExclusionFlags = */EInternalObjectFlags::Garbage); It; ++It)
-	{
-		UWorld* CurrentWorld = *It;
-		if (!CurrentWorld->IsGameWorld())
-		{
-			ULandscapeInfoMap& LandscapeInfoMap = ULandscapeInfoMap::GetLandscapeInfoMap(CurrentWorld);
-			for (auto& Pair : LandscapeInfoMap.Map)
-			{
-				if (ULandscapeInfo* LandscapeInfo = Pair.Value)
-				{
-					if (ALandscape* Landscape = LandscapeInfo->LandscapeActor.Get())
-					{
-						Landscape->RequestLayersInitialization(true);
-					}
-				}
-			}
-		}
-	}
-}
+UE_DEPRECATED(5.7, "The CVar has been fully deprecated.")
+FAutoConsoleVariableDeprecated CVarLandscapeEditLayersMaxComponentsPerWeightmapResolveBatch_Deprecated(
+	/*Name = */TEXT("landscape.EditLayersLocalMerge.MaxComponentsPerWeightmapResolveBatch"),
+	/*ShadowName = */TEXT("landscape.EditLayersLocalMerge.MaxComponentsPerWeightmapResolveBatch.Deprecated"),
+	/*DeprecatedAtVersion = */TEXT("5.7"),
+	/*UsageBehavior = */EShadowCVarBehavior::Warn,
+	/*LookupFailureBehavior = */EShadowCVarBehavior::Warn,
+	TEXT("landscape.EditLayersLocalMerge.MaxComponentsPerWeightmapResolveBatch is deprecated. This was only used in Local merge. Batched merge (landscape.EditLayersLocalMerge.Enable 2) is the only merge technique available now."));
 
-int32 LandscapeEditLayersLocalMerge = 2;
-static FAutoConsoleVariableRef CVarLandscapeEditLayersLocalMerge(
-	TEXT("landscape.EditLayersLocalMerge.Enable"),
-	LandscapeEditLayersLocalMerge,
-	TEXT("Setting this to 1 will allow the local merge algorithm (that merges layers at the landscape component level) to be used on landscapes that support it. This is a temporary measure while waiting for non-compatible landscapes to be deprecated.\n"
-		"Setting this to 2 will allow the batched merge algorithm (that merges layers in batches). Supports all landscape types. This is also a temporary measure but will be the default once it's stable."),
-	FConsoleVariableDelegate::CreateStatic(OnLandscapeEditLayersLocalMergeChanged));
-
-TAutoConsoleVariable<int32> CVarLandscapeEditLayersMaxComponentsPerHeightmapResolveBatch(
-	TEXT("landscape.EditLayersLocalMerge.MaxComponentsPerHeightmapResolveBatch"),
-	16,
-	TEXT("Number of components being rendered in a single batch when resolving heightmaps. The higher the number, the more heightmaps can be resolved in a single batch (and the higher the GPU memory consumption since more transient textures will be needed in memory at a time)"),
-	ECVF_RenderThreadSafe);
-
-TAutoConsoleVariable<int32> CVarLandscapeEditLayersMaxComponentsPerWeightmapResolveBatch(
-	TEXT("landscape.EditLayersLocalMerge.MaxComponentsPerWeightmapResolveBatch"),
-	16,
-	TEXT("Number of components being rendered in a single batch when resolving weightmaps. The higher the number, the more weightmaps can be resolved in a single batch (and the higher the GPU memory consumption since more transient textures will be needed in memory at a time)"),
-	ECVF_RenderThreadSafe);
-
-TAutoConsoleVariable<int32> CVarLandscapeEditLayersMaxResolutionPerRenderBatch(
-	TEXT("landscape.EditLayersLocalMerge.MaxResolutionPerRenderBatch"),
+TAutoConsoleVariable<int32> CVarLandscapeBatchedMergeMaxResolutionPerRenderBatch(
+	TEXT("landscape.BatchedMerge.MaxResolutionPerRenderBatch"),
 	1024,
 	TEXT("Maximum supported resolution for merging edit layers in a single batch. The higher the number, the more components can be resolved in a single batch (and the higher the GPU memory consumption since merging requires several temporary textures)"),
 	ECVF_RenderThreadSafe);
+UE_DEPRECATED(5.7, "The CVar has been renamed CVarLandscapeBatchedMergeMaxResolutionPerRenderBatch.")
+FAutoConsoleVariableDeprecated CVarLandscapeEditLayersMaxResolutionPerRenderBatch_Deprecated(
+	/*Name = */TEXT("landscape.EditLayersLocalMerge.MaxResolutionPerRenderBatch"),
+	/*ShadowName = */TEXT("landscape.BatchedMerge.MaxResolutionPerRenderBatch"),
+	/*DeprecatedAtVersion = */TEXT("5.7"),
+	/*UsageBehavior = */EShadowCVarBehavior::Warn,
+	/*LookupFailureBehavior = */EShadowCVarBehavior::Warn,
+	TEXT("The CVar has been renamed landscape.BatchedMerge.MaxResolutionPerRenderBatch."));
 
-TAutoConsoleVariable<int32> CVarLandscapeEditLayersClearBeforeEachWriteToScratch(
-	TEXT("landscape.EditLayersLocalMerge.ClearBeforeEachWriteToScratch"),
+TAutoConsoleVariable<int32> CVarLandscapeBatchedMergeClearBeforeEachWriteToScratch(
+	TEXT("landscape.BatchedMerge.ClearBeforeEachWriteToScratch"),
 	0,
 	TEXT("Debug to help with RenderDoc debugging : clear each time we're about to write on a scratch render target (since those are reused and can be used to write RTs of different resolutions"),
 	ECVF_RenderThreadSafe);
+UE_DEPRECATED(5.7, "The CVar has been renamed CVarLandscapeBatchedMergeClearBeforeEachWriteToScratch.")
+FAutoConsoleVariableDeprecated CVarLandscapeEditLayersClearBeforeEachWriteToScratch_Deprecated(
+	/*Name = */TEXT("landscape.EditLayersLocalMerge.ClearBeforeEachWriteToScratch"),
+	/*ShadowName = */TEXT("landscape.BatchedMerge.ClearBeforeEachWriteToScratch"),
+	/*DeprecatedAtVersion = */TEXT("5.7"),
+	/*UsageBehavior = */EShadowCVarBehavior::Warn,
+	/*LookupFailureBehavior = */EShadowCVarBehavior::Warn,
+	TEXT("The CVar has been renamed landscape.BatchedMerge.ClearBeforeEachWriteToScratch."));
 
 static void LandscapeDumpSelectiveLayerRender(const TArray<FString>& Args);
-
 FAutoConsoleCommand CmdLandscapeDumpSelectiveLayerRender(
 	TEXT("landscape.DumpSelectiveLayerRender"),
 	TEXT("Perform a selective layer render of a landscape and write the output to an image file."),
-	FConsoleCommandWithArgsDelegate::CreateStatic(&LandscapeDumpSelectiveLayerRender)
-);
+	FConsoleCommandWithArgsDelegate::CreateStatic(&LandscapeDumpSelectiveLayerRender));
 
+static void LandscapeForceLayersFullUpdate(const TArray<FString>& Args);
+FAutoConsoleCommand CmdLandscapeForceLayersFullUpdate(
+	TEXT("landscape.ForceLayersFullUpdate"),
+	TEXT("Trigger an edit layers update for all (no argument) or a specific set of landscapes (partial name match based on the strings passed in argument).\n"),
+	FConsoleCommandWithArgsDelegate::CreateStatic(&LandscapeForceLayersFullUpdate));
+
 struct FLandscapeDirty
 {
@@ -620,81 +624,27 @@
 	}
 }
 
-// Vertex format and vertex buffer
+class FLandscapeCopyTextureVS : public FGlobalShader
+{
+public:
+	DECLARE_GLOBAL_SHADER(FLandscapeCopyTextureVS);
 
-struct FLandscapeLayersVertex
-{
-	FVector2f Position;
-	FVector2f UV;
+	static bool ShouldCompilePermutation(const FGlobalShaderPermutationParameters& InParameters)
+	{
+		return UE::Landscape::DoesPlatformSupportEditLayers(InParameters.Platform);
+	}
+
+	FLandscapeCopyTextureVS()
+	{};
+
+	FLandscapeCopyTextureVS(const ShaderMetaType::CompiledShaderInitializerType& Initializer)
+		: FGlobalShader(Initializer)
+	{}
 };
 
-struct FLandscapeLayersTriangle
+class FLandscapeCopyTexturePS : public FGlobalShader
 {
-	FLandscapeLayersVertex V0;
-	FLandscapeLayersVertex V1;
-	FLandscapeLayersVertex V2;
-};
-
-class FLandscapeLayersVertexDeclaration : public FRenderResource
-{
-public:
-	FVertexDeclarationRHIRef VertexDeclarationRHI;
-
-	/** Destructor. */
-	virtual ~FLandscapeLayersVertexDeclaration() {}
-
-	virtual void InitRHI(FRHICommandListBase& RHICmdList)
-	{
-		FVertexDeclarationElementList Elements;
-		constexpr uint16 Stride = sizeof(FLandscapeLayersVertex);
-		Elements.Add(FVertexElement(0, STRUCT_OFFSET(FLandscapeLayersVertex, Position), VET_Float2, 0, Stride));
-		Elements.Add(FVertexElement(0, STRUCT_OFFSET(FLandscapeLayersVertex, UV), VET_Float2, 1, Stride));
-		VertexDeclarationRHI = PipelineStateCache::GetOrCreateVertexDeclaration(Elements);
-	}
-
-	virtual void ReleaseRHI()
-	{
-		VertexDeclarationRHI.SafeRelease();
-	}
-};
-
-class FLandscapeLayersVertexBuffer : public FVertexBuffer
-{
-public:
-	void Init(const TArray<FLandscapeLayersTriangle>& InTriangleList)
-	{
-		TriangleList = InTriangleList;
-	}
-
-private:
-
-	/** Initialize the RHI for this rendering resource */
-	void InitRHI(FRHICommandListBase& RHICmdList) override
-	{
-		TArray<FLandscapeLayersVertex> Vertices;
-		Vertices.SetNumUninitialized(TriangleList.Num() * 3);
-
-		for (int32 i = 0; i < TriangleList.Num(); ++i)
-		{
-			Vertices[i * 3 + 0] = TriangleList[i].V0;
-			Vertices[i * 3 + 1] = TriangleList[i].V1;
-			Vertices[i * 3 + 2] = TriangleList[i].V2;
-		}
-
-		// Create vertex buffer. Fill buffer with initial data upon creation
-		VertexBufferRHI = UE::RHIResourceUtils::CreateVertexBufferFromArray(RHICmdList, TEXT("FLandscapeLayersVertexBuffer"), EBufferUsageFlags::Static, MakeConstArrayView(Vertices));
-	}
-
-	TArray<FLandscapeLayersTriangle> TriangleList;
-};
-
-// ----------------------------------------------------------------------------------
-// /Engine/Private/LandscapeLayersVS.usf shaders
-
-class FLandscapeLayersVS : public FGlobalShader
-{
-	DECLARE_GLOBAL_SHADER(FLandscapeLayersVS)
-
+	DECLARE_GLOBAL_SHADER(FLandscapeCopyTexturePS);
 public:
 
 	static bool ShouldCompilePermutation(const FGlobalShaderPermutationParameters& Parameters)
@@ -707,337 +657,6 @@
 		UE::Landscape::ModifyShaderCompilerEnvironmentForDebug(OutEnvironment);
 	}
 
-	FLandscapeLayersVS(const ShaderMetaType::CompiledShaderInitializerType& Initializer)
-		: FGlobalShader(Initializer)
-	{
-		TransformParam.Bind(Initializer.ParameterMap, TEXT("Transform"), SPF_Mandatory);
-	}
-
-	FLandscapeLayersVS()
-	{}
-
-	void SetParameters(FRHIBatchedShaderParameters& BatchedParameters, const FMatrix44f& InProjectionMatrix)
-	{
-		SetShaderValue(BatchedParameters, TransformParam, InProjectionMatrix);
-	}
-
-private:
-	LAYOUT_FIELD(FShaderParameter, TransformParam);
-};
-
-IMPLEMENT_GLOBAL_SHADER(FLandscapeLayersVS, "/Engine/Private/LandscapeLayersVS.usf", "VSMain", SF_Vertex);
-
-// ----------------------------------------------------------------------------------
-// /Engine/Private/LandscapeLayersPS.usf shaders
-
-struct FLandscapeLayersHeightmapShaderParameters
-{
-	FLandscapeLayersHeightmapShaderParameters()
-		: ReadHeightmap1(nullptr)
-		, ReadHeightmap2(nullptr)
-		, HeightmapSize(0, 0)
-		, ApplyLayerModifiers(false)
-		, SetAlphaOne(false)
-		, LayerAlpha(1.0f)
-		, LayerVisible(true)
-		, LayerBlendMode(LSBM_AdditiveBlend)
-		, GenerateNormals(false)
-		, GridSize(0.0f, 0.0f, 0.0f)
-		, CurrentMipSize(0, 0)
-		, ParentMipSize(0, 0)
-		, CurrentMipComponentVertexCount(0)
-	{}
-
-	UTexture* ReadHeightmap1;
-	UTexture* ReadHeightmap2;
-	FIntPoint HeightmapSize;
-	bool ApplyLayerModifiers;
-	bool SetAlphaOne;
-	float LayerAlpha;
-	bool LayerVisible;
-	ELandscapeBlendMode LayerBlendMode;
-	bool GenerateNormals;
-	FVector GridSize;
-	FIntPoint CurrentMipSize;
-	FIntPoint ParentMipSize;
-	int32 CurrentMipComponentVertexCount;
-};
-
-class FLandscapeLayersHeightmapPS : public FGlobalShader
-{
-	DECLARE_GLOBAL_SHADER(FLandscapeLayersHeightmapPS);
-public:
-
-	static bool ShouldCompilePermutation(const FGlobalShaderPermutationParameters& Parameters)
-	{
-		return UE::Landscape::DoesPlatformSupportEditLayers(Parameters.Platform);
-	}
-
-	static void ModifyCompilationEnvironment(const FGlobalShaderPermutationParameters& Parameters, FShaderCompilerEnvironment& OutEnvironment)
-	{
-		UE::Landscape::ModifyShaderCompilerEnvironmentForDebug(OutEnvironment);
-	}
-
-	FLandscapeLayersHeightmapPS(const ShaderMetaType::CompiledShaderInitializerType& Initializer)
-		: FGlobalShader(Initializer)
-	{
-		ReadTexture1Param.Bind(Initializer.ParameterMap, TEXT("ReadTexture1"));
-		ReadTexture2Param.Bind(Initializer.ParameterMap, TEXT("ReadTexture2"));
-		ReadTexture1SamplerParam.Bind(Initializer.ParameterMap, TEXT("ReadTexture1Sampler"));
-		ReadTexture2SamplerParam.Bind(Initializer.ParameterMap, TEXT("ReadTexture2Sampler"));
-
-		LayerInfoParam.Bind(Initializer.ParameterMap, TEXT("LayerInfo"));
-		OutputConfigParam.Bind(Initializer.ParameterMap, TEXT("OutputConfig"));
-		TextureSizeParam.Bind(Initializer.ParameterMap, TEXT("TextureSize"));
-		LandscapeGridScaleParam.Bind(Initializer.ParameterMap, TEXT("LandscapeGridScale"));
-		ComponentVertexCountParam.Bind(Initializer.ParameterMap, TEXT("CurrentMipComponentVertexCount"));
-	}
-
-	FLandscapeLayersHeightmapPS()
-	{}
-
-	void SetParameters(FRHIBatchedShaderParameters& BatchedParameters, const FLandscapeLayersHeightmapShaderParameters& InParams)
-	{
-		SetTextureParameter(BatchedParameters, ReadTexture1Param, ReadTexture1SamplerParam, TStaticSamplerState<SF_Point, AM_Clamp, AM_Clamp, AM_Clamp>::GetRHI(), InParams.ReadHeightmap1->GetResource()->TextureRHI);
-		SetTextureParameter(BatchedParameters, ReadTexture2Param, ReadTexture2SamplerParam, TStaticSamplerState<SF_Point, AM_Clamp, AM_Clamp, AM_Clamp>::GetRHI(),
-			InParams.ReadHeightmap2 != nullptr ? InParams.ReadHeightmap2->GetResource()->TextureRHI : GWhiteTexture->TextureRHI);
-
-		FVector4f LayerInfo(InParams.LayerAlpha, InParams.LayerVisible ? 1.0f : 0.0f, InParams.LayerBlendMode == LSBM_AlphaBlend ? 1.0f : 0.f, 0.f);
-		FVector4f OutputConfig(InParams.ApplyLayerModifiers ? 1.0f : 0.0f, InParams.SetAlphaOne ? 1.0f : 0.0f,
-			InParams.ReadHeightmap2 ? 1.0f : 0.0f,
-			InParams.GenerateNormals ? 1.0f : 0.0f);
-		FVector2f TextureSize(static_cast<float>(InParams.HeightmapSize.X), static_cast<float>(InParams.HeightmapSize.Y));
-
-		SetShaderValue(BatchedParameters, LayerInfoParam, LayerInfo);
-		SetShaderValue(BatchedParameters, OutputConfigParam, OutputConfig);
-		SetShaderValue(BatchedParameters, TextureSizeParam, TextureSize);
-		SetShaderValue(BatchedParameters, LandscapeGridScaleParam, FVector3f(InParams.GridSize));
-		SetShaderValue(BatchedParameters, ComponentVertexCountParam, (float)InParams.CurrentMipComponentVertexCount);
-	}
-
-private:
-	LAYOUT_FIELD(FShaderResourceParameter, ReadTexture1Param);
-	LAYOUT_FIELD(FShaderResourceParameter, ReadTexture2Param);
-	LAYOUT_FIELD(FShaderResourceParameter, ReadTexture1SamplerParam);
-	LAYOUT_FIELD(FShaderResourceParameter, ReadTexture2SamplerParam);
-	LAYOUT_FIELD(FShaderParameter, LayerInfoParam);
-	LAYOUT_FIELD(FShaderParameter, OutputConfigParam);
-	LAYOUT_FIELD(FShaderParameter, TextureSizeParam);
-	LAYOUT_FIELD(FShaderParameter, LandscapeGridScaleParam);
-	LAYOUT_FIELD(FShaderParameter, ComponentVertexCountParam);
-};
-
-IMPLEMENT_GLOBAL_SHADER(FLandscapeLayersHeightmapPS, "/Engine/Private/LandscapeLayersPS.usf", "PSHeightmapMain", SF_Pixel);
-
-class FLandscapeLayersHeightmapMipsPS : public FGlobalShader
-{
-	DECLARE_GLOBAL_SHADER(FLandscapeLayersHeightmapMipsPS);
-public:
-
-	static bool ShouldCompilePermutation(const FGlobalShaderPermutationParameters& Parameters)
-	{
-		return UE::Landscape::DoesPlatformSupportEditLayers(Parameters.Platform);
-	}
-
-	static void ModifyCompilationEnvironment(const FGlobalShaderPermutationParameters& Parameters, FShaderCompilerEnvironment& OutEnvironment)
-	{
-		UE::Landscape::ModifyShaderCompilerEnvironmentForDebug(OutEnvironment);
-	}
-
-	FLandscapeLayersHeightmapMipsPS(const ShaderMetaType::CompiledShaderInitializerType& Initializer)
-		: FGlobalShader(Initializer)
-	{
-		ReadTexture1Param.Bind(Initializer.ParameterMap, TEXT("ReadTexture1"));
-		ReadTexture1SamplerParam.Bind(Initializer.ParameterMap, TEXT("ReadTexture1Sampler"));
-		CurrentMipSizeParam.Bind(Initializer.ParameterMap, TEXT("CurrentMipTextureSize"));
-		ParentMipSizeParam.Bind(Initializer.ParameterMap, TEXT("ParentMipTextureSize"));
-		CurrentMipComponentVertexCountParam.Bind(Initializer.ParameterMap, TEXT("CurrentMipComponentVertexCount"));
-	}
-
-	FLandscapeLayersHeightmapMipsPS()
-	{}
-
-	void SetParameters(FRHIBatchedShaderParameters& BatchedParameters, const FLandscapeLayersHeightmapShaderParameters& InParams)
-	{
-		SetTextureParameter(BatchedParameters, ReadTexture1Param, ReadTexture1SamplerParam, TStaticSamplerState<SF_Point, AM_Clamp, AM_Clamp, AM_Clamp>::GetRHI(), InParams.ReadHeightmap1->GetResource()->TextureRHI);
-
-		SetShaderValue(BatchedParameters, CurrentMipSizeParam, FVector2f(static_cast<float>(InParams.CurrentMipSize.X), static_cast<float>(InParams.CurrentMipSize.Y)));
-		SetShaderValue(BatchedParameters, ParentMipSizeParam, FVector2f(static_cast<float>(InParams.ParentMipSize.X), static_cast<float>(InParams.ParentMipSize.Y)));
-		SetShaderValue(BatchedParameters, CurrentMipComponentVertexCountParam, (float)InParams.CurrentMipComponentVertexCount);
-	}
-
-private:
-	LAYOUT_FIELD(FShaderResourceParameter, ReadTexture1Param);
-	LAYOUT_FIELD(FShaderResourceParameter, ReadTexture1SamplerParam);
-	LAYOUT_FIELD(FShaderParameter, CurrentMipSizeParam);
-	LAYOUT_FIELD(FShaderParameter, ParentMipSizeParam);
-	LAYOUT_FIELD(FShaderParameter, CurrentMipComponentVertexCountParam);
-};
-
-IMPLEMENT_GLOBAL_SHADER(FLandscapeLayersHeightmapMipsPS, "/Engine/Private/LandscapeLayersPS.usf", "PSHeightmapMainMips", SF_Pixel);
-
-struct FLandscapeLayersWeightmapShaderParameters
-{
-	FLandscapeLayersWeightmapShaderParameters()
-		: ReadWeightmap1(nullptr)
-		, ReadWeightmap2(nullptr)
-		, ApplyLayerModifiers(false)
-		, LayerAlpha(1.0f)
-		, LayerVisible(true)
-		, LayerBlendMode(LSBM_AdditiveBlend)
-		, OutputAsSubstractive(false)
-		, CurrentMipSize(0, 0)
-		, ParentMipSize(0, 0)
-		, CurrentMipComponentVertexCount(0)
-	{}
-
-	UTexture* ReadWeightmap1;
-	UTexture* ReadWeightmap2;
-	bool ApplyLayerModifiers;
-	float LayerAlpha;
-	bool LayerVisible;
-	ELandscapeBlendMode LayerBlendMode;
-	bool OutputAsSubstractive;
-	FIntPoint CurrentMipSize;
-	FIntPoint ParentMipSize;
-	int32 CurrentMipComponentVertexCount;
-};
-
-class FLandscapeLayersWeightmapPS : public FGlobalShader
-{
-	DECLARE_GLOBAL_SHADER(FLandscapeLayersWeightmapPS);
-public:
-
-	static bool ShouldCompilePermutation(const FGlobalShaderPermutationParameters& Parameters)
-	{
-		return UE::Landscape::DoesPlatformSupportEditLayers(Parameters.Platform);
-	}
-
-	static void ModifyCompilationEnvironment(const FGlobalShaderPermutationParameters& Parameters, FShaderCompilerEnvironment& OutEnvironment)
-	{
-		UE::Landscape::ModifyShaderCompilerEnvironmentForDebug(OutEnvironment);
-	}
-
-	FLandscapeLayersWeightmapPS(const ShaderMetaType::CompiledShaderInitializerType& Initializer)
-		: FGlobalShader(Initializer)
-	{
-		ReadTexture1Param.Bind(Initializer.ParameterMap, TEXT("ReadTexture1"));
-		ReadTexture2Param.Bind(Initializer.ParameterMap, TEXT("ReadTexture2"));
-		ReadTexture1SamplerParam.Bind(Initializer.ParameterMap, TEXT("ReadTexture1Sampler"));
-		ReadTexture2SamplerParam.Bind(Initializer.ParameterMap, TEXT("ReadTexture2Sampler"));
-		LayerInfoParam.Bind(Initializer.ParameterMap, TEXT("LayerInfo"));
-		OutputConfigParam.Bind(Initializer.ParameterMap, TEXT("OutputConfig"));
-		ComponentVertexCountParam.Bind(Initializer.ParameterMap, TEXT("CurrentMipComponentVertexCount"));
-	}
-
-	FLandscapeLayersWeightmapPS()
-	{}
-
-	void SetParameters(FRHIBatchedShaderParameters& BatchedParameters, const FLandscapeLayersWeightmapShaderParameters& InParams)
-	{
-		SetTextureParameter(BatchedParameters, ReadTexture1Param, ReadTexture1SamplerParam, TStaticSamplerState<SF_Point, AM_Clamp, AM_Clamp, AM_Clamp>::GetRHI(), InParams.ReadWeightmap1->GetResource()->TextureRHI);
-		SetTextureParameter(BatchedParameters, ReadTexture2Param, ReadTexture2SamplerParam, TStaticSamplerState<SF_Point, AM_Clamp, AM_Clamp, AM_Clamp>::GetRHI(), InParams.ReadWeightmap2 != nullptr ? InParams.ReadWeightmap2->GetResource()->TextureRHI : GWhiteTexture->TextureRHI);
-
-		FVector4f LayerInfo(InParams.LayerAlpha, InParams.LayerVisible ? 1.0f : 0.0f, InParams.LayerBlendMode == LSBM_AlphaBlend ? 1.0f : 0.f, 0.f);
-		FVector4f OutputConfig(InParams.ApplyLayerModifiers ? 1.0f : 0.0f, InParams.OutputAsSubstractive ? 1.0f : 0.0f, InParams.ReadWeightmap2 != nullptr ? 1.0f : 0.0f, 0.0f);
-
-		SetShaderValue(BatchedParameters, LayerInfoParam, LayerInfo);
-		SetShaderValue(BatchedParameters, OutputConfigParam, OutputConfig);
-		SetShaderValue(BatchedParameters, ComponentVertexCountParam, (float)InParams.CurrentMipComponentVertexCount);
-	}
-
-private:
-	LAYOUT_FIELD(FShaderResourceParameter, ReadTexture1Param);
-	LAYOUT_FIELD(FShaderResourceParameter, ReadTexture2Param);
-	LAYOUT_FIELD(FShaderResourceParameter, ReadTexture1SamplerParam);
-	LAYOUT_FIELD(FShaderResourceParameter, ReadTexture2SamplerParam);
-	LAYOUT_FIELD(FShaderParameter, LayerInfoParam);
-	LAYOUT_FIELD(FShaderParameter, OutputConfigParam);
-	LAYOUT_FIELD(FShaderParameter, ComponentVertexCountParam);
-};
-
-IMPLEMENT_GLOBAL_SHADER(FLandscapeLayersWeightmapPS, "/Engine/Private/LandscapeLayersPS.usf", "PSWeightmapMain", SF_Pixel);
-
-class FLandscapeLayersWeightmapMipsPS : public FGlobalShader
-{
-	DECLARE_GLOBAL_SHADER(FLandscapeLayersWeightmapMipsPS);
-public:
-
-	static bool ShouldCompilePermutation(const FGlobalShaderPermutationParameters& Parameters)
-	{
-		return UE::Landscape::DoesPlatformSupportEditLayers(Parameters.Platform);
-	}
-
-	static void ModifyCompilationEnvironment(const FGlobalShaderPermutationParameters& Parameters, FShaderCompilerEnvironment& OutEnvironment)
-	{
-		UE::Landscape::ModifyShaderCompilerEnvironmentForDebug(OutEnvironment);
-	}
-
-	FLandscapeLayersWeightmapMipsPS(const ShaderMetaType::CompiledShaderInitializerType& Initializer)
-		: FGlobalShader(Initializer)
-	{
-		ReadTexture1Param.Bind(Initializer.ParameterMap, TEXT("ReadTexture1"));
-		ReadTexture1SamplerParam.Bind(Initializer.ParameterMap, TEXT("ReadTexture1Sampler"));
-		CurrentMipSizeParam.Bind(Initializer.ParameterMap, TEXT("CurrentMipTextureSize"));
-		ParentMipSizeParam.Bind(Initializer.ParameterMap, TEXT("ParentMipTextureSize"));
-		CurrentMipComponentVertexCountParam.Bind(Initializer.ParameterMap, TEXT("CurrentMipComponentVertexCount"));
-	}
-
-	FLandscapeLayersWeightmapMipsPS()
-	{}
-
-	void SetParameters(FRHIBatchedShaderParameters& BatchedParameters, const FLandscapeLayersWeightmapShaderParameters& InParams)
-	{
-		SetTextureParameter(BatchedParameters, ReadTexture1Param, ReadTexture1SamplerParam, TStaticSamplerState<SF_Point, AM_Clamp, AM_Clamp, AM_Clamp>::GetRHI(), InParams.ReadWeightmap1->GetResource()->TextureRHI);
-
-		SetShaderValue(BatchedParameters, CurrentMipSizeParam, FVector2f(static_cast<float>(InParams.CurrentMipSize.X), static_cast<float>(InParams.CurrentMipSize.Y)));
-		SetShaderValue(BatchedParameters, ParentMipSizeParam, FVector2f(static_cast<float>(InParams.ParentMipSize.X), static_cast<float>(InParams.ParentMipSize.Y)));
-		SetShaderValue(BatchedParameters, CurrentMipComponentVertexCountParam, (float)InParams.CurrentMipComponentVertexCount);
-	}
-
-private:
-	LAYOUT_FIELD(FShaderResourceParameter, ReadTexture1Param);
-	LAYOUT_FIELD(FShaderResourceParameter, ReadTexture1SamplerParam);
-	LAYOUT_FIELD(FShaderParameter, CurrentMipSizeParam);
-	LAYOUT_FIELD(FShaderParameter, ParentMipSizeParam);
-	LAYOUT_FIELD(FShaderParameter, CurrentMipComponentVertexCountParam);
-};
-
-IMPLEMENT_GLOBAL_SHADER(FLandscapeLayersWeightmapMipsPS, "/Engine/Private/LandscapeLayersPS.usf", "PSWeightmapMainMips", SF_Pixel);
-
-class FLandscapeCopyTextureVS : public FGlobalShader
-{
-public:
-	DECLARE_GLOBAL_SHADER(FLandscapeCopyTextureVS);
-
-	static bool ShouldCompilePermutation(const FGlobalShaderPermutationParameters& InParameters)
-	{
-		return UE::Landscape::DoesPlatformSupportEditLayers(InParameters.Platform);
-	}
-
-	FLandscapeCopyTextureVS()
-	{};
-
-	FLandscapeCopyTextureVS(const ShaderMetaType::CompiledShaderInitializerType& Initializer)
-		: FGlobalShader(Initializer)
-	{}
-};
-
-class FLandscapeCopyTexturePS : public FGlobalShader
-{
-	DECLARE_GLOBAL_SHADER(FLandscapeCopyTexturePS);
-public:
-
-	static bool ShouldCompilePermutation(const FGlobalShaderPermutationParameters& Parameters)
-	{
-		return UE::Landscape::DoesPlatformSupportEditLayers(Parameters.Platform);
-	}
-
-	static void ModifyCompilationEnvironment(const FGlobalShaderPermutationParameters& Parameters, FShaderCompilerEnvironment& OutEnvironment)
-	{
-		UE::Landscape::ModifyShaderCompilerEnvironmentForDebug(OutEnvironment);
-	}
-
 	FLandscapeCopyTexturePS(const ShaderMetaType::CompiledShaderInitializerType& Initializer)
 		: FGlobalShader(Initializer)
 	{
@@ -1075,676 +694,6 @@
 IMPLEMENT_GLOBAL_SHADER(FLandscapeCopyTexturePS, "/Engine/Private/LandscapeLayersPS.usf", "CopyTexturePS", SF_Pixel);
 
 // ----------------------------------------------------------------------------------
-// /Engine/Private/LandscapeLayersCS.usf shaders
-
-int32 GLandscapeLayerWeightmapThreadGroupSizeX = 16;
-int32 GLandscapeLayerWeightmapThreadGroupSizeY = 16;
-
-struct FLandscapeLayerWeightmapExtractMaterialLayersComponentData
-{
-	FIntPoint ComponentVertexPosition;	// Section base converted to vertex instead of quad
-	uint32 DestinationPaintLayerIndex;	// correspond to which layer info object index the data should be stored in the texture 2d array
-	uint32 WeightmapChannelToProcess;	// correspond to which RGBA channel to process
-	FIntPoint AtlasTexturePositionOutput;	// This represent the location we will write layer information
-};
-
-class FLandscapeLayerWeightmapExtractMaterialLayersComputeShaderResource : public FRenderResource
-{
-public:
-	FLandscapeLayerWeightmapExtractMaterialLayersComputeShaderResource(const TArray<FLandscapeLayerWeightmapExtractMaterialLayersComponentData>& InComponentsData)
-		: OriginalComponentsData(InComponentsData)
-		, ComponentsDataCount(OriginalComponentsData.Num())
-	{}
-
-	~FLandscapeLayerWeightmapExtractMaterialLayersComputeShaderResource()
-	{
-		ComponentsData.SafeRelease();
-		ComponentsDataSRV.SafeRelease();
-	}
-
-	/** Called when the resource is initialized. This is only called by the rendering thread. */
-	virtual void InitRHI(FRHICommandListBase& RHICmdList) override
-	{
-		ComponentsData = UE::RHIResourceUtils::CreateBufferFromArray(
-			RHICmdList,
-			TEXT("FLandscapeLayerWeightmapExtractMaterialLayersComputeShaderResource"),
-			EBufferUsageFlags::StructuredBuffer | EBufferUsageFlags::ShaderResource | EBufferUsageFlags::Volatile,
-			MakeConstArrayView(OriginalComponentsData)
-		);
-		ComponentsDataSRV = RHICmdList.CreateShaderResourceView(ComponentsData, FRHIViewDesc::CreateBufferSRV().SetTypeFromBuffer(ComponentsData));
-	}
-
-	virtual void ReleaseRHI() override
-	{
-		ComponentsData.SafeRelease();
-		ComponentsDataSRV.SafeRelease();
-	}
-
-	int32 GetComponentsDataCount() const
-	{
-		return ComponentsDataCount;
-	}
-
-private:
-	friend class FLandscapeLayerWeightmapExtractMaterialLayersCS;
-
-	FBufferRHIRef ComponentsData;
-	FShaderResourceViewRHIRef ComponentsDataSRV;
-	TArray<FLandscapeLayerWeightmapExtractMaterialLayersComponentData> OriginalComponentsData;
-	int32 ComponentsDataCount;
-};
-
-struct FLandscapeLayerWeightmapExtractMaterialLayersComputeShaderParameters
-{
-	FLandscapeLayerWeightmapExtractMaterialLayersComputeShaderParameters()
-		: ComponentWeightmapResource(nullptr)
-		, ComputeShaderResource(nullptr)
-		, AtlasWeightmapsPerLayer(nullptr)
-		, ComponentSize(0)
-	{}
-
-	FLandscapeTexture2DResource* ComponentWeightmapResource;
-	FLandscapeLayerWeightmapExtractMaterialLayersComputeShaderResource* ComputeShaderResource;
-	FLandscapeTexture2DArrayResource* AtlasWeightmapsPerLayer;
-	uint32 ComponentSize;
-};
-
-class FLandscapeLayerWeightmapExtractMaterialLayersCS : public FGlobalShader
-{
-	DECLARE_GLOBAL_SHADER(FLandscapeLayerWeightmapExtractMaterialLayersCS);
-public:
-
-	static bool ShouldCompilePermutation(const FGlobalShaderPermutationParameters& Parameters)
-	{
-		return UE::Landscape::DoesPlatformSupportEditLayers(Parameters.Platform);
-	}
-
-	static void ModifyCompilationEnvironment(const FGlobalShaderPermutationParameters& Parameters, FShaderCompilerEnvironment& OutEnvironment)
-	{
-		OutEnvironment.SetDefine(TEXT("THREADGROUP_SIZEX"), GLandscapeLayerWeightmapThreadGroupSizeX);
-		OutEnvironment.SetDefine(TEXT("THREADGROUP_SIZEY"), GLandscapeLayerWeightmapThreadGroupSizeY);
-		UE::Landscape::ModifyShaderCompilerEnvironmentForDebug(OutEnvironment);
-	}
-
-	FLandscapeLayerWeightmapExtractMaterialLayersCS(const ShaderMetaType::CompiledShaderInitializerType& Initializer)
-		: FGlobalShader(Initializer)
-	{
-		ComponentWeightmapParam.Bind(Initializer.ParameterMap, TEXT("InComponentWeightMaps"));
-		AtlasPaintListsParam.Bind(Initializer.ParameterMap, TEXT("OutAtlasPaintLayers"));
-		ComponentsDataParam.Bind(Initializer.ParameterMap, TEXT("InExtractLayersComponentsData"));
-		ComponentSizeParam.Bind(Initializer.ParameterMap, TEXT("ComponentSize"));
-	}
-
-	FLandscapeLayerWeightmapExtractMaterialLayersCS()
-	{}
-
-	void SetParameters(FRHIBatchedShaderParameters& BatchedParameters, const FLandscapeLayerWeightmapExtractMaterialLayersComputeShaderParameters& InParams)
-	{
-		SetTextureParameter(BatchedParameters, ComponentWeightmapParam, InParams.ComponentWeightmapResource->TextureRHI);
-		SetUAVParameter(BatchedParameters, AtlasPaintListsParam, InParams.AtlasWeightmapsPerLayer->GetTextureUAV(/*InMipLevel = */0));
-		SetSRVParameter(BatchedParameters, ComponentsDataParam, InParams.ComputeShaderResource->ComponentsDataSRV);
-		SetShaderValue(BatchedParameters, ComponentSizeParam, InParams.ComponentSize);
-	}
-
-	void UnsetParameters(FRHIBatchedShaderUnbinds& BatchedUnbinds)
-	{
-		UnsetUAVParameter(BatchedUnbinds, AtlasPaintListsParam);
-	}
-
-private:
-	LAYOUT_FIELD(FShaderResourceParameter, ComponentWeightmapParam);
-	LAYOUT_FIELD(FShaderResourceParameter, AtlasPaintListsParam);
-	LAYOUT_FIELD(FShaderResourceParameter, ComponentsDataParam);
-	LAYOUT_FIELD(FShaderParameter, ComponentSizeParam);
-};
-
-IMPLEMENT_GLOBAL_SHADER(FLandscapeLayerWeightmapExtractMaterialLayersCS, "/Engine/Private/LandscapeLayersCS.usf", "ComputeWeightmapPerPaintLayer", SF_Compute);
-
-class FLandscapeLayerWeightmapExtractMaterialLayersCSDispatch_RenderThread
-{
-public:
-	FLandscapeLayerWeightmapExtractMaterialLayersCSDispatch_RenderThread(const FLandscapeLayerWeightmapExtractMaterialLayersComputeShaderParameters& InShaderParams)
-		: ShaderParams(InShaderParams)
-	{}
-
-	void ExtractLayers(FRHICommandListImmediate& InRHICmdList)
-	{
-		RHI_BREADCRUMB_EVENT_STAT(InRHICmdList, LandscapeLayers_ExtractLayers, "LandscapeLayers_ExtractLayers");
-		SCOPED_GPU_STAT(InRHICmdList, LandscapeLayers_ExtractLayers);
-
-		TShaderMapRef<FLandscapeLayerWeightmapExtractMaterialLayersCS> ComputeShader(GetGlobalShaderMap(GMaxRHIFeatureLevel));
-		SetComputePipelineState(InRHICmdList, ComputeShader.GetComputeShader());
-
-		SetShaderParametersLegacyCS(InRHICmdList, ComputeShader, ShaderParams);
-
-		// In case the CS is executed twice in a row, we need a barrier since we want to prevent UAV overlaps:
-		InRHICmdList.Transition(FRHITransitionInfo(ShaderParams.AtlasWeightmapsPerLayer->TextureRHI, ERHIAccess::UAVMask, ERHIAccess::UAVMask));
-
-		uint32 ThreadGroupCountX = FMath::CeilToInt((float)ShaderParams.ComponentSize / (float)GLandscapeLayerWeightmapThreadGroupSizeX);
-		uint32 ThreadGroupCountY = FMath::CeilToInt((float)ShaderParams.ComponentSize / (float)GLandscapeLayerWeightmapThreadGroupSizeY);
-		check(ThreadGroupCountX > 0 && ThreadGroupCountY > 0);
-
-		DispatchComputeShader(InRHICmdList, ComputeShader.GetShader(), ThreadGroupCountX, ThreadGroupCountY, ShaderParams.ComputeShaderResource->GetComponentsDataCount());
-
-		UnsetShaderParametersLegacyCS(InRHICmdList, ComputeShader);
-
-		ShaderParams.ComputeShaderResource->ReleaseResource();
-		delete ShaderParams.ComputeShaderResource;
-	}
-
-private:
-	FLandscapeLayerWeightmapExtractMaterialLayersComputeShaderParameters ShaderParams;
-};
-
-struct FLandscapeLayerWeightmapPackMaterialLayersComponentData
-{
-	int32 ComponentVertexPositionX[4];		// Section base converted to vertex instead of quad
-	int32 ComponentVertexPositionY[4];		// Section base converted to vertex instead of quad
-	int32 SourcePaintLayerIndex[4];			// correspond to which layer info object index the data should be loaded from the texture 2d array
-	int32 WeightmapChannelToProcess[4];		// correspond to which RGBA channel to process
-};
-
-class FLandscapeLayerWeightmapPackMaterialLayersComputeShaderResource : public FRenderResource
-{
-public:
-	FLandscapeLayerWeightmapPackMaterialLayersComputeShaderResource(const TArray<FLandscapeLayerWeightmapPackMaterialLayersComponentData>& InComponentsData, const TArray<float>& InWeightmapWeightBlendModeData, const TArray<FVector2f>& InTextureOutputOffset)
-		: OriginalComponentsData(InComponentsData)
-		, ComponentsDataCount(OriginalComponentsData.Num())
-		, OriginalWeightmapWeightBlendModeData(InWeightmapWeightBlendModeData)
-		, OriginalTextureOutputOffset(InTextureOutputOffset)
-	{}
-
-	~FLandscapeLayerWeightmapPackMaterialLayersComputeShaderResource()
-	{
-		ComponentsData.SafeRelease();
-		ComponentsDataSRV.SafeRelease();
-		WeightmapWeightBlendModeSRV.SafeRelease();
-		WeightmapTextureOutputOffsetSRV.SafeRelease();
-	}
-
-	/** Called when the resource is initialized. This is only called by the rendering thread. */
-	virtual void InitRHI(FRHICommandListBase& RHICmdList) override
-	{
-		ComponentsData = UE::RHIResourceUtils::CreateBufferFromArray(
-			RHICmdList,
-			TEXT("ComponentsData"),
-			EBufferUsageFlags::StructuredBuffer | EBufferUsageFlags::ShaderResource | EBufferUsageFlags::Volatile,
-			MakeConstArrayView(OriginalComponentsData)
-		);
-		ComponentsDataSRV = RHICmdList.CreateShaderResourceView(ComponentsData, FRHIViewDesc::CreateBufferSRV().SetTypeFromBuffer(ComponentsData));
-
-		WeightmapWeightBlendMode = UE::RHIResourceUtils::CreateVertexBufferFromArray(
-			RHICmdList,
-			TEXT("WeightmapWeightBlendMode"),
-			EBufferUsageFlags::ShaderResource | EBufferUsageFlags::Volatile,
-			MakeConstArrayView(OriginalWeightmapWeightBlendModeData)
-		);
-		WeightmapWeightBlendModeSRV = RHICmdList.CreateShaderResourceView(
-			WeightmapWeightBlendMode, 
-			FRHIViewDesc::CreateBufferSRV()
-				.SetType(FRHIViewDesc::EBufferType::Typed)
-				.SetFormat(PF_R32_FLOAT));
-
-		WeightmapTextureOutputOffset = UE::RHIResourceUtils::CreateVertexBufferFromArray(
-			RHICmdList,
-			TEXT("WeightmapTextureOutputOffset"),
-			EBufferUsageFlags::ShaderResource | EBufferUsageFlags::Volatile,
-			MakeConstArrayView(OriginalTextureOutputOffset)
-		);
-		WeightmapTextureOutputOffsetSRV = RHICmdList.CreateShaderResourceView(
-			WeightmapTextureOutputOffset,
-			FRHIViewDesc::CreateBufferSRV()
-				.SetType(FRHIViewDesc::EBufferType::Typed)
-				.SetFormat(PF_G32R32F));
-	}
-
-	virtual void ReleaseRHI() override
-	{
-		ComponentsData.SafeRelease();
-		ComponentsDataSRV.SafeRelease();
-		WeightmapWeightBlendModeSRV.SafeRelease();
-		WeightmapTextureOutputOffsetSRV.SafeRelease();
-	}
-
-	int32 GetComponentsDataCount() const
-	{
-		return ComponentsDataCount;
-	}
-
-private:
-	friend class FLandscapeLayerWeightmapPackMaterialLayersCS;
-
-	FBufferRHIRef ComponentsData;
-	FShaderResourceViewRHIRef ComponentsDataSRV;
-	TArray<FLandscapeLayerWeightmapPackMaterialLayersComponentData> OriginalComponentsData;
-	int32 ComponentsDataCount;
-
-	TArray<float> OriginalWeightmapWeightBlendModeData;
-	FBufferRHIRef WeightmapWeightBlendMode;
-	FShaderResourceViewRHIRef WeightmapWeightBlendModeSRV;
-
-	TArray<FVector2f> OriginalTextureOutputOffset;
-	FBufferRHIRef WeightmapTextureOutputOffset;
-	FShaderResourceViewRHIRef WeightmapTextureOutputOffsetSRV;
-};
-
-struct FLandscapeLayerWeightmapPackMaterialLayersComputeShaderParameters
-{
-	FLandscapeLayerWeightmapPackMaterialLayersComputeShaderParameters()
-		: ComponentWeightmapResource(nullptr)
-		, ComputeShaderResource(nullptr)
-		, AtlasWeightmapsPerLayer(nullptr)
-		, ComponentSize(0)
-	{}
-
-	FLandscapeTexture2DResource* ComponentWeightmapResource;
-	FLandscapeLayerWeightmapPackMaterialLayersComputeShaderResource* ComputeShaderResource;
-	FLandscapeTexture2DArrayResource* AtlasWeightmapsPerLayer;
-	uint32 ComponentSize;
-};
-
-class FLandscapeLayerWeightmapPackMaterialLayersCS : public FGlobalShader
-{
-	DECLARE_GLOBAL_SHADER(FLandscapeLayerWeightmapPackMaterialLayersCS);
-public:
-
-	static bool ShouldCompilePermutation(const FGlobalShaderPermutationParameters& Parameters)
-	{
-		return UE::Landscape::DoesPlatformSupportEditLayers(Parameters.Platform);
-	}
-
-	static void ModifyCompilationEnvironment(const FGlobalShaderPermutationParameters& Parameters, FShaderCompilerEnvironment& OutEnvironment)
-	{
-		OutEnvironment.SetDefine(TEXT("THREADGROUP_SIZEX"), GLandscapeLayerWeightmapThreadGroupSizeX);
-		OutEnvironment.SetDefine(TEXT("THREADGROUP_SIZEY"), GLandscapeLayerWeightmapThreadGroupSizeY);
-		UE::Landscape::ModifyShaderCompilerEnvironmentForDebug(OutEnvironment);
-	}
-
-	FLandscapeLayerWeightmapPackMaterialLayersCS(const ShaderMetaType::CompiledShaderInitializerType& Initializer)
-		: FGlobalShader(Initializer)
-	{
-		ComponentWeightmapParam.Bind(Initializer.ParameterMap, TEXT("OutComponentWeightMaps"));
-		AtlasPaintListsParam.Bind(Initializer.ParameterMap, TEXT("InAtlasPaintLayers"));
-		ComponentsDataParam.Bind(Initializer.ParameterMap, TEXT("InPackLayersComponentsData"));
-		ComponentSizeParam.Bind(Initializer.ParameterMap, TEXT("ComponentSize"));
-		WeightmapWeightBlendModeParam.Bind(Initializer.ParameterMap, TEXT("InWeightmapWeightBlendMode"));
-		WeightmapTextureOutputOffsetParam.Bind(Initializer.ParameterMap, TEXT("InWeightmapTextureOutputOffset"));
-	}
-
-	FLandscapeLayerWeightmapPackMaterialLayersCS()
-	{}
-
-	void SetParameters(FRHIBatchedShaderParameters& BatchedParameters, const FLandscapeLayerWeightmapPackMaterialLayersComputeShaderParameters& InParams)
-	{
-		SetUAVParameter(BatchedParameters, ComponentWeightmapParam, InParams.ComponentWeightmapResource->GetTextureUAV(/*InMipLevel = */0));
-		SetTextureParameter(BatchedParameters, AtlasPaintListsParam, InParams.AtlasWeightmapsPerLayer->TextureRHI);
-		SetSRVParameter(BatchedParameters, ComponentsDataParam, InParams.ComputeShaderResource->ComponentsDataSRV);
-		SetShaderValue(BatchedParameters, ComponentSizeParam, InParams.ComponentSize);
-		SetSRVParameter(BatchedParameters, WeightmapWeightBlendModeParam, InParams.ComputeShaderResource->WeightmapWeightBlendModeSRV);
-		SetSRVParameter(BatchedParameters, WeightmapTextureOutputOffsetParam, InParams.ComputeShaderResource->WeightmapTextureOutputOffsetSRV);
-	}
-
-	void UnsetParameters(FRHIBatchedShaderUnbinds& BatchedUnbinds)
-	{
-		UnsetUAVParameter(BatchedUnbinds, ComponentWeightmapParam);
-	}
-
-private:
-	LAYOUT_FIELD(FShaderResourceParameter, ComponentWeightmapParam);
-	LAYOUT_FIELD(FShaderResourceParameter, AtlasPaintListsParam);
-	LAYOUT_FIELD(FShaderResourceParameter, ComponentsDataParam);
-	LAYOUT_FIELD(FShaderParameter, ComponentSizeParam);
-	LAYOUT_FIELD(FShaderResourceParameter, WeightmapWeightBlendModeParam);
-	LAYOUT_FIELD(FShaderResourceParameter, WeightmapTextureOutputOffsetParam);
-};
-
-IMPLEMENT_GLOBAL_SHADER(FLandscapeLayerWeightmapPackMaterialLayersCS, "/Engine/Private/LandscapeLayersCS.usf", "PackPaintLayerToWeightmap", SF_Compute);
-
-class FLandscapeLayerWeightmapPackMaterialLayersCSDispatch_RenderThread
-{
-public:
-	FLandscapeLayerWeightmapPackMaterialLayersCSDispatch_RenderThread(const FLandscapeLayerWeightmapPackMaterialLayersComputeShaderParameters& InShaderParams)
-		: ShaderParams(InShaderParams)
-	{}
-
-	void PackLayers(FRHICommandListImmediate& InRHICmdList)
-	{
-		RHI_BREADCRUMB_EVENT_STAT(InRHICmdList, LandscapeLayers_PackLayers, "LandscapeLayers_PackLayers");
-		SCOPED_GPU_STAT(InRHICmdList, LandscapeLayers_PackLayers);
-
-		TShaderMapRef<FLandscapeLayerWeightmapPackMaterialLayersCS> ComputeShader(GetGlobalShaderMap(GMaxRHIFeatureLevel));
-		SetComputePipelineState(InRHICmdList, ComputeShader.GetComputeShader());
-
-		SetShaderParametersLegacyCS(InRHICmdList, ComputeShader, ShaderParams);
-
-		uint32 ThreadGroupCountX = FMath::CeilToInt((float)ShaderParams.ComponentSize / (float)GLandscapeLayerWeightmapThreadGroupSizeX);
-		uint32 ThreadGroupCountY = FMath::CeilToInt((float)ShaderParams.ComponentSize / (float)GLandscapeLayerWeightmapThreadGroupSizeY);
-		check(ThreadGroupCountX > 0 && ThreadGroupCountY > 0);
-
-		DispatchComputeShader(InRHICmdList, ComputeShader.GetShader(), ThreadGroupCountX, ThreadGroupCountY, ShaderParams.ComputeShaderResource->GetComponentsDataCount());
-
-		UnsetShaderParametersLegacyCS(InRHICmdList, ComputeShader);
-
-		ShaderParams.ComputeShaderResource->ReleaseResource();
-		delete ShaderParams.ComputeShaderResource;
-	}
-
-private:
-	FLandscapeLayerWeightmapPackMaterialLayersComputeShaderParameters ShaderParams;
-};
-
-// ----------------------------------------------------------------------------------
-// /Engine/Private/Landscape/LandscapeLayersHeightmapsPS.usf shaders : 
-
-class FLandscapeLayersHeightmapsMergeEditLayersPS : public FGlobalShader
-{
-	DECLARE_GLOBAL_SHADER(FLandscapeLayersHeightmapsMergeEditLayersPS);
-	SHADER_USE_PARAMETER_STRUCT(FLandscapeLayersHeightmapsMergeEditLayersPS, FGlobalShader);
-
-public:
-	BEGIN_SHADER_PARAMETER_STRUCT(FParameters, )
-		SHADER_PARAMETER(uint32, InNumEditLayers)
-		SHADER_PARAMETER_RDG_TEXTURE_SRV(Texture2DArray<float4>, InEditLayersTextures)
-		SHADER_PARAMETER_RDG_BUFFER_SRV(StructuredBuffer<FLandscapeEditLayerHeightmapMergeInfo>, InEditLayersMergeInfos)
-		RENDER_TARGET_BINDING_SLOTS()
-	END_SHADER_PARAMETER_STRUCT()
-
-	static bool ShouldCompilePermutation(const FGlobalShaderPermutationParameters& InParameters)
-	{
-		return UE::Landscape::DoesPlatformSupportEditLayers(InParameters.Platform);
-	}
-
-	static void ModifyCompilationEnvironment(const FGlobalShaderPermutationParameters& InParameters, FShaderCompilerEnvironment& OutEnvironment)
-	{
-		OutEnvironment.SetDefine(TEXT("MERGE_EDIT_LAYERS"), 1);
-		UE::Landscape::ModifyShaderCompilerEnvironmentForDebug(OutEnvironment);
-	}
-
-	static void MergeEditLayers(FRDGBuilder& GraphBuilder, FParameters* InParameters, const FIntPoint& InTextureSize)
-	{
-		FGlobalShaderMap* ShaderMap = GetGlobalShaderMap(GMaxRHIFeatureLevel);
-		TShaderMapRef<FLandscapeLayersHeightmapsMergeEditLayersPS> PixelShader(ShaderMap);
-
-		FPixelShaderUtils::AddFullscreenPass(
-			GraphBuilder,
-			ShaderMap,
-			RDG_EVENT_NAME("LandscapeLayers_MergeEditLayers"),
-			PixelShader,
-			InParameters,
-			FIntRect(0, 0, InTextureSize.X, InTextureSize.Y));
-	}
-};
-
-IMPLEMENT_GLOBAL_SHADER(FLandscapeLayersHeightmapsMergeEditLayersPS, "/Engine/Private/Landscape/LandscapeLayersHeightmapsPS.usf", "MergeEditLayers", SF_Pixel);
-
-class FLandscapeLayersHeightmapsStitchHeightmapPS : public FGlobalShader
-{
-	DECLARE_GLOBAL_SHADER(FLandscapeLayersHeightmapsStitchHeightmapPS);
-	SHADER_USE_PARAMETER_STRUCT(FLandscapeLayersHeightmapsStitchHeightmapPS, FGlobalShader);
-
-public:
-	BEGIN_SHADER_PARAMETER_STRUCT(FParameters, )
-		SHADER_PARAMETER(FUintVector2, InSourceTextureSize)
-		SHADER_PARAMETER(uint32, InNumSubsections)
-		SHADER_PARAMETER_RDG_TEXTURE_SRV(Texture2DArray<float2>, InSourceHeightmaps)
-		SHADER_PARAMETER_SCALAR_ARRAY(uint32, InNeighborHeightmapIndices, [9])
-		RENDER_TARGET_BINDING_SLOTS()
-	END_SHADER_PARAMETER_STRUCT()
-
-	static bool ShouldCompilePermutation(const FGlobalShaderPermutationParameters& Parameters)
-	{
-		return UE::Landscape::DoesPlatformSupportEditLayers(Parameters.Platform);
-	}
-
-	static void ModifyCompilationEnvironment(const FGlobalShaderPermutationParameters& Parameters, FShaderCompilerEnvironment& OutEnvironment)
-	{
-		OutEnvironment.SetDefine(TEXT("STITCH_HEIGHTMAP"), 1);
-		UE::Landscape::ModifyShaderCompilerEnvironmentForDebug(OutEnvironment);
-	}
-
-	static void StitchHeightmap(FRDGBuilder& GraphBuilder, FParameters* InParameters)
-	{
-		FGlobalShaderMap* ShaderMap = GetGlobalShaderMap(GMaxRHIFeatureLevel);
-		TShaderMapRef<FLandscapeLayersHeightmapsStitchHeightmapPS> PixelShader(ShaderMap);
-
-		FPixelShaderUtils::AddFullscreenPass(
-			GraphBuilder,
-			ShaderMap,
-			RDG_EVENT_NAME("LandscapeLayers_StitchHeightmap"),
-			PixelShader,
-			InParameters,
-			FIntRect(0, 0, InParameters->InSourceTextureSize.X, InParameters->InSourceTextureSize.Y));
-	}
-};
-
-IMPLEMENT_GLOBAL_SHADER(FLandscapeLayersHeightmapsStitchHeightmapPS, "/Engine/Private/Landscape/LandscapeLayersHeightmapsPS.usf", "StitchHeightmap", SF_Pixel);
-
-class FLandscapeLayersHeightmapsFinalizeHeightmapPS : public FGlobalShader
-{
-	DECLARE_GLOBAL_SHADER(FLandscapeLayersHeightmapsFinalizeHeightmapPS);
-	SHADER_USE_PARAMETER_STRUCT(FLandscapeLayersHeightmapsFinalizeHeightmapPS, FGlobalShader);
-
-public:
-	BEGIN_SHADER_PARAMETER_STRUCT(FParameters, )
-		SHADER_PARAMETER(FUintVector2, InSourceTextureSize)
-		SHADER_PARAMETER(uint32, InNumSubsections)
-		SHADER_PARAMETER_RDG_TEXTURE_SRV(Texture2DArray<float2>, InSourceHeightmaps)
-		SHADER_PARAMETER_SCALAR_ARRAY(uint32, InNeighborHeightmapIndices, [9])
-		SHADER_PARAMETER(FUintVector4, InDestinationTextureSubregion)
-		SHADER_PARAMETER(FVector3f, InLandscapeGridScale)
-		RENDER_TARGET_BINDING_SLOTS()
-	END_SHADER_PARAMETER_STRUCT()
-
-	static bool ShouldCompilePermutation(const FGlobalShaderPermutationParameters& Parameters)
-	{
-		return UE::Landscape::DoesPlatformSupportEditLayers(Parameters.Platform);
-	}
-
-	static void ModifyCompilationEnvironment(const FGlobalShaderPermutationParameters& Parameters, FShaderCompilerEnvironment& OutEnvironment)
-	{
-		OutEnvironment.SetDefine(TEXT("FINALIZE_HEIGHTMAP"), 1);
-		UE::Landscape::ModifyShaderCompilerEnvironmentForDebug(OutEnvironment);
-	}
-
-	static void FinalizeHeightmap(FRDGBuilder& GraphBuilder, FParameters* InParameters)
-	{
-		FGlobalShaderMap* ShaderMap = GetGlobalShaderMap(GMaxRHIFeatureLevel);
-		TShaderMapRef<FLandscapeLayersHeightmapsFinalizeHeightmapPS> PixelShader(ShaderMap);
-
-		FPixelShaderUtils::AddFullscreenPass(
-			GraphBuilder,
-			ShaderMap,
-			RDG_EVENT_NAME("LandscapeLayers_FinalizeHeightmap"),
-			PixelShader,
-			InParameters,
-			FIntRect(InParameters->InDestinationTextureSubregion.X, InParameters->InDestinationTextureSubregion.Y, InParameters->InDestinationTextureSubregion.Z, InParameters->InDestinationTextureSubregion.W));
-	}
-};
-
-IMPLEMENT_GLOBAL_SHADER(FLandscapeLayersHeightmapsFinalizeHeightmapPS, "/Engine/Private/Landscape/LandscapeLayersHeightmapsPS.usf", "FinalizeHeightmap", SF_Pixel);
-
-class FLandscapeLayersHeightmapsGenerateMipsPS : public FGlobalShader
-{
-	DECLARE_GLOBAL_SHADER(FLandscapeLayersHeightmapsGenerateMipsPS);
-	SHADER_USE_PARAMETER_STRUCT(FLandscapeLayersHeightmapsGenerateMipsPS, FGlobalShader);
-
-public:
-	BEGIN_SHADER_PARAMETER_STRUCT(FParameters, )
-		SHADER_PARAMETER(FUintVector2, InCurrentMipSubregionSize)
-		SHADER_PARAMETER(uint32, InNumSubsections)
-		SHADER_PARAMETER_RDG_TEXTURE_SRV(Texture2D<float4>, InSourceHeightmap)
-		RENDER_TARGET_BINDING_SLOTS()
-	END_SHADER_PARAMETER_STRUCT()
-
-	static bool ShouldCompilePermutation(const FGlobalShaderPermutationParameters& Parameters)
-	{
-		return UE::Landscape::DoesPlatformSupportEditLayers(Parameters.Platform);
-	}
-
-	static void ModifyCompilationEnvironment(const FGlobalShaderPermutationParameters& Parameters, FShaderCompilerEnvironment& OutEnvironment)
-	{
-		OutEnvironment.SetDefine(TEXT("GENERATE_MIPS"), 1);
-		UE::Landscape::ModifyShaderCompilerEnvironmentForDebug(OutEnvironment);
-	}
-
-	static void GenerateSingleMip(FRDGBuilder& GraphBuilder, FParameters* InParameters)
-	{
-		FGlobalShaderMap* ShaderMap = GetGlobalShaderMap(GMaxRHIFeatureLevel);
-		TShaderMapRef<FLandscapeLayersHeightmapsGenerateMipsPS> PixelShader(ShaderMap);
-
-		FIntVector MipSize = InParameters->RenderTargets[0].GetTexture()->Desc.GetSize();
-
-		FPixelShaderUtils::AddFullscreenPass(
-			GraphBuilder,
-			ShaderMap,
-			RDG_EVENT_NAME("LandscapeLayers_GenerateMip"),
-			PixelShader,
-			InParameters,
-			FIntRect(0, 0, MipSize.X, MipSize.Y));
-	}
-};
-
-IMPLEMENT_GLOBAL_SHADER(FLandscapeLayersHeightmapsGenerateMipsPS, "/Engine/Private/Landscape/LandscapeLayersHeightmapsPS.usf", "GenerateMips", SF_Pixel);
-
-// ----------------------------------------------------------------------------------
-// /Engine/Private/Landscape/LandscapeLayersWeightmapsPS.usf shaders :
-
-class FLandscapeLayersWeightmapsMergeEditLayersPS : public FGlobalShader
-{
-	DECLARE_GLOBAL_SHADER(FLandscapeLayersWeightmapsMergeEditLayersPS);
-	SHADER_USE_PARAMETER_STRUCT(FLandscapeLayersWeightmapsMergeEditLayersPS, FGlobalShader);
-
-public:
-	BEGIN_SHADER_PARAMETER_STRUCT(FParameters, )
-		SHADER_PARAMETER(uint32, InNumEditLayers)
-		SHADER_PARAMETER(uint32, InStartIndexInEditLayersMergeInfos)
-		SHADER_PARAMETER_RDG_TEXTURE_SRV(Texture2DArray<float4>, InPackedWeightmaps)
-		SHADER_PARAMETER_RDG_BUFFER_SRV(StructuredBuffer<FLandscapeEditLayerWeightmapMergeInfo>, InEditLayersMergeInfos)
-		RENDER_TARGET_BINDING_SLOTS()
-		END_SHADER_PARAMETER_STRUCT()
-
-		static bool ShouldCompilePermutation(const FGlobalShaderPermutationParameters& InParameters)
-	{
-		return UE::Landscape::DoesPlatformSupportEditLayers(InParameters.Platform);
-	}
-
-	static void ModifyCompilationEnvironment(const FGlobalShaderPermutationParameters& InParameters, FShaderCompilerEnvironment& OutEnvironment)
-	{
-		OutEnvironment.SetDefine(TEXT("MERGE_EDIT_LAYERS"), 1);
-		UE::Landscape::ModifyShaderCompilerEnvironmentForDebug(OutEnvironment);
-	}
-
-	static void MergeEditLayers(FRDGBuilder& GraphBuilder, FParameters* InParameters)
-	{
-		FGlobalShaderMap* ShaderMap = GetGlobalShaderMap(GMaxRHIFeatureLevel);
-		TShaderMapRef<FLandscapeLayersWeightmapsMergeEditLayersPS> PixelShader(ShaderMap);
-
-		FIntVector TextureSize = InParameters->RenderTargets[0].GetTexture()->Desc.GetSize();
-
-		FPixelShaderUtils::AddFullscreenPass(
-			GraphBuilder,
-			ShaderMap,
-			RDG_EVENT_NAME("LandscapeLayers_MergeEditLayers"),
-			PixelShader,
-			InParameters,
-			FIntRect(0, 0, TextureSize.X, TextureSize.Y));
-	}
-};
-
-IMPLEMENT_GLOBAL_SHADER(FLandscapeLayersWeightmapsMergeEditLayersPS, "/Engine/Private/Landscape/LandscapeLayersWeightmapsPS.usf", "MergeEditLayers", SF_Pixel);
-
-class FLandscapeLayersWeightmapsFinalizeWeightmapPS : public FGlobalShader
-{
-	DECLARE_GLOBAL_SHADER(FLandscapeLayersWeightmapsFinalizeWeightmapPS);
-	SHADER_USE_PARAMETER_STRUCT(FLandscapeLayersWeightmapsFinalizeWeightmapPS, FGlobalShader);
-
-public:
-	BEGIN_SHADER_PARAMETER_STRUCT(FParameters, )
-		SHADER_PARAMETER(uint32, InValidTextureChannelsMask)
-		SHADER_PARAMETER(FUintVector4, InPerChannelPaintLayerIndexInWeightmaps)
-		SHADER_PARAMETER(FUintVector4, InPerChannelStartPaintLayerIndex)
-		SHADER_PARAMETER(FUintVector4, InPerChannelNumPaintLayers)
-		SHADER_PARAMETER_RDG_TEXTURE_SRV_ARRAY(Texture2D<float4>, InPerChannelPaintLayerWeightmaps, [4])
-		SHADER_PARAMETER_RDG_BUFFER_SRV(StructuredBuffer<uint>, InPaintLayerInfoIndices)
-		SHADER_PARAMETER_RDG_BUFFER_SRV(StructuredBuffer<FLandscapeEditLayerWeightmapPaintLayerInfo>, InPaintLayerInfos)
-		RENDER_TARGET_BINDING_SLOTS()
-		END_SHADER_PARAMETER_STRUCT()
-
-		static bool ShouldCompilePermutation(const FGlobalShaderPermutationParameters& Parameters)
-	{
-		return UE::Landscape::DoesPlatformSupportEditLayers(Parameters.Platform);
-	}
-
-	static void ModifyCompilationEnvironment(const FGlobalShaderPermutationParameters& Parameters, FShaderCompilerEnvironment& OutEnvironment)
-	{
-		OutEnvironment.SetDefine(TEXT("FINALIZE_WEIGHTMAP"), 1);
-		UE::Landscape::ModifyShaderCompilerEnvironmentForDebug(OutEnvironment);
-	}
-
-	static void FinalizeWeightmap(FRDGBuilder& GraphBuilder, FParameters* InParameters)
-	{
-		FGlobalShaderMap* ShaderMap = GetGlobalShaderMap(GMaxRHIFeatureLevel);
-		TShaderMapRef<FLandscapeLayersWeightmapsFinalizeWeightmapPS> PixelShader(ShaderMap);
-
-		FIntVector TextureSize = InParameters->RenderTargets[0].GetTexture()->Desc.GetSize();
-
-		FPixelShaderUtils::AddFullscreenPass(
-			GraphBuilder,
-			ShaderMap,
-			RDG_EVENT_NAME("LandscapeLayers_PackWeightmap"),
-			PixelShader,
-			InParameters,
-			FIntRect(0, 0, TextureSize.X, TextureSize.Y));
-	}
-};
-
-IMPLEMENT_GLOBAL_SHADER(FLandscapeLayersWeightmapsFinalizeWeightmapPS, "/Engine/Private/Landscape/LandscapeLayersWeightmapsPS.usf", "FinalizeWeightmap", SF_Pixel);
-
-class FLandscapeLayersWeightmapsGenerateMipsPS : public FGlobalShader
-{
-	DECLARE_GLOBAL_SHADER(FLandscapeLayersWeightmapsGenerateMipsPS);
-	SHADER_USE_PARAMETER_STRUCT(FLandscapeLayersWeightmapsGenerateMipsPS, FGlobalShader);
-
-public:
-	BEGIN_SHADER_PARAMETER_STRUCT(FParameters, )
-		SHADER_PARAMETER(FUintVector2, InCurrentMipSize)
-		SHADER_PARAMETER(uint32, InNumSubsections)
-		SHADER_PARAMETER_RDG_TEXTURE_SRV(Texture2D<float4>, InSourceWeightmap)
-		RENDER_TARGET_BINDING_SLOTS()
-		END_SHADER_PARAMETER_STRUCT()
-
-		static bool ShouldCompilePermutation(const FGlobalShaderPermutationParameters& Parameters)
-	{
-		return UE::Landscape::DoesPlatformSupportEditLayers(Parameters.Platform);
-	}
-
-	static void ModifyCompilationEnvironment(const FGlobalShaderPermutationParameters& Parameters, FShaderCompilerEnvironment& OutEnvironment)
-	{
-		OutEnvironment.SetDefine(TEXT("GENERATE_MIPS"), 1);
-		UE::Landscape::ModifyShaderCompilerEnvironmentForDebug(OutEnvironment);
-	}
-
-	static void GenerateSingleMip(FRDGBuilder& GraphBuilder, FParameters* InParameters)
-	{
-		FGlobalShaderMap* ShaderMap = GetGlobalShaderMap(GMaxRHIFeatureLevel);
-		TShaderMapRef<FLandscapeLayersWeightmapsGenerateMipsPS> PixelShader(ShaderMap);
-
-		FIntVector MipSize = InParameters->RenderTargets[0].GetTexture()->Desc.GetSize();
-
-		FPixelShaderUtils::AddFullscreenPass(
-			GraphBuilder,
-			ShaderMap,
-			RDG_EVENT_NAME("LandscapeLayers_GenerateMip"),
-			PixelShader,
-			InParameters,
-			FIntRect(0, 0, MipSize.X, MipSize.Y));
-	}
-};
-
-IMPLEMENT_GLOBAL_SHADER(FLandscapeLayersWeightmapsGenerateMipsPS, "/Engine/Private/Landscape/LandscapeLayersWeightmapsPS.usf", "GenerateMips", SF_Pixel);
-
-
-// ----------------------------------------------------------------------------------
 // /Engine/Private/Landscape/LandscapeEditLayersHeightmaps.usf shaders :
 
 class FLandscapeEditLayersHeightmapsGenerateNormalsPS : public FGlobalShader
@@ -1986,7 +935,7 @@
 	using FPermutationDomain = TShaderPermutationDomain<FCopyWeightmap>;
 
 	static bool ShouldCompilePermutation(const FGlobalShaderPermutationParameters& InParameters)
-	{
+{
 		return UE::Landscape::DoesPlatformSupportEditLayers(InParameters.Platform);
 	}
 
@@ -2338,140 +1287,6 @@
 	FTextureRenderTargetResource* RenderTargetResource;
 };
 
-// ----------------------------------------------------------------------------------
-// Render command
-
-template<typename ShaderDataType, typename ShaderPixelClass, typename ShaderPixelMipsClass>
-class FLandscapeLayersRender_RenderThread
-{
-public:
-
-	FLandscapeLayersRender_RenderThread(const FString& InDebugName, UTextureRenderTarget2D* InWriteRenderTarget, const FIntPoint& InWriteRenderTargetSize, const FIntPoint& InReadRenderTargetSize, const FMatrix& InProjectionMatrix,
-		const ShaderDataType& InShaderParams, uint8 InCurrentMip, const TArray<FLandscapeLayersTriangle>& InTriangleList)
-		: RenderTargetResource(InWriteRenderTarget->GameThread_GetRenderTargetResource())
-		, WriteRenderTargetSize(InWriteRenderTargetSize)
-		, ReadRenderTargetSize(InReadRenderTargetSize)
-		, ProjectionMatrix(InProjectionMatrix)
-		, ShaderParams(InShaderParams)
-		, PrimitiveCount(InTriangleList.Num())
-		, DebugName(InDebugName)
-		, CurrentMip(InCurrentMip)
-	{
-		VertexBufferResource.Init(InTriangleList);
-	}
-
-	virtual ~FLandscapeLayersRender_RenderThread()
-	{}
-
-	void Render(FRHICommandListImmediate& InRHICmdList, bool InClearRT)
-	{
-		RHI_BREADCRUMB_EVENT_STAT_F(InRHICmdList, LandscapeLayers_Render, "LandscapeLayers_Render", "LandscapeLayers_Render %s", DebugName);
-		SCOPED_GPU_STAT(InRHICmdList, LandscapeLayers_Render);
-		INC_DWORD_STAT(STAT_LandscapeLayersRegenerateDrawCalls);
-		TRACE_CPUPROFILER_EVENT_SCOPE(FLandscapeLayersRender_RenderThread::Render);
-
-		check(IsInRenderingThread());
-
-		FSceneViewFamilyContext ViewFamily(FSceneViewFamily::ConstructionValues(RenderTargetResource, NULL, FEngineShowFlags(ESFIM_Game))
-			.SetTime(FGameTime::GetTimeSinceAppStart()));
-
-		FSceneViewInitOptions ViewInitOptions;
-		ViewInitOptions.SetViewRectangle(FIntRect(0, 0, WriteRenderTargetSize.X, WriteRenderTargetSize.Y));
-		ViewInitOptions.ViewOrigin = FVector::ZeroVector;
-		ViewInitOptions.ViewRotationMatrix = FMatrix::Identity;
-		ViewInitOptions.ProjectionMatrix = ProjectionMatrix;
-		ViewInitOptions.ViewFamily = &ViewFamily;
-		ViewInitOptions.BackgroundColor = FLinearColor::Black;
-		ViewInitOptions.OverlayColor = FLinearColor::White;
-
-		// Create and add the new view
-		FSceneView* View = new FSceneView(ViewInitOptions);
-		ViewFamily.Views.Add(View);
-
-		InRHICmdList.Transition(FRHITransitionInfo(ViewFamily.RenderTarget->GetRenderTargetTexture(), ERHIAccess::SRVMask, ERHIAccess::RTV));
-
-		// Init VB/IB Resource
-		VertexDeclaration.InitResource(InRHICmdList);
-		VertexBufferResource.InitResource(InRHICmdList);
-
-		// Setup Pipeline
-		FGraphicsPipelineStateInitializer GraphicsPSOInit;
-		GraphicsPSOInit.BoundShaderState.VertexDeclarationRHI = VertexDeclaration.VertexDeclarationRHI;
-		GraphicsPSOInit.PrimitiveType = PT_TriangleList;
-
-		GraphicsPSOInit.BlendState = TStaticBlendState<CW_RGBA, BO_Add, BF_One, BF_Zero, BO_Add, BF_One, BF_Zero>::GetRHI();
-		GraphicsPSOInit.RasterizerState = TStaticRasterizerState<>::GetRHI();
-		GraphicsPSOInit.DepthStencilState = TStaticDepthStencilState<false, CF_Always>::GetRHI();
-
-		FRHIRenderPassInfo RenderPassInfo(ViewFamily.RenderTarget->GetRenderTargetTexture(), CurrentMip == 0 || InClearRT ? ERenderTargetActions::Clear_Store : ERenderTargetActions::Load_Store, nullptr, 0, 0);
-		InRHICmdList.BeginRenderPass(RenderPassInfo, TEXT("DrawLayers"));
-
-		if (CurrentMip == 0)
-		{
-			// Setup Shaders
-			TShaderMapRef<FLandscapeLayersVS> VertexShader(GetGlobalShaderMap(View->GetFeatureLevel()));
-			TShaderMapRef<ShaderPixelClass> PixelShader(GetGlobalShaderMap(View->GetFeatureLevel()));
-
-			GraphicsPSOInit.BoundShaderState.VertexShaderRHI = VertexShader.GetVertexShader();
-			GraphicsPSOInit.BoundShaderState.PixelShaderRHI = PixelShader.GetPixelShader();
-
-			InRHICmdList.SetViewport(static_cast<float>(View->UnscaledViewRect.Min.X), static_cast<float>(View->UnscaledViewRect.Min.Y), 0.0f,
-			                         static_cast<float>(View->UnscaledViewRect.Max.X), static_cast<float>(View->UnscaledViewRect.Max.Y), 1.0f);
-
-			InRHICmdList.ApplyCachedRenderTargets(GraphicsPSOInit);
-			SetGraphicsPipelineState(InRHICmdList, GraphicsPSOInit, 0);
-
-			// Set shader params
-			SetShaderParametersLegacyVS(InRHICmdList, VertexShader, FMatrix44f(ProjectionMatrix));		// LWC_TODo: Precision loss?
-			SetShaderParametersLegacyPS(InRHICmdList, PixelShader, ShaderParams);
-		}
-		else
-		{
-			// Setup Shaders
-			TShaderMapRef<FLandscapeLayersVS> VertexShader(GetGlobalShaderMap(View->GetFeatureLevel()));
-			TShaderMapRef<ShaderPixelMipsClass> PixelShader(GetGlobalShaderMap(View->GetFeatureLevel()));
-
-			GraphicsPSOInit.BoundShaderState.VertexShaderRHI = VertexShader.GetVertexShader();
-			GraphicsPSOInit.BoundShaderState.PixelShaderRHI = PixelShader.GetPixelShader();
-
-			InRHICmdList.SetViewport(0.0f, 0.0f, 0.0f, static_cast<float>(WriteRenderTargetSize.X), static_cast<float>(WriteRenderTargetSize.Y), 1.0f);
-
-			InRHICmdList.ApplyCachedRenderTargets(GraphicsPSOInit);
-			SetGraphicsPipelineState(InRHICmdList, GraphicsPSOInit, 0);
-
-			// Set shader params
-			SetShaderParametersLegacyVS(InRHICmdList, VertexShader, FMatrix44f(ProjectionMatrix));		// LWC_TODo: Precision loss?
-			SetShaderParametersLegacyPS(InRHICmdList, PixelShader, ShaderParams);
-		}
-
-		InRHICmdList.SetScissorRect(false, 0, 0, 0, 0);
-		InRHICmdList.SetStreamSource(0, VertexBufferResource.VertexBufferRHI, 0);
-
-		InRHICmdList.DrawPrimitive(0, PrimitiveCount, 1);
-
-		InRHICmdList.EndRenderPass();
-		InRHICmdList.Transition(FRHITransitionInfo(ViewFamily.RenderTarget->GetRenderTargetTexture(), ERHIAccess::RTV, ERHIAccess::SRVMask));
-
-		VertexDeclaration.ReleaseResource();
-		VertexBufferResource.ReleaseResource();
-	}
-
-private:
-	FTextureRenderTargetResource* RenderTargetResource;
-	FIntPoint WriteRenderTargetSize;
-	FIntPoint ReadRenderTargetSize;
-	FMatrix ProjectionMatrix;
-	ShaderDataType ShaderParams;
-	FLandscapeLayersVertexBuffer VertexBufferResource;
-	int32 PrimitiveCount;
-	FLandscapeLayersVertexDeclaration VertexDeclaration;
-	FString DebugName;
-	uint8 CurrentMip;
-};
-
-typedef FLandscapeLayersRender_RenderThread<FLandscapeLayersHeightmapShaderParameters, FLandscapeLayersHeightmapPS, FLandscapeLayersHeightmapMipsPS> FLandscapeLayersHeightmapRender_RenderThread;
-typedef FLandscapeLayersRender_RenderThread<FLandscapeLayersWeightmapShaderParameters, FLandscapeLayersWeightmapPS, FLandscapeLayersWeightmapMipsPS> FLandscapeLayersWeightmapRender_RenderThread;
-
 #if WITH_EDITOR
 
 bool ALandscape::IsMaterialResourceCompiled(FMaterialResource* InMaterialResource, bool bInWaitForCompilation)
@@ -2526,31 +1341,10 @@
 	return true;
 }
 
-ELandscapeEditLayersMergeMode ALandscape::GetEditLayersMergeMode() const
-{
-	ELandscapeEditLayersMergeMode DesiredMode = static_cast<ELandscapeEditLayersMergeMode>(FMath::Min<uint8>(static_cast<uint8>(LandscapeEditLayersLocalMerge), static_cast<uint8>(ELandscapeEditLayersMergeMode::Invalid) - 1));
-
-	if (DesiredMode == ELandscapeEditLayersMergeMode::LocalMerge)
-	{
-		for (const FLandscapeLayer& Layer : LandscapeEditLayers)
-		{
-			// No BP brush is supported for local merge of edit layers : 
-			if (Layer.EditLayer->IsVisible() && Algo::AnyOf(Layer.Brushes, [](const FLandscapeLayerBrush& Brush) { return (Brush.GetBrush() != nullptr) && Brush.GetBrush()->IsVisible(); }))
-			{
-				return ELandscapeEditLayersMergeMode::GlobalMerge;
-			}
-		}
-	}
-
-	return DesiredMode;
-}
-
+// TODO:  remove with IntermediateRender.
 bool ALandscape::AllowSelectiveRenderEditLayers() const
 {
-	// TODO:  remove when GlobalMerge goes
-
-	return GetEditLayersMergeMode() == ELandscapeEditLayersMergeMode::BatchedMerge &&
-		   CVarLandscapeSelectiveRenderEditLayers.GetValueOnGameThread();
+	return true;
 }
 
 bool ALandscape::GetIsPendingEditLayerConversion() const
@@ -2563,164 +1357,6 @@
 	bIsPendingEditLayerConversion = bInIsPendingEditLayerConversion;
 }
 
-bool ALandscape::CreateLayersRenderingResource(bool bInForceLayerResourceReset)
-{
-	ULandscapeInfo* Info = GetLandscapeInfo();
-	if (Info == nullptr)
-	{
-		return false;
-	}
-
-	const FIntPoint ComponentCounts = ComputeComponentCounts();
-
-	// No components, can't update the render targets	
-	if (ComponentCounts.X <= 0 || ComponentCounts.Y <= 0)
-	{
-		// We return true, still, because we don't want to prevent the update from happening. It will update nothing, but at least it will clear the update flags LayerContentUpdateModes
-		//  which will in turn unblock things like grass map rendering, etc.
-		return true;
-	}
-
-	ALandscape* Landscape = GetLandscapeActor();
-	check(Landscape);
-
-	// Always reset layer resources in global merge mode (using the second block).  In other merge modes, only release
-	// resources when forced, often because we switched out of global merge mode.
-	bool bCurrentlyGlobalMerge = GetEditLayersMergeMode() == ELandscapeEditLayersMergeMode::GlobalMerge;
-	if (!bCurrentlyGlobalMerge && bInForceLayerResourceReset)
-	{
-		ReleaseLayersRenderingResource();
-
-		Landscape->HeightmapRTList.Empty();
-		Landscape->WeightmapRTList.Empty();
-	}
-	else if (bCurrentlyGlobalMerge)
-	{
-		// Heightmap mip size
-		int32 ComponentVerts = ((SubsectionSizeQuads + 1) * NumSubsections);
-		int32 CurrentMipSizeX = ComponentVerts * ComponentCounts.X;
-		int32 CurrentMipSizeY = ComponentVerts * ComponentCounts.Y;
-
-		if (CurrentMipSizeX > GRHIGlobals.MaxTextureDimensions || CurrentMipSizeY > GRHIGlobals.MaxTextureDimensions)
-		{
-			if (!bWarnedGlobalMergeDimensionsExceeded)
-			{
-				UE_LOG(LogLandscape, Error, TEXT("Cannot initialize resources for Landscape Layer Merge because the current device does not support render targets of the required size.  Please reduce landscape size, or use a different render device, or try to enable local merge with `landscape.EditLayersLocalMerge.Enable 1` (local merge works only if no landscape blueprint brushes are used)"));
-				bWarnedGlobalMergeDimensionsExceeded = true;
-			}
-			return false;
-		}
-
-		// once the issue is fixed, clear the warn flag
-		bWarnedGlobalMergeDimensionsExceeded = false;
-
-		bool bCreateFromScratch = (Landscape->HeightmapRTList.Num() == 0);
-		if (bCreateFromScratch)
-		{
-			Landscape->HeightmapRTList.Init(nullptr, (int32)EHeightmapRTType::HeightmapRT_Count);
-		}
-
-		auto InitOrResizeRT = [](UTextureRenderTarget2D* RT, int32 ResX, int32 ResY, bool bInit)
-		{
-			if (bInit)
-			{
-				RT->InitAutoFormat(FMath::RoundUpToPowerOfTwo(ResX), FMath::RoundUpToPowerOfTwo(ResY));
-				RT->UpdateResourceImmediate(true);
-			}
-			else
-			{
-				RT->ResizeTarget(FMath::RoundUpToPowerOfTwo(ResX), FMath::RoundUpToPowerOfTwo(ResY));
-			}
-		};
-
-		for (int32 i = 0; i < (int32)EHeightmapRTType::HeightmapRT_Count; ++i)
-		{
-			if (bCreateFromScratch)
-			{
-				FText DisplayName = StaticEnum<EHeightmapRTType>()->GetDisplayValueAsText((EHeightmapRTType)i);
-				FName RTName = MakeUniqueObjectName(GetTransientPackage(), UTextureRenderTarget2D::StaticClass(), FName(*DisplayName.ToString()));
-				Landscape->HeightmapRTList[i] = NewObject<UTextureRenderTarget2D>(GetTransientPackage(), RTName, RF_Transient);
-				check(Landscape->HeightmapRTList[i]);
-				Landscape->HeightmapRTList[i]->RenderTargetFormat = RTF_RGBA8;
-				Landscape->HeightmapRTList[i]->AddressX = TextureAddress::TA_Clamp;
-				Landscape->HeightmapRTList[i]->AddressY = TextureAddress::TA_Clamp;
-				Landscape->HeightmapRTList[i]->ClearColor = FLinearColor(0.0f, 0.0f, 0.0f, 0.0f);
-			}
-
-			if (i < (int32)EHeightmapRTType::HeightmapRT_Mip1) // Landscape size RT
-			{
-				InitOrResizeRT(Landscape->HeightmapRTList[i], CurrentMipSizeX, CurrentMipSizeY, bCreateFromScratch);
-			}
-			else // Mips
-			{
-				CurrentMipSizeX >>= 1;
-				CurrentMipSizeY >>= 1;
-				InitOrResizeRT(Landscape->HeightmapRTList[i], CurrentMipSizeX, CurrentMipSizeY, bCreateFromScratch);
-			}
-
-			// Only generate required mips RT
-			if (CurrentMipSizeX == ComponentCounts.X && CurrentMipSizeY == ComponentCounts.Y)
-			{
-				break;
-			}
-		}
-
-		// Weightmap mip size
-		CurrentMipSizeX = ((SubsectionSizeQuads + 1) * NumSubsections) * ComponentCounts.X;
-		CurrentMipSizeY = ((SubsectionSizeQuads + 1) * NumSubsections) * ComponentCounts.Y;
-		bCreateFromScratch = (Landscape->WeightmapRTList.Num() == 0);
-
-		if (bCreateFromScratch)
-		{
-			Landscape->WeightmapRTList.Init(nullptr, (int32)EWeightmapRTType::WeightmapRT_Count);
-		}
-
-		for (int32 i = 0; i < (int32)EWeightmapRTType::WeightmapRT_Count; ++i)
-		{
-			if (bCreateFromScratch)
-			{
-				FText DisplayName = StaticEnum<EHeightmapRTType>()->GetDisplayValueAsText((EWeightmapRTType)i);
-				FName RTName = MakeUniqueObjectName(GetTransientPackage(), UTextureRenderTarget2D::StaticClass(), FName(*DisplayName.ToString()));
-				Landscape->WeightmapRTList[i] = NewObject<UTextureRenderTarget2D>(GetTransientPackage(), RTName, RF_Transient);
-
-				check(Landscape->WeightmapRTList[i]);
-				Landscape->WeightmapRTList[i]->AddressX = TextureAddress::TA_Clamp;
-				Landscape->WeightmapRTList[i]->AddressY = TextureAddress::TA_Clamp;
-				Landscape->WeightmapRTList[i]->ClearColor = FLinearColor(0.0f, 0.0f, 0.0f, 0.0f);
-				Landscape->WeightmapRTList[i]->RenderTargetFormat = RTF_RGBA8;
-
-				// scratch 1/2/3 RTs are R8 format
-				if ((i >= (int32)EWeightmapRTType::WeightmapRT_Scratch1) &&
-					(i < (int32)EWeightmapRTType::WeightmapRT_Mip0))
-				{
-					Landscape->WeightmapRTList[i]->RenderTargetFormat = RTF_R8;
-				}
-			}
-
-			if (i < (int32)EWeightmapRTType::WeightmapRT_Mip0)
-			{
-				InitOrResizeRT(Landscape->WeightmapRTList[i], CurrentMipSizeX, CurrentMipSizeY, bCreateFromScratch);
-			}
-			else // Mips
-			{
-				InitOrResizeRT(Landscape->WeightmapRTList[i], CurrentMipSizeX, CurrentMipSizeY, bCreateFromScratch);
-
-				CurrentMipSizeX >>= 1;
-				CurrentMipSizeY >>= 1;
-			}
-
-			// Only generate required mips RT
-			if (CurrentMipSizeX < ComponentCounts.X && CurrentMipSizeY < ComponentCounts.Y)
-			{
-				break;
-			}
-		}
-
-		InitializeLayersWeightmapResources();
-	}
-	return true;
-}
-
 void ALandscape::ConvertNonEditLayerLandscape()
 {
 	// Detect any attempt to call this in the middle of UpdateLayersContent.  If called from blueprint, log an error and return early instead of asserting.
@@ -2744,74 +1380,12 @@
 	// This will only copy data from proxies that have completed the registration process
 	CopyOldDataToDefaultLayer();
 }
- 
+
 // Deprecated
 void ALandscape::ToggleCanHaveLayersContent()
 {
 }
 
-void ALandscape::ReleaseLayersRenderingResource()
-{
-	ULandscapeInfo* Info = GetLandscapeInfo();
-	if (Info == nullptr)
-	{
-		return;
-	}
-
-	Info->ForEachLandscapeProxy([&](ALandscapeProxy* Proxy)
-	{
-		for (auto& ItPair : Proxy->HeightmapsCPUReadback)
-		{
-			FLandscapeEditLayerReadback* HeightmapCPUReadback = ItPair.Value;
-			delete HeightmapCPUReadback;
-		}
-		Proxy->HeightmapsCPUReadback.Empty();
-
-		for (auto& ItPair : Proxy->WeightmapsCPUReadback)
-		{
-			FLandscapeEditLayerReadback* WeightmapCPUReadback = ItPair.Value;
-			delete WeightmapCPUReadback;
-		}
-		Proxy->WeightmapsCPUReadback.Empty();
-		return true;
-	});
-
-	if (CombinedLayersWeightmapAllMaterialLayersResource != nullptr)
-	{
-		BeginReleaseResource(CombinedLayersWeightmapAllMaterialLayersResource);
-	}
-
-	if (CurrentLayersWeightmapAllMaterialLayersResource != nullptr)
-	{
-		BeginReleaseResource(CurrentLayersWeightmapAllMaterialLayersResource);
-	}
-
-	if (WeightmapScratchExtractLayerTextureResource != nullptr)
-	{
-		BeginReleaseResource(WeightmapScratchExtractLayerTextureResource);
-	}
-
-	if (WeightmapScratchPackLayerTextureResource != nullptr)
-	{
-		BeginReleaseResource(WeightmapScratchPackLayerTextureResource);
-	}
-
-	{
-		TRACE_CPUPROFILER_EVENT_SCOPE(LandscapeLayers_Flush_ResourceRelease);
-		FlushRenderingCommands();
-	}
-
-	delete CombinedLayersWeightmapAllMaterialLayersResource;
-	delete CurrentLayersWeightmapAllMaterialLayersResource;
-	delete WeightmapScratchExtractLayerTextureResource;
-	delete WeightmapScratchPackLayerTextureResource;
-
-	CombinedLayersWeightmapAllMaterialLayersResource = nullptr;
-	CurrentLayersWeightmapAllMaterialLayersResource = nullptr;
-	WeightmapScratchExtractLayerTextureResource = nullptr;
-	WeightmapScratchPackLayerTextureResource = nullptr;
-}
-
 FIntPoint ALandscape::ComputeComponentCounts() const
 {
 	ULandscapeInfo* Info = GetLandscapeInfo();
@@ -3393,788 +1967,6 @@
 	return Result;
 }
 
-void ALandscape::CopyTexturePS(const FString& InSourceDebugName, FTextureResource* InSourceResource, const FString& InDestDebugName, FTextureResource* InDestResource) const
-{
-	check(InSourceResource != nullptr);
-	check(InDestResource != nullptr);
-
-	ENQUEUE_RENDER_COMMAND(LandscapeLayers_Cmd_CopyTexturePS)(
-		[InSourceDebugName, InSourceResource, InDestDebugName, InDestResource](FRHICommandListImmediate& RHICmdList)
-	{
-		TRACE_CPUPROFILER_EVENT_SCOPE(LandscapeLayers_RT_CopyTexturePS);
-		RHI_BREADCRUMB_EVENT_STAT_F(RHICmdList, LandscapeLayers_CopyTexturePS, "LandscapeLayers_CopyTexturePS", "LandscapeLayers_CopyTexturePS %s -> %s", InSourceDebugName, InDestDebugName);
-		SCOPED_GPU_STAT(RHICmdList, LandscapeLayers_CopyTexturePS);
-
-		check(InSourceResource->GetSizeX() == InDestResource->GetSizeX());
-		check(InSourceResource->GetSizeY() == InDestResource->GetSizeY());
-		FRHIRenderPassInfo RPInfo(InDestResource->TextureRHI, ERenderTargetActions::DontLoad_Store);
-		RHICmdList.BeginRenderPass(RPInfo, TEXT("CopyTexture"));
-
-		FGlobalShaderMap* GlobalShaderMap = GetGlobalShaderMap(GMaxRHIFeatureLevel);
-		TShaderMapRef< FLandscapeCopyTextureVS > VertexShader(GlobalShaderMap);
-		TShaderMapRef< FLandscapeCopyTexturePS > PixelShader(GlobalShaderMap);
-
-		FGraphicsPipelineStateInitializer GraphicsPSOInit;
-		RHICmdList.ApplyCachedRenderTargets(GraphicsPSOInit);
-		GraphicsPSOInit.BlendState = TStaticBlendState<>::GetRHI();
-		GraphicsPSOInit.RasterizerState = TStaticRasterizerState<>::GetRHI();
-		GraphicsPSOInit.DepthStencilState = TStaticDepthStencilState<false, CF_Always>::GetRHI();
-		GraphicsPSOInit.PrimitiveType = PT_TriangleList;
-		GraphicsPSOInit.BoundShaderState.VertexDeclarationRHI = GetVertexDeclarationFVector4();
-		GraphicsPSOInit.BoundShaderState.VertexShaderRHI = VertexShader.GetVertexShader();
-		GraphicsPSOInit.BoundShaderState.PixelShaderRHI = PixelShader.GetPixelShader();
-		SetGraphicsPipelineState(RHICmdList, GraphicsPSOInit, 0);
-
-		SetShaderParametersLegacyPS(RHICmdList, PixelShader, InSourceResource->TextureRHI);
-
-		RHICmdList.SetViewport(0.0f, 0.0f, 0.0f, (float)InDestResource->GetSizeX(), (float)InDestResource->GetSizeY(), 1.0f);
-		RHICmdList.DrawIndexedPrimitive(GTwoTrianglesIndexBuffer.IndexBufferRHI, 0, 0, 4, 0, 2, 1);
-
-		RHICmdList.EndRenderPass();
-	});
-}
-
-void ALandscape::DrawWeightmapComponentsToRenderTarget(const FString& InDebugName, const TArray<FIntPoint>& InSectionBaseList, const FVector2f& InScaleBias, TArray<FVector2f>* InScaleBiasPerSection, UTexture* InWeightmapRTRead, UTextureRenderTarget2D* InOptionalWeightmapRTRead2, UTextureRenderTarget2D* InWeightmapRTWrite,
-	ERTDrawingType InDrawType, bool InClearRTWrite, FLandscapeLayersWeightmapShaderParameters& InShaderParams, uint8 InMipRender) const
-{
-	check(InWeightmapRTRead != nullptr);
-	check(InWeightmapRTWrite != nullptr);
-	check(InScaleBiasPerSection == nullptr || InScaleBiasPerSection->Num() == InSectionBaseList.Num());
-
-	FIntPoint WeightmapWriteTextureSize(InWeightmapRTWrite->SizeX, InWeightmapRTWrite->SizeY);
-	FIntPoint WeightmapReadTextureSize(InWeightmapRTRead->Source.GetSizeX(), InWeightmapRTRead->Source.GetSizeY());
-	UTextureRenderTarget2D* WeightmapRTRead = Cast<UTextureRenderTarget2D>(InWeightmapRTRead);
-
-	if (WeightmapRTRead != nullptr)
-	{
-		WeightmapReadTextureSize.X = WeightmapRTRead->SizeX;
-		WeightmapReadTextureSize.Y = WeightmapRTRead->SizeY;
-	}
-
-	// Quad Setup
-	TArray<FLandscapeLayersTriangle> TriangleList;
-	TriangleList.Reserve(InSectionBaseList.Num() * 2 * NumSubsections);
-
-	for (int i = 0; i < InSectionBaseList.Num(); ++i)
-	{
-		const FVector2f& WeightmapScaleBias = InScaleBiasPerSection != nullptr ? (*InScaleBiasPerSection)[i] : InScaleBias;
-		switch (InDrawType)
-		{
-		case ERTDrawingType::RTAtlas:
-		{
-			GenerateLayersRenderQuadsAtlas(InSectionBaseList[i], FVector2D(WeightmapScaleBias), static_cast<float>(SubsectionSizeQuads), WeightmapReadTextureSize, WeightmapWriteTextureSize, TriangleList);
-		} break;
-
-		case ERTDrawingType::RTAtlasToNonAtlas:
-		{
-			GenerateLayersRenderQuadsAtlasToNonAtlas(InSectionBaseList[i], FVector2D(WeightmapScaleBias), static_cast<float>(SubsectionSizeQuads), WeightmapReadTextureSize, WeightmapWriteTextureSize, TriangleList);
-		} break;
-
-		case ERTDrawingType::RTNonAtlas:
-		{
-			GenerateLayersRenderQuadsNonAtlas(InSectionBaseList[i], FVector2D(WeightmapScaleBias), static_cast<float>(SubsectionSizeQuads), WeightmapReadTextureSize, WeightmapWriteTextureSize, TriangleList);
-		} break;
-
-		case ERTDrawingType::RTNonAtlasToAtlas:
-		{
-			GenerateLayersRenderQuadsNonAtlasToAtlas(InSectionBaseList[i], FVector2D(WeightmapScaleBias), static_cast<float>(SubsectionSizeQuads), WeightmapReadTextureSize, WeightmapWriteTextureSize, TriangleList);
-		} break;
-
-		case ERTDrawingType::RTMips:
-		{
-			GenerateLayersRenderQuadsMip(InSectionBaseList[i], FVector2D(WeightmapScaleBias), static_cast<float>(SubsectionSizeQuads), WeightmapReadTextureSize, WeightmapWriteTextureSize, InMipRender, TriangleList);
-		} break;
-
-		default:
-		{
-			check(false);
-			return;
-		}
-		}
-	}
-
-	InShaderParams.ReadWeightmap1 = InWeightmapRTRead;
-	InShaderParams.ReadWeightmap2 = InOptionalWeightmapRTRead2;
-	InShaderParams.CurrentMipComponentVertexCount = ((SubsectionSizeQuads + 1) >> InMipRender);
-
-	if (InMipRender > 0)
-	{
-		InShaderParams.CurrentMipSize = WeightmapWriteTextureSize;
-		InShaderParams.ParentMipSize = WeightmapReadTextureSize;
-	}
-
-	FMatrix ProjectionMatrix = AdjustProjectionMatrixForRHI(FTranslationMatrix(FVector(0, 0, 0)) *
-		FMatrix(FPlane(1.0f / (FMath::Max<uint32>(WeightmapWriteTextureSize.X, 1.f) / 2.0f), 0.0, 0.0f, 0.0f), FPlane(0.0f, -1.0f / (FMath::Max<uint32>(WeightmapWriteTextureSize.Y, 1.f) / 2.0f), 0.0f, 0.0f), FPlane(0.0f, 0.0f, 1.0f, 0.0f), FPlane(-1.0f, 1.0f, 0.0f, 1.0f)));
-
-	FLandscapeLayersWeightmapRender_RenderThread LayersRender(InDebugName, InWeightmapRTWrite, WeightmapWriteTextureSize, WeightmapReadTextureSize, ProjectionMatrix, InShaderParams, InMipRender, TriangleList);
-
-	ENQUEUE_RENDER_COMMAND(LandscapeLayers_Cmd_RenderWeightmap)(
-		[LayersRender, InDebugName, InDrawType, InClearRTWrite](FRHICommandListImmediate& RHICmdList) mutable
-	{
-		SCOPED_DRAW_EVENTF(RHICmdList, LandscapeLayers, TEXT("DrawWeightmapComponentsToRenderTarget %s (%s)"), InDebugName, StaticEnum<EWeightmapRTType>()->GetDisplayValueAsText(InDrawType).ToString());
-		TRACE_CPUPROFILER_EVENT_SCOPE(LandscapeLayers_RT_RenderWeightmap);
-		LayersRender.Render(RHICmdList, InClearRTWrite);
-	});
-
-	PrintLayersDebugRT(InDebugName, InWeightmapRTWrite, InMipRender, false);
-}
-
-void ALandscape::DrawWeightmapComponentsToRenderTarget(const FString& InDebugName, const TArray<ULandscapeComponent*>& InComponentsToDraw, const FIntPoint& InLandscapeBase, UTexture* InWeightmapRTRead, UTextureRenderTarget2D* InOptionalWeightmapRTRead2, UTextureRenderTarget2D* InWeightmapRTWrite,
-	ERTDrawingType InDrawType, bool InClearRTWrite, FLandscapeLayersWeightmapShaderParameters& InShaderParams, uint8 InMipRender) const
-{
-
-	TArray<FIntPoint> SectionBaseList;
-	SectionBaseList.Reserve(InComponentsToDraw.Num());
-	TArray<FVector2f> WeightmapScaleBiasList;
-	WeightmapScaleBiasList.Reserve(InComponentsToDraw.Num());
-
-	for (ULandscapeComponent* Component : InComponentsToDraw)
-	{
-		FVector2f WeightmapScaleBias(static_cast<float>(Component->WeightmapScaleBias.Z), static_cast<float>(Component->WeightmapScaleBias.W));
-		WeightmapScaleBiasList.Add(WeightmapScaleBias);
-
-		FIntPoint ComponentSectionBase = Component->GetSectionBase() - InLandscapeBase;
-		SectionBaseList.Add(ComponentSectionBase);
-	}
-
-	DrawWeightmapComponentsToRenderTarget(InDebugName, SectionBaseList, FVector2f::ZeroVector, &WeightmapScaleBiasList, InWeightmapRTRead, InOptionalWeightmapRTRead2, InWeightmapRTWrite, InDrawType, InClearRTWrite, InShaderParams, InMipRender);
-
-	PrintLayersDebugRT(InDebugName, InWeightmapRTWrite, InMipRender, false);
-}
-
-void ALandscape::DrawWeightmapComponentToRenderTargetMips(const TArray<FVector2f>& InTexturePositionsToDraw, UTexture* InReadWeightmap, bool InClearRTWrite, struct FLandscapeLayersWeightmapShaderParameters& InShaderParams) const
-{
-	int32 CurrentMip = 1;
-	UTexture* ReadMipRT = InReadWeightmap;
-
-	// Convert from Texture position to SectionBase
-	int32 LocalComponentSizeQuad = SubsectionSizeQuads * NumSubsections;
-	int32 LocalComponentSizeVerts = (SubsectionSizeQuads + 1) * NumSubsections;
-
-	TArray<FIntPoint> SectionBaseToDraw;
-	SectionBaseToDraw.Reserve(InTexturePositionsToDraw.Num());
-
-	for (const FVector2f& TexturePosition : InTexturePositionsToDraw)
-	{
-		FVector2f PositionOffset(static_cast<float>(FMath::RoundToInt(TexturePosition.X / LocalComponentSizeVerts)), static_cast<float>(FMath::RoundToInt(TexturePosition.Y / LocalComponentSizeVerts)));
-		SectionBaseToDraw.Add(FIntPoint(static_cast<int32>(PositionOffset.X * LocalComponentSizeQuad), static_cast<int32>(PositionOffset.Y * LocalComponentSizeQuad)));
-	}
-
-	FVector2f WeightmapScaleBias(0.0f, 0.0f); // we dont need a scale bias for mip drawing
-
-	for (int32 MipRTIndex = (int32)EWeightmapRTType::WeightmapRT_Mip1; MipRTIndex < (int32)EWeightmapRTType::WeightmapRT_Count; ++MipRTIndex)
-	{
-		UTextureRenderTarget2D* WriteMipRT = WeightmapRTList[MipRTIndex];
-
-		if (WriteMipRT != nullptr)
-		{
-			DrawWeightmapComponentsToRenderTarget(FString::Printf(TEXT("LS Weight: %s = -> %s Mips %d"), *ReadMipRT->GetName(), *WriteMipRT->GetName(), CurrentMip),
-				SectionBaseToDraw, WeightmapScaleBias, nullptr, ReadMipRT, nullptr, WriteMipRT, ERTDrawingType::RTMips, InClearRTWrite, InShaderParams, static_cast<uint8>(CurrentMip));
-			++CurrentMip;
-		}
-
-		ReadMipRT = WeightmapRTList[MipRTIndex];
-	}
-}
-
-void ALandscape::ClearLayersWeightmapTextureResource(const FString& InDebugName, FTextureRenderTargetResource* InTextureResourceToClear) const
-{
-	LandscapeLayersWeightmapClear_RenderThread LayersClear(InDebugName, InTextureResourceToClear);
-
-	ENQUEUE_RENDER_COMMAND(LandscapeLayers_Cmd_Clear)(
-		[LayersClear](FRHICommandListImmediate& RHICmdList) mutable
-	{
-		TRACE_CPUPROFILER_EVENT_SCOPE(LandscapeLayers_RT_Clear);
-		LayersClear.Clear(RHICmdList);
-	});
-}
-
-void ALandscape::DrawHeightmapComponentsToRenderTargetMips(const TArray<ULandscapeComponent*>& InComponentsToDraw, const FIntPoint& InLandscapeBase, UTexture* InReadHeightmap, bool InClearRTWrite, struct FLandscapeLayersHeightmapShaderParameters& InShaderParams) const
-{
-	int32 CurrentMip = 1;
-	UTexture* ReadMipRT = InReadHeightmap;
-
-	for (int32 MipRTIndex = (int32)EHeightmapRTType::HeightmapRT_Mip1; MipRTIndex < (int32)EHeightmapRTType::HeightmapRT_Count; ++MipRTIndex)
-	{
-		UTextureRenderTarget2D* WriteMipRT = HeightmapRTList[MipRTIndex];
-
-		if (WriteMipRT != nullptr)
-		{
-			DrawHeightmapComponentsToRenderTarget(FString::Printf(TEXT("LS Height: %s = -> %s CombinedAtlasWithMips %d"), *ReadMipRT->GetName(), *WriteMipRT->GetName(), CurrentMip),
-				InComponentsToDraw, InLandscapeBase, ReadMipRT, nullptr, WriteMipRT, ERTDrawingType::RTMips, InClearRTWrite, InShaderParams, static_cast<uint8>(CurrentMip));
-			++CurrentMip;
-		}
-
-		ReadMipRT = HeightmapRTList[MipRTIndex];
-	}
-}
-
-void ALandscape::DrawHeightmapComponentsToRenderTarget(const FString& InDebugName, const TArray<ULandscapeComponent*>& InComponentsToDraw, const FIntPoint& InLandscapeBase, UTexture* InHeightmapRTRead, UTextureRenderTarget2D* InOptionalHeightmapRTRead2, UTextureRenderTarget2D* InHeightmapRTWrite,
-	ERTDrawingType InDrawType, bool InClearRTWrite, FLandscapeLayersHeightmapShaderParameters& InShaderParams, uint8 InMipRender) const
-{
-	check(InHeightmapRTRead != nullptr);
-	check(InHeightmapRTWrite != nullptr);
-	if (!FApp::CanEverRender())
-	{
-		return;
-	}
-
-	FIntPoint HeightmapWriteTextureSize(InHeightmapRTWrite->SizeX, InHeightmapRTWrite->SizeY);
-	FIntPoint HeightmapReadTextureSize(InHeightmapRTRead->Source.GetSizeX(), InHeightmapRTRead->Source.GetSizeY());
-	UTextureRenderTarget2D* HeightmapRTRead = Cast<UTextureRenderTarget2D>(InHeightmapRTRead);
-
-	if (HeightmapRTRead != nullptr)
-	{
-		HeightmapReadTextureSize.X = HeightmapRTRead->SizeX;
-		HeightmapReadTextureSize.Y = HeightmapRTRead->SizeY;
-	}
-
-	// Quad Setup
-	TArray<FLandscapeLayersTriangle> TriangleList;
-	TriangleList.Reserve(InComponentsToDraw.Num() * 2 * NumSubsections);
-
-	for (ULandscapeComponent* Component : InComponentsToDraw)
-	{
-		FVector2f HeightmapScaleBias(static_cast<float>(Component->HeightmapScaleBias.Z), static_cast<float>(Component->HeightmapScaleBias.W));
-		FIntPoint ComponentSectionBase = Component->GetSectionBase() - InLandscapeBase;
-
-		switch (InDrawType)
-		{
-			case ERTDrawingType::RTAtlas:
-			{
-				GenerateLayersRenderQuadsAtlas(ComponentSectionBase, FVector2D(HeightmapScaleBias), static_cast<float>(SubsectionSizeQuads), HeightmapReadTextureSize, HeightmapWriteTextureSize, TriangleList);
-			} break;
-
-			case ERTDrawingType::RTAtlasToNonAtlas:
-			{
-				GenerateLayersRenderQuadsAtlasToNonAtlas(ComponentSectionBase, FVector2D(HeightmapScaleBias), static_cast<float>(SubsectionSizeQuads), HeightmapReadTextureSize, HeightmapWriteTextureSize, TriangleList);
-			} break;
-
-			case ERTDrawingType::RTNonAtlas:
-			{
-				GenerateLayersRenderQuadsNonAtlas(ComponentSectionBase, FVector2D(HeightmapScaleBias), static_cast<float>(SubsectionSizeQuads), HeightmapReadTextureSize, HeightmapWriteTextureSize, TriangleList);
-			} break;
-
-			case ERTDrawingType::RTNonAtlasToAtlas:
-			{
-				GenerateLayersRenderQuadsNonAtlasToAtlas(ComponentSectionBase, FVector2D(HeightmapScaleBias), static_cast<float>(SubsectionSizeQuads), HeightmapReadTextureSize, HeightmapWriteTextureSize, TriangleList);
-			} break;
-
-			case ERTDrawingType::RTMips:
-			{
-				GenerateLayersRenderQuadsMip(ComponentSectionBase, FVector2D(HeightmapScaleBias), static_cast<float>(SubsectionSizeQuads), HeightmapReadTextureSize, HeightmapWriteTextureSize, InMipRender, TriangleList);
-			} break;
-
-		default:
-		{
-			check(false);
-			return;
-		}
-		}
-	}
-
-	InShaderParams.ReadHeightmap1 = InHeightmapRTRead;
-	InShaderParams.ReadHeightmap2 = InOptionalHeightmapRTRead2;
-	InShaderParams.HeightmapSize = HeightmapReadTextureSize;
-	InShaderParams.CurrentMipComponentVertexCount = ((SubsectionSizeQuads + 1) >> InMipRender);
-
-	if (InMipRender > 0)
-	{
-		InShaderParams.CurrentMipSize = HeightmapWriteTextureSize;
-		InShaderParams.ParentMipSize = HeightmapReadTextureSize;
-	}
-
-	FMatrix ProjectionMatrix = AdjustProjectionMatrixForRHI(FTranslationMatrix(FVector(0, 0, 0)) *
-		FMatrix(FPlane(1.0f / (FMath::Max<uint32>(HeightmapWriteTextureSize.X, 1.f) / 2.0f), 0.0, 0.0f, 0.0f), FPlane(0.0f, -1.0f / (FMath::Max<uint32>(HeightmapWriteTextureSize.Y, 1.f) / 2.0f), 0.0f, 0.0f), FPlane(0.0f, 0.0f, 1.0f, 0.0f), FPlane(-1.0f, 1.0f, 0.0f, 1.0f)));
-
-	FLandscapeLayersHeightmapRender_RenderThread LayersRender(InDebugName, InHeightmapRTWrite, HeightmapWriteTextureSize, HeightmapReadTextureSize, ProjectionMatrix, InShaderParams, InMipRender, TriangleList);
-
-	ENQUEUE_RENDER_COMMAND(LandscapeLayers_Cmd_RenderHeightmap)(
-		[LayersRender, InDebugName, InDrawType, InClearRTWrite](FRHICommandListImmediate& RHICmdList) mutable
-	{
-		SCOPED_DRAW_EVENTF(RHICmdList, LandscapeLayers, TEXT("DrawHeightmapComponentsToRenderTarget %s (%s)"), InDebugName, StaticEnum<EHeightmapRTType>()->GetDisplayValueAsText(InDrawType).ToString());
-		TRACE_CPUPROFILER_EVENT_SCOPE(LandscapeLayers_RT_RenderHeightmap);
-		LayersRender.Render(RHICmdList, InClearRTWrite);
-	});
-
-	PrintLayersDebugRT(InDebugName, InHeightmapRTWrite, InMipRender, true, InShaderParams.GenerateNormals);
-}
-
-void ALandscape::GenerateLayersRenderQuad(const FIntPoint& InVertexPosition, float InVertexSize, const FVector2D& InUVStart, const FVector2D& InUVSize, TArray<FLandscapeLayersTriangle>& OutTriangles) const
-{
-	// Set min/max values for rectangle in XY and UV.
-	const float X[] = { static_cast<float>(InVertexPosition.X), static_cast<float>(InVertexPosition.X) + InVertexSize };
-	const float Y[] = { static_cast<float>(InVertexPosition.Y), static_cast<float>(InVertexPosition.Y) + InVertexSize };
-	const float U[] = { static_cast<float>(InUVStart.X), static_cast<float>(InUVStart.X + InUVSize.X) };
-	const float V[] = { static_cast<float>(InUVStart.Y), static_cast<float>(InUVStart.Y + InUVSize.Y) };
-
-	// Helper function for creating a vertex from given min/max indices.
-	auto SetVertex = [&X, &Y, &U, &V](FLandscapeLayersVertex& Vertex, const int32 Index1, const int32 Index2)
-	{
-		Vertex.Position.X = X[Index1];
-		Vertex.Position.Y = Y[Index2];
-		Vertex.UV.X = U[Index1];
-		Vertex.UV.Y = V[Index2];
-	};
-
-	FLandscapeLayersTriangle Tri;
-
-	// Create first triangle.
-	SetVertex(Tri.V0, 0, 0);
-	SetVertex(Tri.V1, 1, 0);
-	SetVertex(Tri.V2, 1, 1);
-	OutTriangles.Add(Tri);
-
-	// Create second triangle; V0 is identical to the previous triangle.
-	SetVertex(Tri.V1, 1, 1);
-	SetVertex(Tri.V2, 0, 1);
-	OutTriangles.Add(Tri);
-}
-
-void ALandscape::GenerateLayersRenderQuadsAtlas(const FIntPoint& InSectionBase, const FVector2D& InScaleBias, float InSubSectionSizeQuad, const FIntPoint& InReadSize, const FIntPoint& InWriteSize, TArray<FLandscapeLayersTriangle>& OutTriangles) const
-{
-	FIntPoint ComponentSectionBase = InSectionBase;
-	FIntPoint UVComponentSectionBase = InSectionBase;
-
-	const int32 SubsectionSizeVerts = static_cast<int32>(InSubSectionSizeQuad) + 1;
-	const int32 LocalComponentSizeQuad = static_cast<int32>(InSubSectionSizeQuad * NumSubsections);
-	const int32 LocalComponentSizeVerts = SubsectionSizeVerts * NumSubsections;
-
-	FVector2D PositionOffset(FMath::RoundToInt((float)ComponentSectionBase.X / (float)LocalComponentSizeQuad), FMath::RoundToInt((float)ComponentSectionBase.Y / (float)LocalComponentSizeQuad));
-	FVector2D ComponentsPerTexture(FMath::RoundToInt((float)InWriteSize.X / (float)LocalComponentSizeQuad), FMath::RoundToInt((float)InWriteSize.Y / (float)LocalComponentSizeQuad));
-
-	if (InReadSize.X >= InWriteSize.X)
-	{
-		if (InReadSize.X == InWriteSize.X)
-		{
-			if (ComponentsPerTexture.X > 1.0f)
-			{
-				UVComponentSectionBase.X = static_cast<int32>(PositionOffset.X * LocalComponentSizeVerts);
-			}
-			else
-			{
-				UVComponentSectionBase.X -= (UVComponentSectionBase.X + LocalComponentSizeQuad > InWriteSize.X)
-					                            ? static_cast<int32>(FMath::FloorToInt(PositionOffset.X / ComponentsPerTexture.X) * ComponentsPerTexture.X * LocalComponentSizeQuad)
-					                            : 0;
-			}
-		}
-
-		ComponentSectionBase.X -= (ComponentSectionBase.X + LocalComponentSizeQuad > InWriteSize.X)
-			                          ? static_cast<int32>(FMath::FloorToInt(PositionOffset.X / ComponentsPerTexture.X) * ComponentsPerTexture.X * LocalComponentSizeQuad)
-			                          : 0;
-		PositionOffset.X = ComponentSectionBase.X / LocalComponentSizeQuad;
-	}
-
-	if (InReadSize.Y >= InWriteSize.Y)
-	{
-		if (InReadSize.Y == InWriteSize.Y)
-		{
-			if (ComponentsPerTexture.Y > 1.0f)
-			{
-				UVComponentSectionBase.Y = static_cast<int32>(PositionOffset.Y * LocalComponentSizeVerts);
-			}
-			else
-			{
-				UVComponentSectionBase.Y -= (UVComponentSectionBase.Y + LocalComponentSizeQuad > InWriteSize.Y)
-					                            ? static_cast<int32>(FMath::FloorToInt((PositionOffset.Y / ComponentsPerTexture.Y)) * ComponentsPerTexture.Y * LocalComponentSizeQuad)
-					                            : 0;
-			}
-		}
-
-		ComponentSectionBase.Y -= (ComponentSectionBase.Y + LocalComponentSizeQuad > InWriteSize.Y)
-			                          ? static_cast<int32>(FMath::FloorToInt((PositionOffset.Y / ComponentsPerTexture.Y)) * ComponentsPerTexture.Y * LocalComponentSizeQuad)
-			                          : 0;
-		PositionOffset.Y = ComponentSectionBase.Y / LocalComponentSizeQuad;
-	}
-
-	ComponentSectionBase.X = static_cast<int32>(PositionOffset.X * LocalComponentSizeVerts);
-	ComponentSectionBase.Y = static_cast<int32>(PositionOffset.Y * LocalComponentSizeVerts);
-
-	FVector2D UVSize((float)SubsectionSizeVerts / (float)InReadSize.X, (float)SubsectionSizeVerts / (float)InReadSize.Y);
-	FIntPoint SubSectionSectionBase;
-
-	for (int8 SubY = 0; SubY < NumSubsections; ++SubY)
-	{
-		for (int8 SubX = 0; SubX < NumSubsections; ++SubX)
-		{
-			SubSectionSectionBase.X = ComponentSectionBase.X + SubsectionSizeVerts * SubX;
-			SubSectionSectionBase.Y = ComponentSectionBase.Y + SubsectionSizeVerts * SubY;
-
-			// Offset for this component's data in texture
-			FVector2D UVStart;
-
-			if (InReadSize.X >= InWriteSize.X)
-			{
-				UVStart.X = ((float)UVComponentSectionBase.X / (float)InReadSize.X) + UVSize.X * (float)SubX;
-			}
-			else
-			{
-				UVStart.X = InScaleBias.X + UVSize.X * (float)SubX;
-			}
-
-			if (InReadSize.Y >= InWriteSize.Y)
-			{
-				UVStart.Y = ((float)UVComponentSectionBase.Y / (float)InReadSize.Y) + UVSize.Y * (float)SubY;
-			}
-			else
-			{
-				UVStart.Y = InScaleBias.Y + UVSize.Y * (float)SubY;
-			}
-
-			GenerateLayersRenderQuad(SubSectionSectionBase, static_cast<float>(SubsectionSizeVerts), UVStart, UVSize, OutTriangles);
-		}
-	}
-}
-
-void ALandscape::GenerateLayersRenderQuadsMip(const FIntPoint& InSectionBase, const FVector2D& InScaleBias, float InSubSectionSizeQuad, const FIntPoint& InReadSize, const FIntPoint& InWriteSize, uint8 InCurrentMip, TArray<FLandscapeLayersTriangle>& OutTriangles) const
-{
-	int32 SubsectionSizeVerts = static_cast<int32>(InSubSectionSizeQuad) + 1;
-	int32 LocalComponentSizeQuad = static_cast<int32>(InSubSectionSizeQuad) * NumSubsections;
-	int32 LocalComponentSizeVerts = SubsectionSizeVerts * NumSubsections;
-	int32 MipSubsectionSizeVerts = SubsectionSizeVerts >> InCurrentMip;
-	int32 MipLocalComponentSizeVerts = MipSubsectionSizeVerts * NumSubsections;
-
-	FVector2D PositionOffset(FMath::RoundToInt((float)InSectionBase.X / (float)LocalComponentSizeQuad), FMath::RoundToInt((float)InSectionBase.Y / (float)LocalComponentSizeQuad));
-	FVector2D ComponentsPerTexture(FMath::RoundToInt((float)InWriteSize.X / (float)LocalComponentSizeQuad), FMath::RoundToInt((float)InWriteSize.Y / (float)LocalComponentSizeQuad));
-
-	FIntPoint ComponentSectionBase(static_cast<int32>(PositionOffset.X * MipLocalComponentSizeVerts), static_cast<int32>(PositionOffset.Y * MipLocalComponentSizeVerts));
-	FIntPoint UVComponentSectionBase(static_cast<int32>(PositionOffset.X * LocalComponentSizeVerts), static_cast<int32>(PositionOffset.Y * LocalComponentSizeVerts));
-	FVector2D UVSize((float)(SubsectionSizeVerts >> (InCurrentMip - 1)) / (float)InReadSize.X, (float)(SubsectionSizeVerts >> (InCurrentMip - 1)) / (float)InReadSize.Y);
-	FIntPoint SubSectionSectionBase;
-
-	for (int8 SubY = 0; SubY < NumSubsections; ++SubY)
-	{
-		for (int8 SubX = 0; SubX < NumSubsections; ++SubX)
-		{
-			SubSectionSectionBase.X = ComponentSectionBase.X + MipSubsectionSizeVerts * SubX;
-			SubSectionSectionBase.Y = ComponentSectionBase.Y + MipSubsectionSizeVerts * SubY;
-
-			// Offset for this component's data in texture
-			FVector2D UVStart(((float)(UVComponentSectionBase.X >> (InCurrentMip - 1)) / (float)InReadSize.X) + UVSize.X * (float)SubX, ((float)(UVComponentSectionBase.Y >> (InCurrentMip - 1)) / (float)InReadSize.Y) + UVSize.Y * (float)SubY);
-
-			GenerateLayersRenderQuad(SubSectionSectionBase, static_cast<float>(MipSubsectionSizeVerts), UVStart, UVSize, OutTriangles);
-		}
-	}
-}
-
-void ALandscape::GenerateLayersRenderQuadsAtlasToNonAtlas(const FIntPoint& InSectionBase, const FVector2D& InScaleBias, float InSubSectionSizeQuad, const FIntPoint& InReadSize, const FIntPoint& InWriteSize, TArray<struct FLandscapeLayersTriangle>& OutTriangles) const
-{
-	int32 SubsectionSizeVerts = static_cast<int32>(InSubSectionSizeQuad) + 1;
-	FVector2D UVSize((float)SubsectionSizeVerts / (float)InReadSize.X, (float)SubsectionSizeVerts / (float)InReadSize.Y);
-
-	for (int8 SubY = 0; SubY < NumSubsections; ++SubY)
-	{
-		for (int8 SubX = 0; SubX < NumSubsections; ++SubX)
-		{
-			FIntPoint SubSectionSectionBase(static_cast<int32>(InSectionBase.X + InSubSectionSizeQuad * SubX), static_cast<int32>(InSectionBase.Y + InSubSectionSizeQuad * SubY));
-			FVector2D PositionOffset(FMath::RoundToInt(SubSectionSectionBase.X / InSubSectionSizeQuad), FMath::RoundToInt(SubSectionSectionBase.Y / InSubSectionSizeQuad));
-			FIntPoint UVComponentSectionBase(static_cast<int32>(PositionOffset.X * SubsectionSizeVerts), static_cast<int32>(PositionOffset.Y * SubsectionSizeVerts));
-
-			// Offset for this component's data in texture
-			FVector2D UVStart;
-
-			if (InReadSize.X >= InWriteSize.X)
-			{
-				UVStart.X = ((float)UVComponentSectionBase.X / (float)InReadSize.X);
-			}
-			else
-			{
-				UVStart.X = InScaleBias.X + UVSize.X * (float)SubX;
-			}
-
-			if (InReadSize.Y >= InWriteSize.Y)
-			{
-				UVStart.Y = ((float)UVComponentSectionBase.Y / (float)InReadSize.Y);
-			}
-			else
-			{
-				UVStart.Y = InScaleBias.Y + UVSize.Y * (float)SubY;
-			}
-
-			GenerateLayersRenderQuad(SubSectionSectionBase, static_cast<float>(SubsectionSizeVerts), UVStart, UVSize, OutTriangles);
-		}
-	}
-}
-
-void ALandscape::GenerateLayersRenderQuadsNonAtlas(const FIntPoint& InSectionBase, const FVector2D& InScaleBias, float InSubSectionSizeQuad, const FIntPoint& InReadSize, const FIntPoint& InWriteSize, TArray<struct FLandscapeLayersTriangle>& OutTriangles) const
-{
-	// We currently only support drawing in non atlas mode with the same texture size
-	check(InReadSize.X == InWriteSize.X && InReadSize.Y == InWriteSize.Y);
-
-	int32 SubsectionSizeVerts = static_cast<int32>(InSubSectionSizeQuad) + 1;
-
-	FVector2D UVSize((float)SubsectionSizeVerts / (float)InReadSize.X, (float)SubsectionSizeVerts / (float)InReadSize.Y);
-
-	for (int8 SubY = 0; SubY < NumSubsections; ++SubY)
-	{
-		for (int8 SubX = 0; SubX < NumSubsections; ++SubX)
-		{
-			FIntPoint SubSectionSectionBase(InSectionBase.X + SubsectionSizeQuads * SubX, InSectionBase.Y + SubsectionSizeQuads * SubY);
-			FVector2D PositionOffset(FMath::RoundToInt(SubSectionSectionBase.X / InSubSectionSizeQuad), FMath::RoundToInt(SubSectionSectionBase.Y / InSubSectionSizeQuad));
-			FIntPoint UVComponentSectionBase(static_cast<int32>(PositionOffset.X * InSubSectionSizeQuad), static_cast<int32>(PositionOffset.Y * InSubSectionSizeQuad));
-
-			// Offset for this component's data in texture
-			FVector2D UVStart(((float)UVComponentSectionBase.X / (float)InReadSize.X), ((float)UVComponentSectionBase.Y / (float)InReadSize.Y));
-			GenerateLayersRenderQuad(SubSectionSectionBase, static_cast<float>(SubsectionSizeVerts), UVStart, UVSize, OutTriangles);
-		}
-	}
-}
-
-void ALandscape::GenerateLayersRenderQuadsNonAtlasToAtlas(const FIntPoint& InSectionBase, const FVector2D& InScaleBias, float InSubSectionSizeQuad, const FIntPoint& InReadSize, const FIntPoint& InWriteSize, TArray<struct FLandscapeLayersTriangle>& OutTriangles) const
-{
-	int32 SubsectionSizeVerts = static_cast<int32>(InSubSectionSizeQuad) + 1;
-	int32 LocalComponentSizeQuad = static_cast<int32>(InSubSectionSizeQuad) * NumSubsections;
-	int32 LocalComponentSizeVerts = SubsectionSizeVerts * NumSubsections;
-
-	FVector2D PositionOffset(FMath::RoundToInt((float)InSectionBase.X / (float)LocalComponentSizeQuad), FMath::RoundToInt((float)InSectionBase.Y / (float)LocalComponentSizeQuad));
-	FIntPoint ComponentSectionBase(static_cast<int32>(PositionOffset.X * LocalComponentSizeVerts), static_cast<int32>(PositionOffset.Y * LocalComponentSizeVerts));
-	FVector2D UVSize((float)SubsectionSizeVerts / (float)InReadSize.X, (float)SubsectionSizeVerts / (float)InReadSize.Y);
-
-	FIntPoint SubSectionSectionBase;
-
-	for (int8 SubY = 0; SubY < NumSubsections; ++SubY)
-	{
-		for (int8 SubX = 0; SubX < NumSubsections; ++SubX)
-		{
-			SubSectionSectionBase.X = ComponentSectionBase.X + SubsectionSizeVerts * SubX;
-			SubSectionSectionBase.Y = ComponentSectionBase.Y + SubsectionSizeVerts * SubY;
-
-			// Offset for this component's data in texture
-			float ScaleBiasZ = (float)InSectionBase.X / (float)InReadSize.X;
-			float ScaleBiasW = (float)InSectionBase.Y / (float)InReadSize.Y;
-			FVector2D UVStart(ScaleBiasZ + ((float)InSubSectionSizeQuad / (float)InReadSize.X) * (float)SubX, ScaleBiasW + ((float)InSubSectionSizeQuad / (float)InReadSize.Y) * (float)SubY);
-
-			GenerateLayersRenderQuad(SubSectionSectionBase, static_cast<float>(SubsectionSizeVerts), UVStart, UVSize, OutTriangles);
-		}
-	}
-}
-
-void ALandscape::PrintLayersDebugHeightData(const FString& InContext, const TArray<FColor>& InHeightmapData, const FIntPoint& InDataSize, uint8 InMipRender, bool InOutputNormals) const
-{
-	bool DisplayDebugPrint = CVarOutputLayersRTContent.GetValueOnAnyThread() == 1 ? true : false;
-	bool DisplayHeightAsDelta = false;
-
-	if (!DisplayDebugPrint)
-	{
-		return;
-	}
-
-	TArray<uint16> HeightData;
-	TArray<FVector> NormalData;
-	HeightData.Reserve(InHeightmapData.Num());
-	NormalData.Reserve(InHeightmapData.Num());
-
-	for (FColor Color : InHeightmapData)
-	{
-		uint16 Height = ((Color.R << 8) | Color.G);
-		HeightData.Add(Height);
-
-		if (InOutputNormals)
-		{
-			FVector Normal;
-			Normal.X = Color.B > 0.0f ? ((float)Color.B / 127.5f - 1.0f) : 0.0f;
-			Normal.Y = Color.A > 0.0f ? ((float)Color.A / 127.5f - 1.0f) : 0.0f;
-			Normal.Z = 0.0f;
-
-			NormalData.Add(Normal);
-		}
-	}
-
-	UE_LOG(LogLandscape, Display, TEXT("Context: %s"), *InContext);
-
-	int32 MipSize = ((SubsectionSizeQuads + 1) >> InMipRender);
-
-	for (int32 Y = 0; Y < InDataSize.Y; ++Y)
-	{
-		FString HeightmapHeightOutput;
-
-		for (int32 X = 0; X < InDataSize.X; ++X)
-		{
-			int32 HeightDelta = HeightData[X + Y * InDataSize.X];
-
-			if (DisplayHeightAsDelta)
-			{
-				HeightDelta = static_cast<int32>(HeightDelta >= LandscapeDataAccess::MidValue ? HeightDelta - LandscapeDataAccess::MidValue : HeightDelta);
-			}
-
-			if (X > 0 && MipSize > 0 && X % MipSize == 0)
-			{
-				HeightmapHeightOutput += FString::Printf(TEXT("  "));
-			}
-
-			FString HeightStr = FString::Printf(TEXT("%d"), HeightDelta);
-
-			int32 PadCount = 5 - HeightStr.Len();
-			if (PadCount > 0)
-			{
-				HeightStr = FString::ChrN(PadCount, '0') + HeightStr;
-			}
-
-			HeightmapHeightOutput += HeightStr + TEXT(" ");
-		}
-
-		if (Y > 0 && MipSize > 0 && Y % MipSize == 0)
-		{
-			UE_LOG(LogLandscape, Display, TEXT(""));
-		}
-
-		UE_LOG(LogLandscape, Display, TEXT("%s"), *HeightmapHeightOutput);
-	}
-
-	if (InOutputNormals)
-	{
-		UE_LOG(LogLandscape, Display, TEXT(""));
-
-		for (int32 Y = 0; Y < InDataSize.Y; ++Y)
-		{
-			FString HeightmapNormaltOutput;
-
-			for (int32 X = 0; X < InDataSize.X; ++X)
-			{
-				FVector Normal = NormalData[X + Y * InDataSize.X];
-
-				if (X > 0 && MipSize > 0 && X % MipSize == 0)
-				{
-					HeightmapNormaltOutput += FString::Printf(TEXT("  "));
-				}
-
-				HeightmapNormaltOutput += FString::Printf(TEXT(" %s"), *Normal.ToString());
-			}
-
-			if (Y > 0 && MipSize > 0 && Y % MipSize == 0)
-			{
-				UE_LOG(LogLandscape, Display, TEXT(""));
-			}
-
-			UE_LOG(LogLandscape, Display, TEXT("%s"), *HeightmapNormaltOutput);
-		}
-	}
-}
-
-void ALandscape::PrintLayersDebugWeightData(const FString& InContext, const TArray<FColor>& InWeightmapData, const FIntPoint& InDataSize, uint8 InMipRender) const
-{
-	bool DisplayDebugPrint = (CVarOutputLayersRTContent.GetValueOnAnyThread() == 1 || CVarOutputLayersWeightmapsRTContent.GetValueOnAnyThread() == 1) ? true : false;
-
-	if (!DisplayDebugPrint)
-	{
-		return;
-	}
-
-	UE_LOG(LogLandscape, Display, TEXT("Context: %s"), *InContext);
-
-	int32 MipSize = ((SubsectionSizeQuads + 1) >> InMipRender);
-
-	for (int32 Y = 0; Y < InDataSize.Y; ++Y)
-	{
-		FString WeightmapOutput;
-
-		for (int32 X = 0; X < InDataSize.X; ++X)
-		{
-			const FColor& Weight = InWeightmapData[X + Y * InDataSize.X];
-
-			if (X > 0 && MipSize > 0 && X % MipSize == 0)
-			{
-				WeightmapOutput += FString::Printf(TEXT("  "));
-			}
-
-			WeightmapOutput += FString::Printf(TEXT("%s "), *Weight.ToString());
-		}
-
-		if (Y > 0 && MipSize > 0 && Y % MipSize == 0)
-		{
-			UE_LOG(LogLandscape, Display, TEXT(""));
-		}
-
-		UE_LOG(LogLandscape, Display, TEXT("%s"), *WeightmapOutput);
-	}
-}
-
-void ALandscape::PrintLayersDebugRT(const FString& InContext, UTextureRenderTarget2D* InDebugRT, uint8 InMipRender, bool InOutputHeight, bool InOutputNormals) const
-{
-	bool DisplayDebugPrint = (CVarOutputLayersRTContent.GetValueOnAnyThread() == 1 || CVarOutputLayersWeightmapsRTContent.GetValueOnAnyThread() == 1) ? true : false;
-
-	if (!DisplayDebugPrint)
-	{
-		return;
-	}
-
-	FTextureRenderTargetResource* RenderTargetResource = InDebugRT->GameThread_GetRenderTargetResource();
-	ENQUEUE_RENDER_COMMAND(LandscapeLayers_Cmd_DebugResolve)(
-		[RenderTargetResource](FRHICommandListImmediate& RHICmdList) mutable
-	{
-		TRACE_CPUPROFILER_EVENT_SCOPE(LandscapeLayers_RT_DebugResolve);
-		TransitionAndCopyTexture(RHICmdList, RenderTargetResource->GetRenderTargetTexture(), RenderTargetResource->TextureRHI, {});
-	});
-
-	{
-		TRACE_CPUPROFILER_EVENT_SCOPE(LandscapeLayers_Flush_DebugResolve);
-		FlushRenderingCommands();
-	}
-
-	int32 MinX, MinY, MaxX, MaxY;
-	const ULandscapeInfo* LandscapeInfo = GetLandscapeInfo();
-	LandscapeInfo->GetLandscapeExtent(MinX, MinY, MaxX, MaxY);
-	FIntRect SampleRect = FIntRect(0, 0, InDebugRT->SizeX, InDebugRT->SizeY);
-
-	FReadSurfaceDataFlags Flags(RCM_UNorm, CubeFace_MAX);
-
-	TArray<FColor> OutputRT;
-	OutputRT.Reserve(SampleRect.Width() * SampleRect.Height());
-
-	InDebugRT->GameThread_GetRenderTargetResource()->ReadPixels(OutputRT, Flags, SampleRect);
-
-	if (InOutputHeight)
-	{
-		PrintLayersDebugHeightData(InContext, OutputRT, FIntPoint(SampleRect.Width(), SampleRect.Height()), InMipRender, InOutputNormals);
-	}
-	else
-	{
-		PrintLayersDebugWeightData(InContext, OutputRT, FIntPoint(SampleRect.Width(), SampleRect.Height()), InMipRender);
-	}
-}
-
-void ALandscape::PrintLayersDebugTextureResource(const FString& InContext, FTextureResource* InTextureResource, uint8 InMipRender, bool InOutputHeight, bool InOutputNormals) const
-{
-	bool DisplayDebugPrint = (CVarOutputLayersRTContent.GetValueOnAnyThread() == 1 || CVarOutputLayersWeightmapsRTContent.GetValueOnAnyThread() == 1) ? true : false;
-
-	if (!DisplayDebugPrint)
-	{
-		return;
-	}
-
-	int32 MinX, MinY, MaxX, MaxY;
-	const ULandscapeInfo* LandscapeInfo = GetLandscapeInfo();
-	LandscapeInfo->GetLandscapeExtent(MinX, MinY, MaxX, MaxY);
-	FIntRect SampleRect = FIntRect(0, 0, InTextureResource->GetSizeX(), InTextureResource->GetSizeY());
-
-	TArray<FColor> OutputTexels;
-	OutputTexels.Reserve(SampleRect.Width() * SampleRect.Height());
-
-	FReadSurfaceDataFlags Flags(RCM_UNorm, CubeFace_MAX);
-	Flags.SetMip(InMipRender);
-
-	ENQUEUE_RENDER_COMMAND(LandscapeLayers_Cmd_Readback)(
-		[InTextureResource, SampleRect = SampleRect, OutData = &OutputTexels, ReadFlags = Flags](FRHICommandListImmediate& RHICmdList) mutable
-	{
-		TRACE_CPUPROFILER_EVENT_SCOPE(LandscapeLayers_RT_Readback);
-		RHICmdList.ReadSurfaceData(InTextureResource->TextureRHI, SampleRect, *OutData, ReadFlags);
-	});
-
-	{
-		TRACE_CPUPROFILER_EVENT_SCOPE(LandscapeLayers_Flush_Readback);
-		FlushRenderingCommands();
-	}
-
-	if (InOutputHeight)
-	{
-		PrintLayersDebugHeightData(InContext, OutputTexels, FIntPoint(SampleRect.Width(), SampleRect.Height()), InMipRender, InOutputNormals);
-	}
-	else
-	{
-		PrintLayersDebugWeightData(InContext, OutputTexels, FIntPoint(SampleRect.Width(), SampleRect.Height()), InMipRender);
-	}
-}
-
 bool ALandscape::PrepareTextureResourcesLimited(bool bInWaitForStreaming)
 {
 	// Rate-limit to once per frame. Note that commandlets don't advance the frame counter.
@@ -4776,7 +2568,6 @@
 };
 ENUM_CLASS_FLAGS(FUpdateLayersContentContext::ERefreshFlags);
 
-// Little struct that holds information common to PerformLayersHeightmapsLocalMerge, PerformLayersHeightmapsBatchedMerge and PerformLayersHeightmapsGlobalMerge
 struct FEditLayersHeightmapMergeParams
 {
 	int32 HeightmapUpdateModes = 0;
@@ -4784,765 +2575,6 @@
 	bool bSkipBrush = false;
 };
 
-// Render-thread version of the data / functions we need for the local merge of edit layers : 
-namespace EditLayersHeightmapLocalMerge_RenderThread
-{
-	struct FComponentRenderInfo
-	{
-		// Name of the component for debug purposes :
-		FString Name;
-
-		// The component's visible layer's heightmaps (Num must be == Num of VisibleEditLayerInfos)
-		TArray<FTexture2DResourceSubregion> VisibleLayerHeightmapTextures;
-
-		// List of 9 component render info indices corresponding to the 9 neighbors of this component (index into ComponentToRenderInfos) :
-		//  Some may not be set and that's fine : the only goal is to know which are the valid neighbors when it comes to stitching adjacent rendered components
-		//  If a neighbor is missing, we can't stitch the border adjacent to it, but it also means it won't contribute to what we really want to compute : the 
-		//  component to resolve, which is guaranteed to have its neighbors present (if those do exist) :
-		TStaticArray<int32, 9> NeighborComponentToRenderInfoIndices;
-	};
-
-	struct FComponentResolveInfo
-	{
-		FComponentResolveInfo(int32 InNumComponentsToRender)
-		{
-			NeighborComponentToRenderInfoBitIndices.Init(false, InNumComponentsToRender);
-		}
-
-		void SetNeighborRenderInfo(int32 InNeighborIndex, int32 InNeighborComponentRenderInfoIndex)
-		{
-			NeighborComponentToRenderInfoIndices[InNeighborIndex] = InNeighborComponentRenderInfoIndex;
-			if (InNeighborComponentRenderInfoIndex != INDEX_NONE)
-			{
-				check(!NeighborComponentToRenderInfoBitIndices[InNeighborComponentRenderInfoIndex]);
-				NeighborComponentToRenderInfoBitIndices[InNeighborComponentRenderInfoIndex] = true;
-			}
-		}
-
-	public:
-		// Index of this components in ComponentToResolveInfos :
-		int32 ComponentToResolveInfoIndex = INDEX_NONE;
-		// Name of the component for debug purposes :
-		FString Name;
-		// Subregion of the heightmap that we want to compute normals for and resolve, that corresponds to this component :
-		FTexture2DResourceSubregion Heightmap;
-
-		// List of 9 component render info indices corresponding to the 9 neighbors of this component (index into ComponentToRenderInfos) :
-		//  Some may not be set if the component is on the border of the landscape, for example :
-		TStaticArray<int32, 9> NeighborComponentToRenderInfoIndices;
-		// Same as NeighborComponentToRenderInfoIndices but as a bit array (1 bit per component to render info) to vastly optimize the division of component resolve infos into batches, which is a O(N^2) operation : 
-		TBitArray<> NeighborComponentToRenderInfoBitIndices;
-	};
-
-	struct FTextureResolveInfo
-	{
-		// Size of the entire texture that needs resolving :
-		FIntPoint TextureSize = FIntPoint(ForceInitToZero);
-		// Number of mips corresponding to that size : 
-		int32 NumMips = 0;
-		// Texture that was updated and needs resolving : 
-		FTexture2DResource* Texture = nullptr;
-		// CPU readback utility to bring back the result on the CPU : 
-		FLandscapeEditLayerReadback* CPUReadback;
-	};
-
-	// Because of heightmaps being shared between one component and another, we have to group the components to render into batches
-	//  where we'll render all of the heightmaps into slices of a single scratch texture array, which we'll then be able to re-assemble into the final, packed, heightmaps (subregions) :
-	struct FComponentResolveBatchInfo
-	{
-		FComponentResolveBatchInfo(int32 InNumComponentsToRender, int32 InBatchIndex)
-			: BatchIndex(InBatchIndex)
-		{
-			ComponentToRenderInfoBitIndices.Init(false, InNumComponentsToRender);
-		}
-
-		void AddComponent(const FComponentResolveInfo& InComponentResolveInfo)
-		{
-			check(!ComponentToResolveInfoIndices.Contains(InComponentResolveInfo.ComponentToResolveInfoIndex));
-			ComponentToResolveInfoIndices.Add(InComponentResolveInfo.ComponentToResolveInfoIndex);
-
-			// Remember all the unique components that this texture needs for resolving:
-			ComponentToRenderInfoBitIndices.CombineWithBitwiseOR(InComponentResolveInfo.NeighborComponentToRenderInfoBitIndices, EBitwiseOperatorFlags::MinSize);
-		}
-
-	public:
-		int32 BatchIndex = INDEX_NONE;
-
-		// Indices (in ComponentToRenderInfos) of the components whose heightmaps we need to render within this batch in order to produce (and then resolve) the textures in TextureToResolveInfos:
-		//  It's a bit array (1 bit per component to render info) to vastly optimize the division of texture resolve infos into batches, which is a O(N^2) operation : 
-		TBitArray<> ComponentToRenderInfoBitIndices;
-		// Indices (in ComponentToResolveInfos) of components whose heightmap subregion needs to be resolved / read back on the CPU :
-		TArray<int32> ComponentToResolveInfoIndices;
-	};
-
-	struct FEditLayerInfo
-	{
-		ELandscapeEditLayerHeightmapBlendMode BlendMode = ELandscapeEditLayerHeightmapBlendMode::Num;
-		float Alpha = 1.0f;
-	};
-
-	// Description of the entire merge pass :
-	struct FMergeInfo
-	{
-		bool NeedsMerge() const
-		{
-			// If no edit layer or if no paint layer present on any edit layer, we've got nothing to do :
-			bool bNeedsMerge = (MaxNumEditLayersTexturesToMerge > 0) && (MaxNumComponentsToRenderPerResolveComponentBatch > 0);
-			check(!bNeedsMerge || !TextureToResolveInfos.IsEmpty()); // If we need merging, we must have at one texture to resolve
-			return bNeedsMerge;
-		}
-
-	public:
-		// Number of vertices per component
-		FIntPoint ComponentSizeVerts = FIntPoint(ForceInitToZero);
-
-		// Number of sub sections for this landscape 
-		uint32 NumSubsections = 1;
-
-		// Maximum size of all heigthmaps (one heightmap can contain multiple components due to heightmap sharing)
-		FIntPoint MaxHeightmapSize = FIntPoint(ForceInitToZero);
-
-		// Maximum number of mips of all heightmaps (which can be of different sizes) :
-		int32 MaxHeightmapNumMips = 0;
-
-		// Heightmap pixel to world scale factor
-		FVector LandscapeGridScale = FVector::ZeroVector;
-		
-		// Maximum number of visible edit layers that have to be merged for a single FComponentRenderInfo : 
-		int32 MaxNumEditLayersTexturesToMerge = 0;
-
-		// Maximum number of components to render in any given FComponentResolveBatchInfo. This is the number of slices needed for the scratch texture arrays that we reuse from one batch to another :
-		int32 MaxNumComponentsToRenderPerResolveComponentBatch = 0;
-
-		// Describes how to access each visible edit layer's heightmap and how to blend it in the final heightmap for this paint layer :
-		TArray<FEditLayerInfo> VisibleEditLayerInfos;
-
-		// List of infos for each component that needs its edit layers' heightmaps to be rendered (merged) and ultimately participate to the final heightmap of the component we're trying to resolve :
-		TArray<FComponentRenderInfo> ComponentToRenderInfos;
-
-		// List of infos for each component that needs to be resolved :
-		TArray<FComponentResolveInfo> ComponentToResolveInfos;
-
-		// List of batches of FComponentResolveInfo that needs to be resolved in the same pass. This allows massive saves on transient resources on large landscapes because those can be re-cycled from one pass to another :
-		TArray<FComponentResolveBatchInfo> ComponentResolveBatchInfos;
-
-		// List of infos for each texture that needs to be resolved :
-		TArray<FTextureResolveInfo> TextureToResolveInfos;
-
-		// Not truly render-thread data because it references UTextures but it's just because FLandscapeEditLayerReadback were historically game-thread initiated so for as long as we'll use those for readback, we need to store this here : 
-		TArray<FLandscapeLayersCopyReadbackTextureParams> DeferredCopyReadbackTextures;
-	};
-
-	struct FRDGResources
-	{
-		// Texture array in which all possible edit layers heightmaps can fit : we copy the edit layers heightmaps there in order to dynamically access it in the merge shader : 
-		FRDGTextureRef EditLayersHeightmapsTextureArray;
-		FRDGTextureSRVRef EditLayersHeightmapsTextureArraySRV;
-
-		// Temporary scratch texture array that stores the output (packed height only) of all (edit layer-merged) landscape components to render within a batch (one per landscape component)  
-		//  Can be reused from one batch to another :
-		FRDGTextureRef ScratchMergedHeightmapTextureArray;
-		FRDGTextureSRVRef ScratchMergedHeightmapTextureArraySRV;
-
-		// Temporary scratch texture array that stores the output (packed height only) of all (edit layer-merged) landscape components to render within a batch (one per landscape component), after stitching step is done:
-		//  Can be reused from one batch to another :
-		FRDGTextureRef ScratchStitchedHeightmapTextureArray;
-		FRDGTextureSRVRef ScratchStitchedHeightmapTextureArraySRV;
-
-		// Single structured buffer that will contain layer merge infos (FLandscapeEditLayerHeightmapMergeInfo): doesn't changed from one component to another :
-		FRDGBufferRef EditLayersMergeInfosBuffer;
-		FRDGBufferSRVRef EditLayersMergeInfosBufferSRV;
-	};
-
-	void PrepareLayersHeightmapsLocalMergeRDGResources(const FMergeInfo& InLocalMergeInfo, FRDGBuilder& GraphBuilder, FRDGResources& OutResources)
-	{
-		{
-			int32 SizeZ = InLocalMergeInfo.MaxNumEditLayersTexturesToMerge;
-			check(SizeZ > 0);
-			// TODO [jonathan.bard] : change to PF_R8G8 once edit layers heightmaps are stored as such :
-			FRDGTextureDesc Desc = FRDGTextureDesc::Create2DArray(InLocalMergeInfo.ComponentSizeVerts, PF_B8G8R8A8, FClearValueBinding::None, TexCreate_RenderTargetable | TexCreate_ShaderResource, 
-				static_cast<uint16>(SizeZ), /*InNumMips = */1, /*InNumSamples = */1);
-			OutResources.EditLayersHeightmapsTextureArray = GraphBuilder.CreateTexture(Desc, TEXT("LandscapeEditLayersHeightmapsTextureArray"));
-			OutResources.EditLayersHeightmapsTextureArraySRV = GraphBuilder.CreateSRV(FRDGTextureSRVDesc::Create(OutResources.EditLayersHeightmapsTextureArray));
-		}
-
-		{
-			int32 SizeZ = InLocalMergeInfo.MaxNumComponentsToRenderPerResolveComponentBatch;
-			// We only need 2 channels since this only stores the (packed) height :
-			FRDGTextureDesc Desc = FRDGTextureDesc::Create2DArray(InLocalMergeInfo.ComponentSizeVerts, PF_R8G8, FClearValueBinding::None, TexCreate_RenderTargetable | TexCreate_TargetArraySlicesIndependently | TexCreate_ShaderResource, 
-				static_cast<uint16>(SizeZ), /*InNumMips = */1, /*InNumSamples = */1);
-			// Create 2 texture arrays "Merged" and "Stitched" (ScratchMergedHeightmapTextureArray will be copied/merged into ScratchStitchedHeightmapTextureArray, slice by slice)
-			OutResources.ScratchMergedHeightmapTextureArray = GraphBuilder.CreateTexture(Desc, TEXT("LandscapeEditLayersMergedHeightmapTextureArray"));
-			OutResources.ScratchStitchedHeightmapTextureArray = GraphBuilder.CreateTexture(Desc, TEXT("LandscapeEditLayersStitchedHeightmapTextureArray"));
-			OutResources.ScratchMergedHeightmapTextureArraySRV = GraphBuilder.CreateSRV(FRDGTextureSRVDesc::Create(OutResources.ScratchMergedHeightmapTextureArray));
-			OutResources.ScratchStitchedHeightmapTextureArraySRV = GraphBuilder.CreateSRV(FRDGTextureSRVDesc::Create(OutResources.ScratchStitchedHeightmapTextureArray));
-		}
-
-		{
-			// Upload layer merge info buffer once and for all since it's unchanged from one component to another : 
-			TArray<FLandscapeEditLayerHeightmapMergeInfo, TInlineAllocator<16>> EditLayerMergeInfos;
-			for (const FEditLayerInfo& EditLayerInfo : InLocalMergeInfo.VisibleEditLayerInfos)
-			{
-				FLandscapeEditLayerHeightmapMergeInfo& EditLayerMergeInfo = EditLayerMergeInfos.Emplace_GetRef();
-				EditLayerMergeInfo.BlendMode = EditLayerInfo.BlendMode;
-				EditLayerMergeInfo.Alpha = EditLayerInfo.Alpha;
-			}
-			OutResources.EditLayersMergeInfosBuffer = CreateStructuredBuffer(GraphBuilder, TEXT("LandscapeEditLayersMergeInfosBuffer"), EditLayerMergeInfos);
-			OutResources.EditLayersMergeInfosBufferSRV = GraphBuilder.CreateSRV(FRDGBufferSRVDesc(OutResources.EditLayersMergeInfosBuffer));
-		}
-	}
-
-	// Gather all textures we will want to write into or read from in the render graph and output them in OutTrackedTextures:
-	void GatherLayersHeightmapsLocalMergeRDGTextures(const FMergeInfo& InLocalMergeInfo, TMap<FTexture2DResource*, FLandscapeRDGTrackedTexture>& OutTrackedTextures)
-	{
-		// First pass, gather all textures we'll need for merging layers, i.e. the component edit layers' heightmaps : 
-		for (const FComponentRenderInfo& ComponentRenderInfo : InLocalMergeInfo.ComponentToRenderInfos)
-		{
-			for (const FTexture2DResourceSubregion& LayerHeightmap : ComponentRenderInfo.VisibleLayerHeightmapTextures)
-			{
-				check(LayerHeightmap.Texture != nullptr);
-				FLandscapeRDGTrackedTexture* LayerTrackedTexture = OutTrackedTextures.Find(LayerHeightmap.Texture);
-				if (LayerTrackedTexture == nullptr)
-				{
-					LayerTrackedTexture = &(OutTrackedTextures.Add(LayerHeightmap.Texture, FLandscapeRDGTrackedTexture(LayerHeightmap.Texture)));
-				}
-			}
-		}
-
-		// Second pass, gather all textures we'll need to regenerate : 
-		for (const FComponentResolveInfo& ComponentResolveInfo : InLocalMergeInfo.ComponentToResolveInfos)
-		{
-			check(ComponentResolveInfo.Heightmap.Texture != nullptr);
-			FLandscapeRDGTrackedTexture* TrackedTexture = OutTrackedTextures.Find(ComponentResolveInfo.Heightmap.Texture);
-			if (TrackedTexture == nullptr)
-			{
-				TrackedTexture = &(OutTrackedTextures.Add(ComponentResolveInfo.Heightmap.Texture, FLandscapeRDGTrackedTexture(ComponentResolveInfo.Heightmap.Texture)));
-			}
-			TrackedTexture->bNeedsScratch = true;
-		}
-	}
-
-	void MergeEditLayersHeightmapsForBatch(const FComponentResolveBatchInfo& InComponentResolveBatchInfo, const FMergeInfo& InLocalMergeInfo, const TMap<FTexture2DResource*, FLandscapeRDGTrackedTexture>& InTrackedTextures, FRDGBuilder& GraphBuilder, FRDGResources& RDGResources)
-	{
-		RDG_EVENT_SCOPE(GraphBuilder, "Merge %d edit layers on %d components", InLocalMergeInfo.VisibleEditLayerInfos.Num(), InComponentResolveBatchInfo.ComponentToRenderInfoBitIndices.CountSetBits());
-
-		// For each component to render, perform the edit layers merge and write the resulting heightmap :
-		int32 IndexInBatch = 0;
-		for (TConstSetBitIterator<> BitIt(InComponentResolveBatchInfo.ComponentToRenderInfoBitIndices); BitIt; ++BitIt, ++IndexInBatch)
-		{
-			int32 ComponentRenderInfoIndex = BitIt.GetIndex();
-			const FComponentRenderInfo& ComponentRenderInfo = InLocalMergeInfo.ComponentToRenderInfos[ComponentRenderInfoIndex];
-		
-			RDG_EVENT_SCOPE(GraphBuilder, "Component %s", *ComponentRenderInfo.Name);
-
-			// Prepare the texture array of layer heightmaps for this component : 
-			int32 NumLayers = InLocalMergeInfo.VisibleEditLayerInfos.Num();
-			for (int32 i = 0; i < NumLayers; ++i)
-			{
-				const FTexture2DResourceSubregion& LayerHeighmapSubregion = ComponentRenderInfo.VisibleLayerHeightmapTextures[i];
-
-				const FLandscapeRDGTrackedTexture* TrackedTexture = InTrackedTextures.Find(LayerHeighmapSubregion.Texture);
-				check(TrackedTexture != nullptr);
-
-				// We need to copy the (portion of the) layer's texture to the texture array : 
-				FRHICopyTextureInfo CopyTextureInfo;
-				CopyTextureInfo.Size = FIntVector(LayerHeighmapSubregion.Subregion.Size().X, LayerHeighmapSubregion.Subregion.Size().Y, 0);
-				CopyTextureInfo.DestSliceIndex = i;
-				CopyTextureInfo.SourcePosition = FIntVector(LayerHeighmapSubregion.Subregion.Min.X, LayerHeighmapSubregion.Subregion.Min.Y, 0);
-
-				AddCopyTexturePass(GraphBuilder, TrackedTexture->ExternalTextureRef, RDGResources.EditLayersHeightmapsTextureArray, CopyTextureInfo);
-			}
-
-			// Then, merge all heightmaps using the MergeEditLayers PS and write into a single slice in ScratchMergedHeightmapTextureArray: 
-			{
-				check(IndexInBatch < RDGResources.ScratchMergedHeightmapTextureArray->Desc.ArraySize);
-
-				FLandscapeLayersHeightmapsMergeEditLayersPS::FParameters* MergeEditLayersPSParams = GraphBuilder.AllocParameters<FLandscapeLayersHeightmapsMergeEditLayersPS::FParameters>();
-				MergeEditLayersPSParams->RenderTargets[0] = FRenderTargetBinding(RDGResources.ScratchMergedHeightmapTextureArray, ERenderTargetLoadAction::ENoAction, /*InMipIndex = */0, /*InArraySlice = */static_cast<int16>(IndexInBatch));
-				MergeEditLayersPSParams->InNumEditLayers = NumLayers;
-				MergeEditLayersPSParams->InEditLayersTextures = RDGResources.EditLayersHeightmapsTextureArraySRV;
-				MergeEditLayersPSParams->InEditLayersMergeInfos = RDGResources.EditLayersMergeInfosBufferSRV;
-
-				FLandscapeLayersHeightmapsMergeEditLayersPS::MergeEditLayers(GraphBuilder, MergeEditLayersPSParams, /*InTextureSize =*/InLocalMergeInfo.ComponentSizeVerts);
-			}
-		}
-	}
-
-	void StitchHeightmapsForBatch(const FComponentResolveBatchInfo& InComponentResolveBatchInfo, const FMergeInfo& InLocalMergeInfo, FRDGBuilder& GraphBuilder, FRDGResources& RDGResources)
-	{
-		RDG_EVENT_SCOPE(GraphBuilder, "Stitch %d components for batch %d", InComponentResolveBatchInfo.ComponentToRenderInfoBitIndices.CountSetBits(), InComponentResolveBatchInfo.BatchIndex);
-
-		// For each component to render (i.e. including the neighbors of dirty components, that are also need to properly recompute normals), correct the heightmap 
-		//  subregion (by stitching adjacent component borders) :
-		int32 IndexInBatch = 0;
-		for (TConstSetBitIterator<> BitIt(InComponentResolveBatchInfo.ComponentToRenderInfoBitIndices); BitIt; ++BitIt, ++IndexInBatch)
-		{
-			int32 ComponentRenderInfoIndex = BitIt.GetIndex();
-			const FComponentRenderInfo& ComponentRenderInfo = InLocalMergeInfo.ComponentToRenderInfos[ComponentRenderInfoIndex];
-
-			RDG_EVENT_SCOPE(GraphBuilder, "Component %s", *ComponentRenderInfo.Name);
-
-			// Now, stitch the heightmap using the StitchHeightmapPS and output the result to a single slice in ScratchStitchedHeightmapTextureArray: 
-			{
-				FLandscapeLayersHeightmapsStitchHeightmapPS::FParameters* StitchHeightmapPSParams = GraphBuilder.AllocParameters<FLandscapeLayersHeightmapsStitchHeightmapPS::FParameters>();
-				StitchHeightmapPSParams->RenderTargets[0] = FRenderTargetBinding(RDGResources.ScratchStitchedHeightmapTextureArray, ERenderTargetLoadAction::ENoAction, /*InMipIndex = */0, /*InArraySlice = */static_cast<int16>(IndexInBatch));
-				StitchHeightmapPSParams->InSourceTextureSize = FUintVector2(InLocalMergeInfo.ComponentSizeVerts.X, InLocalMergeInfo.ComponentSizeVerts.Y);
-				StitchHeightmapPSParams->InNumSubsections = InLocalMergeInfo.NumSubsections;
-				StitchHeightmapPSParams->InSourceHeightmaps = RDGResources.ScratchMergedHeightmapTextureArraySRV;
-
-				for (int32 NeighborIndex = 0; NeighborIndex < 9; ++NeighborIndex)
-				{
-					int32 NeighborComponentToRenderInfoIndex = ComponentRenderInfo.NeighborComponentToRenderInfoIndices[NeighborIndex];
-					// Index of the neighbor component in this batch : allows to retrieve the proper slice in the source heightmap array :
-					int32 NeighborIndexInBatch = INDEX_NONE;
-					// The neighbor could be totally absent (index == INDEX_NONE) or it can be absent from the batch. That means that it's not actually relevant for this step 
-					//  since the vertices that will be "invalid" won't be taken into account by the components we're actually trying to resolve :
-					if ((NeighborComponentToRenderInfoIndex != INDEX_NONE) && InComponentResolveBatchInfo.ComponentToRenderInfoBitIndices[NeighborComponentToRenderInfoIndex])
-					{
-						// Components are rendered in order within the batch so the index of this component to render in the batch is == to how many components are there before it:
-						NeighborIndexInBatch = InComponentResolveBatchInfo.ComponentToRenderInfoBitIndices.CountSetBits(0, NeighborComponentToRenderInfoIndex);
-						check(NeighborIndexInBatch < RDGResources.ScratchMergedHeightmapTextureArray->Desc.ArraySize);
-					}
-
-					check((NeighborIndex != 4) || (NeighborIndexInBatch != INDEX_NONE)); // The central component (the one we finalize) should always be valid
-					GET_SCALAR_ARRAY_ELEMENT(StitchHeightmapPSParams->InNeighborHeightmapIndices, NeighborIndex) = NeighborIndexInBatch;
-				}
-
-				FLandscapeLayersHeightmapsStitchHeightmapPS::StitchHeightmap(GraphBuilder, StitchHeightmapPSParams);
-			}
-		}
-	}
-
-	void FinalizeComponentsForBatch(const FComponentResolveBatchInfo& InComponentResolveBatchInfo, const FMergeInfo& InLocalMergeInfo, const TMap<FTexture2DResource*, FLandscapeRDGTrackedTexture>& InTrackedTextures, FRDGBuilder& GraphBuilder, FRDGResources& RDGResources)
-	{
-		RDG_EVENT_SCOPE(GraphBuilder, "Finalize %d components", InLocalMergeInfo.ComponentToResolveInfos.Num());
-
-		// For each component to resolve (i.e. including the neighbors of dirty components, that also need to have their normals recomputed), finalize the heightmap subregion (i.e. compute the normals) :
-		for (int32 ComponentResolveInfoIndex : InComponentResolveBatchInfo.ComponentToResolveInfoIndices)
-		{
-			const FComponentResolveInfo& ComponentResolveInfo = InLocalMergeInfo.ComponentToResolveInfos[ComponentResolveInfoIndex];
-
-			RDG_EVENT_SCOPE(GraphBuilder, "Component %s", *ComponentResolveInfo.Name);
-
-			const FLandscapeRDGTrackedTexture* TrackedTexture = InTrackedTextures.Find(ComponentResolveInfo.Heightmap.Texture);
-			check((TrackedTexture != nullptr) && (TrackedTexture->ScratchTextureRef != nullptr));
-
-			// Now, finalize the heightmap using the FinalizeHeighmap PS : 
-			{
-				FLandscapeLayersHeightmapsFinalizeHeightmapPS::FParameters* FinalizeHeightmapPSParams = GraphBuilder.AllocParameters<FLandscapeLayersHeightmapsFinalizeHeightmapPS::FParameters>();
-				FinalizeHeightmapPSParams->RenderTargets[0] = FRenderTargetBinding(TrackedTexture->ScratchTextureRef, ERenderTargetLoadAction::ENoAction);
-				FinalizeHeightmapPSParams->InSourceTextureSize = FUintVector2(InLocalMergeInfo.ComponentSizeVerts.X, InLocalMergeInfo.ComponentSizeVerts.Y);
-				FinalizeHeightmapPSParams->InNumSubsections = InLocalMergeInfo.NumSubsections;
-				FinalizeHeightmapPSParams->InSourceHeightmaps = RDGResources.ScratchStitchedHeightmapTextureArraySRV;
-				FinalizeHeightmapPSParams->InDestinationTextureSubregion = FUintVector4(ComponentResolveInfo.Heightmap.Subregion.Min.X, ComponentResolveInfo.Heightmap.Subregion.Min.Y, ComponentResolveInfo.Heightmap.Subregion.Max.X, ComponentResolveInfo.Heightmap.Subregion.Max.Y);
-				FinalizeHeightmapPSParams->InLandscapeGridScale = (FVector3f)InLocalMergeInfo.LandscapeGridScale;
-
-				for (int32 NeighborIndex = 0; NeighborIndex < 9; ++NeighborIndex)
-				{
-					int32 NeighborComponentToRenderInfoIndex = ComponentResolveInfo.NeighborComponentToRenderInfoIndices[NeighborIndex];
-					// Index of the neighbor component in this batch : allows to retrieve the proper slice in the source heightmap array :
-					int32 NeighborIndexInBatch = INDEX_NONE;
-					// The neighbor could be absent : 
-					if (NeighborComponentToRenderInfoIndex != INDEX_NONE)
-					{
-						check(InComponentResolveBatchInfo.ComponentToRenderInfoBitIndices[NeighborComponentToRenderInfoIndex]);
-						// Components are rendered in order within the batch so the index of this component to render in the batch is == to how many components are there before it:
-						NeighborIndexInBatch = InComponentResolveBatchInfo.ComponentToRenderInfoBitIndices.CountSetBits(0, NeighborComponentToRenderInfoIndex);
-						check(NeighborIndexInBatch < RDGResources.ScratchMergedHeightmapTextureArray->Desc.ArraySize);
-					}
-
-					check((NeighborIndex != 4) || (NeighborIndexInBatch != INDEX_NONE)); // The central component (the one we finalize) should always be valid
-					GET_SCALAR_ARRAY_ELEMENT(FinalizeHeightmapPSParams->InNeighborHeightmapIndices, NeighborIndex) = NeighborIndexInBatch;
-				}
-
-				FLandscapeLayersHeightmapsFinalizeHeightmapPS::FinalizeHeightmap(GraphBuilder, FinalizeHeightmapPSParams);
-			}
-		}
-	}
-
-	void GenerateHeightmapMips(const FMergeInfo& InLocalMergeInfo, const TMap<FTexture2DResource*, FLandscapeRDGTrackedTexture>& InTrackedTextures, FRDGBuilder& GraphBuilder)
-	{
-		RDG_EVENT_SCOPE(GraphBuilder, "Generate mips for %d heightmaps", InLocalMergeInfo.TextureToResolveInfos.Num());
-
-		// For each texture to resolve, simply regenerate the mips on the entire texture :
-		for (const FTextureResolveInfo& TextureResolveInfo : InLocalMergeInfo.TextureToResolveInfos)
-		{
-			RDG_EVENT_SCOPE(GraphBuilder, "Texture %s", *TextureResolveInfo.Texture->GetTextureName().ToString());
-
-			const FLandscapeRDGTrackedTexture* TrackedTexture = InTrackedTextures.Find(TextureResolveInfo.Texture);
-			check((TrackedTexture != nullptr) && (TrackedTexture->ScratchTextureRef != nullptr));
-
-			//check(TextureResolveInfo.NumMips == TrackedTexture->ScratchTextureRef->Desc.NumMips);
-			//check(TrackedTexture->ScratchTextureMipsSRVRefs.Num() == TextureResolveInfo.NumMips);
-
-			FIntPoint CurrentMipSubregionSize = InLocalMergeInfo.ComponentSizeVerts;
-			for (int32 MipLevel = 1; MipLevel < TextureResolveInfo.NumMips; ++MipLevel)
-			{
-				CurrentMipSubregionSize.X >>= 1;
-				CurrentMipSubregionSize.Y >>= 1;
-
-				FLandscapeLayersHeightmapsGenerateMipsPS::FParameters* GenerateMipsPSParams = GraphBuilder.AllocParameters<FLandscapeLayersHeightmapsGenerateMipsPS::FParameters>();
-				GenerateMipsPSParams->RenderTargets[0] = FRenderTargetBinding(TrackedTexture->ScratchTextureRef, ERenderTargetLoadAction::ENoAction, static_cast<uint8>(MipLevel));
-				GenerateMipsPSParams->InCurrentMipSubregionSize = FUintVector2(CurrentMipSubregionSize.X, CurrentMipSubregionSize.Y);
-				GenerateMipsPSParams->InNumSubsections = InLocalMergeInfo.NumSubsections;
-				GenerateMipsPSParams->InSourceHeightmap = TrackedTexture->ScratchTextureMipsSRVRefs[MipLevel - 1];
-
-				FLandscapeLayersHeightmapsGenerateMipsPS::GenerateSingleMip(GraphBuilder, GenerateMipsPSParams);
-			}
-		}
-	}
-
-	void CopyScratchToSourceHeightmaps(const FMergeInfo& InLocalMergeInfo, const TMap<FTexture2DResource*, FLandscapeRDGTrackedTexture>& InTrackedTextures, FRDGBuilder& GraphBuilder)
-	{
-		RDG_EVENT_SCOPE(GraphBuilder, "Copy %d scratch to source heightmaps", InLocalMergeInfo.ComponentToResolveInfos.Num());
-
-		// For each texture to resolve, copy from scratch to final texture :
-		for (const FTextureResolveInfo& TextureResolveInfo : InLocalMergeInfo.TextureToResolveInfos)
-		{
-			const FLandscapeRDGTrackedTexture* TrackedTexture = InTrackedTextures.Find(TextureResolveInfo.Texture);
-			check((TrackedTexture != nullptr) && (TrackedTexture->ScratchTextureRef != nullptr) && (TrackedTexture->ExternalTextureRef != nullptr));
-
-			FRHICopyTextureInfo CopyTextureInfo;
-			// We want to copy all mips : 
-			CopyTextureInfo.NumMips = TextureResolveInfo.NumMips;
-
-			AddCopyTexturePass(GraphBuilder, TrackedTexture->ScratchTextureRef, TrackedTexture->ExternalTextureRef, CopyTextureInfo);
-		}
-	}
-
-	FRDGTextureRef CreateAndClearEmptyTexture(const FMergeInfo& InLocalMergeInfo, FRDGBuilder& GraphBuilder)
-	{
-		// Convert the height value 0.0f to how it's stored in the texture : 
-		const uint16 HeightValue = LandscapeDataAccess::GetTexHeight(0.0f);
-		FLinearColor ClearHeightColor((float)((HeightValue - (HeightValue & 255)) >> 8) / 255.0f, (float)(HeightValue & 255) / 255.0f, 0.0f, 0.0f);
-
-		// Even if we have heightmaps of different sizes to handle, we only need one empty heightmap to copy from (whose size is MaxHeightmapSize) :
-		// TODO [jonathan.bard] : change to PF_R8G8 once edit layers heightmaps are stored as such :
-		FRDGTextureDesc Desc = FRDGTextureDesc::Create2D(InLocalMergeInfo.MaxHeightmapSize, PF_B8G8R8A8, FClearValueBinding(ClearHeightColor), TexCreate_ShaderResource | TexCreate_RenderTargetable, static_cast<uint8>(InLocalMergeInfo.MaxHeightmapNumMips), /*InNumSamples = */1);
-		FRDGTextureRef EmptyTexture = GraphBuilder.CreateTexture(Desc, TEXT("LandscapeEditLayersEmptyHeightmap"));
-
-		FRDGTextureClearInfo ClearInfo;
-		ClearInfo.NumMips = InLocalMergeInfo.MaxHeightmapNumMips;
-		AddClearRenderTargetPass(GraphBuilder, EmptyTexture, ClearInfo);
-
-		return EmptyTexture;
-	}
-
-	void ClearSourceHeightmaps(const FMergeInfo& InLocalMergeInfo, FRDGBuilder& GraphBuilder)
-	{
-		RDG_EVENT_SCOPE(GraphBuilder, "Clear %d source heightmaps", InLocalMergeInfo.ComponentToResolveInfos.Num());
-
-		// We cannot clear heightmaps directly, since they are external, non-render targetable, textures so we need to copy from an empty source texture :
-		FRDGTextureRef EmptyTexture = CreateAndClearEmptyTexture(InLocalMergeInfo, GraphBuilder);
-
-		// For each texture to resolve, copy from empty to final texture :
-		for (const FTextureResolveInfo& TextureResolveInfo : InLocalMergeInfo.TextureToResolveInfos)
-		{
-			// Register the output texture to the GraphBuilder so that we can copy to it: 
-			FString* DebugName = GraphBuilder.AllocObject<FString>(TextureResolveInfo.Texture->GetTextureName().ToString());
-			TRefCountPtr<IPooledRenderTarget> RenderTarget = CreateRenderTarget(TextureResolveInfo.Texture->TextureRHI, **DebugName);
-
-			// Force tracking on the external texture, so that it can be copied to via CopyTexture within the graph : 
-			FRDGTextureRef DestinationTexture = GraphBuilder.RegisterExternalTexture(RenderTarget);
-
-			FRHICopyTextureInfo CopyTextureInfo;
-			// We want to copy all mips : 
-			CopyTextureInfo.NumMips = TextureResolveInfo.NumMips;
-			// We need specify the size since the empty texture might be of higher size :
-			check(EmptyTexture->Desc.GetSize().X >= TextureResolveInfo.TextureSize.X);
-			CopyTextureInfo.Size = FIntVector(TextureResolveInfo.TextureSize.X, TextureResolveInfo.TextureSize.X, 0);
-
-			AddCopyTexturePass(GraphBuilder, EmptyTexture, DestinationTexture, CopyTextureInfo);
-		}
-	}
-}; // End namespace EditLayersHeightmapLocalMerge_RenderThread
-
-void ALandscape::PrepareLayersHeightmapsLocalMergeRenderThreadData(const FUpdateLayersContentContext& InUpdateLayersContentContext, const FEditLayersHeightmapMergeParams& InMergeParams, EditLayersHeightmapLocalMerge_RenderThread::FMergeInfo& OutRenderThreadData)
-{
-	using namespace EditLayersHeightmapLocalMerge_RenderThread;
-	using namespace UE::Landscape::Private;
-
-	TRACE_CPUPROFILER_EVENT_SCOPE(LandscapeLayers_PrepareLayersHeightmapsLocalMergeRenderThreadData);
-
-	ULandscapeInfo* Info = GetLandscapeInfo();
-	check(Info != nullptr);
-
-	// Number of vertices for each landscape component :
-	int32 ComponentSizeVerts = (SubsectionSizeQuads + 1) * NumSubsections;
-	OutRenderThreadData.ComponentSizeVerts = FIntPoint(ComponentSizeVerts, ComponentSizeVerts);
-	OutRenderThreadData.LandscapeGridScale = GetRootComponent()->GetRelativeScale3D();
-	OutRenderThreadData.NumSubsections = NumSubsections;
-
-	// Prepare landscape edit layers data common to all landscape components: 
-	OutRenderThreadData.VisibleEditLayerInfos.Reserve(LandscapeEditLayers.Num());
-	for (const FLandscapeLayer& Layer : LandscapeEditLayers)
-	{
-		check(Layer.EditLayer != nullptr);
-		if (Layer.EditLayer->IsVisible() && !InMergeParams.bSkipBrush)
-		{
-			FEditLayerInfo& EditLayerInfo = OutRenderThreadData.VisibleEditLayerInfos.Emplace_GetRef();
-			EditLayerInfo.BlendMode = LandscapeBlendModeToEditLayerBlendMode(Layer.EditLayer->GetBlendMode());
-			EditLayerInfo.Alpha = Layer.EditLayer->GetAlphaForTargetType(ELandscapeToolTargetType::Heightmap);
-		}
-	}
-
-	const int32 NumComponentsToRender = InUpdateLayersContentContext.LandscapeComponentsHeightmapsToRender.Num();
-
-	// Lookup table to retrieve, for a given rendered component, the index of its FComponentRenderInfo in ComponentToRenderInfos :
-	TMap<ULandscapeComponent*, int32> ComponentToComponentRenderInfoIndex;
-
-	// Prepare all per-landscape component render data : 
-	{
-		TRACE_CPUPROFILER_EVENT_SCOPE(LandscapeLayers_PrepareHeightmapComponentRenderInfos);
-
-		OutRenderThreadData.ComponentToRenderInfos.Reserve(InUpdateLayersContentContext.LandscapeComponentsHeightmapsToRender.Num());
-		for (ULandscapeComponent* Component : InUpdateLayersContentContext.LandscapeComponentsHeightmapsToRender)
-		{
-			// Add a new component render info and set it up :
-			FComponentRenderInfo& NewComponentRenderInfo = OutRenderThreadData.ComponentToRenderInfos.Emplace_GetRef();
-			NewComponentRenderInfo.Name = Component->GetName();
-
-			// Associate the component with its index in ComponentToRenderInfos :
-			check(!ComponentToComponentRenderInfoIndex.Contains(Component));
-			ComponentToComponentRenderInfoIndex.Add(Component, OutRenderThreadData.ComponentToRenderInfos.Num() - 1);
-
-			UTexture2D* ComponentHeightmap = Component->GetHeightmap();
-			FIntPoint TextureSize(ComponentHeightmap->Source.GetSizeX(), ComponentHeightmap->Source.GetSizeY());
-			OutRenderThreadData.MaxHeightmapSize = TextureSize.ComponentMax(OutRenderThreadData.MaxHeightmapSize);
-			OutRenderThreadData.MaxHeightmapNumMips = FMath::Max((int32)FMath::CeilLogTwo(TextureSize.GetMin()) + 1, OutRenderThreadData.MaxHeightmapNumMips);
-
-			FIntPoint HeightmapOffset(static_cast<int32>(Component->HeightmapScaleBias.Z * TextureSize.X), static_cast<int32>(Component->HeightmapScaleBias.W * TextureSize.Y));
-			// Effective area of the texture affecting this component (because of texture sharing) :
-			const FIntRect ComponentTextureSubregion(HeightmapOffset, HeightmapOffset + OutRenderThreadData.ComponentSizeVerts);
-
-			NewComponentRenderInfo.VisibleLayerHeightmapTextures.Reserve(OutRenderThreadData.VisibleEditLayerInfos.Num());
-			for (const FLandscapeLayer& Layer : LandscapeEditLayers)
-			{
-				if (Layer.EditLayer->IsVisible() && !InMergeParams.bSkipBrush)
-				{
-					if (FLandscapeLayerComponentData* ComponentLayerData = Component->GetLayerData(Layer.EditLayer->GetGuid()))
-					{
-						if (UTexture2D* LayerHeightmap = ComponentLayerData->HeightmapData.Texture.Get())
-						{
-							NewComponentRenderInfo.VisibleLayerHeightmapTextures.Add(FTexture2DResourceSubregion(LayerHeightmap->GetResource()->GetTexture2DResource(), ComponentTextureSubregion));
-						}
-					}
-				}
-			}
-		}
-
-		// Now that all landscape components have been registered, identify the (valid) neighbors for each : 
-		for (int32 ComponentToRenderIndex = 0; ComponentToRenderIndex < NumComponentsToRender; ++ComponentToRenderIndex)
-		{
-			ULandscapeComponent* Component = InUpdateLayersContentContext.LandscapeComponentsHeightmapsToRender[ComponentToRenderIndex];
-			FComponentRenderInfo& ComponentRenderInfo = OutRenderThreadData.ComponentToRenderInfos[ComponentToRenderIndex];
-
-			// Gather neighboring component infos : 
-			TStaticArray<ULandscapeComponent*, 9> NeighborComponents;
-			Component->GetLandscapeComponentNeighbors3x3(NeighborComponents);
-			for (int32 NeighborIndex = 0; NeighborIndex < 9; ++NeighborIndex)
-			{
-				int32 NeighborComponentRenderInfoIndex = INDEX_NONE;
-				if (ULandscapeComponent* NeighborComponent = NeighborComponents[NeighborIndex])
-				{
-					if (int32* NeighborComponentRenderInfoIndexPtr = ComponentToComponentRenderInfoIndex.Find(NeighborComponent))
-					{
-						NeighborComponentRenderInfoIndex = *NeighborComponentRenderInfoIndexPtr;
-					}
-				}
-				ComponentRenderInfo.NeighborComponentToRenderInfoIndices[NeighborIndex] = NeighborComponentRenderInfoIndex;
-			}
-		}
-	}
-
-	// List of UTexture2D that we need to kick off readbacks for :
-	TArray<UTexture2D*> TexturesNeedingReadback;
-
-	// Prepare per-landscape component resolve data : 
-	{
-		TRACE_CPUPROFILER_EVENT_SCOPE(LandscapeLayers_PrepareHeightmapComponentResolveInfos);
-		OutRenderThreadData.ComponentToResolveInfos.Reserve(InUpdateLayersContentContext.LandscapeComponentsHeightmapsToResolve.Num());
-		for (ULandscapeComponent* Component : InUpdateLayersContentContext.LandscapeComponentsHeightmapsToResolve)
-		{
-			UTexture2D* ComponentHeightmap = Component->GetHeightmap();
-			const int32 HeightmapOffsetX = static_cast<int32>(Component->HeightmapScaleBias.Z * ComponentHeightmap->Source.GetSizeX());
-			const int32 HeightmapOffsetY = static_cast<int32>(Component->HeightmapScaleBias.W * ComponentHeightmap->Source.GetSizeY());
-			// Effective area of the texture affecting this component (because of texture sharing) :
-			const FIntRect ComponentTextureSubregion(FIntPoint(HeightmapOffsetX, HeightmapOffsetY), FIntPoint(HeightmapOffsetX, HeightmapOffsetY) + OutRenderThreadData.ComponentSizeVerts);
-
-			FComponentResolveInfo& NewComponentResolveInfo = OutRenderThreadData.ComponentToResolveInfos.Add_GetRef(FComponentResolveInfo(NumComponentsToRender));
-			NewComponentResolveInfo.ComponentToResolveInfoIndex = OutRenderThreadData.ComponentToResolveInfos.Num() - 1;
-			NewComponentResolveInfo.Name = Component->GetName();
-			NewComponentResolveInfo.Heightmap = FTexture2DResourceSubregion(ComponentHeightmap->GetResource()->GetTexture2DResource(), ComponentTextureSubregion);
-
-			// Gather neighboring component infos : 
-			TStaticArray<ULandscapeComponent*, 9> NeighborComponents;
-			Component->GetLandscapeComponentNeighbors3x3(NeighborComponents);
-			for (int32 NeighborIndex = 0; NeighborIndex < 9; ++NeighborIndex)
-			{
-				int32 NeighborComponentRenderInfoIndex = INDEX_NONE;
-				if (ULandscapeComponent* NeighborComponent = NeighborComponents[NeighborIndex])
-				{
-					if (int32* NeighborComponentRenderInfoIndexPtr = ComponentToComponentRenderInfoIndex.Find(NeighborComponent))
-					{
-						NeighborComponentRenderInfoIndex = *NeighborComponentRenderInfoIndexPtr;
-					}
-				}
-				NewComponentResolveInfo.SetNeighborRenderInfo(NeighborIndex, NeighborComponentRenderInfoIndex);
-			}
-
-			FTextureResolveInfo* TextureResolveInfo = OutRenderThreadData.TextureToResolveInfos.FindByPredicate(
-				[TextureToResolve = NewComponentResolveInfo.Heightmap.Texture](const FTextureResolveInfo& TextureResolveInfo) { return TextureResolveInfo.Texture == TextureToResolve; });
-			if (TextureResolveInfo == nullptr)
-			{
-				ALandscapeProxy* Proxy = Component->GetLandscapeProxy();
-				FLandscapeEditLayerReadback** CPUReadback = Proxy->HeightmapsCPUReadback.Find(ComponentHeightmap);
-				check((CPUReadback != nullptr) && (*CPUReadback != nullptr));
-
-				FTextureResolveInfo NewTextureResolveInfo;
-				NewTextureResolveInfo.TextureSize = FIntPoint(ComponentHeightmap->Source.GetSizeX(), ComponentHeightmap->Source.GetSizeY());
-				NewTextureResolveInfo.NumMips = (int32)FMath::CeilLogTwo(NewTextureResolveInfo.TextureSize.GetMin()) + 1;
-				NewTextureResolveInfo.Texture = NewComponentResolveInfo.Heightmap.Texture;
-				NewTextureResolveInfo.CPUReadback = *CPUReadback;
-				OutRenderThreadData.TextureToResolveInfos.Add(NewTextureResolveInfo);
-
-				check(!TexturesNeedingReadback.Contains(ComponentHeightmap));
-				TexturesNeedingReadback.Add(ComponentHeightmap);
-			}
-		}
-	}
-
-	// Prepare the texture resolve batches: 
-	{
-		TRACE_CPUPROFILER_EVENT_SCOPE(LandscapeLayers_PrepareHeightmapComponentResolveInfoBatches);
-
-		int32 MaxComponentsPerResolveBatch = CVarLandscapeEditLayersMaxComponentsPerHeightmapResolveBatch.GetValueOnGameThread();
-
-		// Copy the component infos because TextureToResolveInfos indices need to remain stable at this point :
-		TArray<FComponentResolveInfo> RemainingComponentToResolveInfos = OutRenderThreadData.ComponentToResolveInfos;
-		if (!RemainingComponentToResolveInfos.IsEmpty())
-		{
-			TBitArray<TInlineAllocator<1>> TempBitArray;
-			TempBitArray.Reserve(OutRenderThreadData.ComponentToResolveInfos.Num());
-
-			while (!RemainingComponentToResolveInfos.IsEmpty())
-			{
-				const FComponentResolveInfo& ComponentResolveInfo = RemainingComponentToResolveInfos.Pop(EAllowShrinking::No);
-
-				int32 BestBatchIndex = INDEX_NONE;
-				int32 MinNumComponents = MAX_int32;
-
-				// Iterate through all all batches and try to find which would be able to accept it and amongst those, which it would share the most components to render with:
-				int32 NumBatches = OutRenderThreadData.ComponentResolveBatchInfos.Num();
-				for (int32 BatchIndex = 0; BatchIndex < NumBatches; ++BatchIndex)
-				{
-					const FComponentResolveBatchInfo& Batch = OutRenderThreadData.ComponentResolveBatchInfos[BatchIndex];
-					TempBitArray = TBitArray<>::BitwiseOR(Batch.ComponentToRenderInfoBitIndices, ComponentResolveInfo.NeighborComponentToRenderInfoBitIndices, EBitwiseOperatorFlags::MinSize);
-
-					// If after adding its components, the batch still has less than MaxComponentsPerResolveBatch, it can accept it
-					int32 NumComponentsAfter = TempBitArray.CountSetBits();
-					if (NumComponentsAfter <= MaxComponentsPerResolveBatch)
-					{
-						// Is this the best candidate so far?
-						if (NumComponentsAfter < MinNumComponents)
-						{
-							BestBatchIndex = BatchIndex;
-							MinNumComponents = NumComponentsAfter;
-						}
-
-						// If the number of components after addition of this texture is unchanged, it's a perfect match, we won't ever find a better batch so just stop there for this texture:
-						if (NumComponentsAfter == Batch.ComponentToRenderInfoBitIndices.CountSetBits())
-						{
-							break;
-						}
-					}
-				}
-
-				// If we have found a batch, just add the texture to it, otherwise, add a new batch:
-				FComponentResolveBatchInfo& SelectedBatch = (BestBatchIndex != INDEX_NONE) ? OutRenderThreadData.ComponentResolveBatchInfos[BestBatchIndex]
-					: OutRenderThreadData.ComponentResolveBatchInfos.Add_GetRef(FComponentResolveBatchInfo(OutRenderThreadData.ComponentToRenderInfos.Num(), /*InBatchIndex = */OutRenderThreadData.ComponentResolveBatchInfos.Num()));
-
-				SelectedBatch.AddComponent(ComponentResolveInfo);
-				check(SelectedBatch.ComponentToRenderInfoBitIndices.CountSetBits() <= MaxComponentsPerResolveBatch);
-
-				// Keep track of the maximum number of slices in the scratch texture arrays we'll need for any given batch :
-				OutRenderThreadData.MaxNumComponentsToRenderPerResolveComponentBatch = FMath::Max(SelectedBatch.ComponentToRenderInfoBitIndices.CountSetBits(), OutRenderThreadData.MaxNumComponentsToRenderPerResolveComponentBatch);
-			}
-		}
-	}
-
-	// Finalize :
-	{
-		// Prepare the UTexture2D readbacks we'll need to perform :
-		OutRenderThreadData.DeferredCopyReadbackTextures = PrepareLandscapeLayersCopyReadbackTextureParams(InUpdateLayersContentContext.MapHelper, TexturesNeedingReadback, /*bWeightmaps = */false);
-
-		// We'll only ever need this amount of edit layers textures for any MergeEditLayers operation :
-		OutRenderThreadData.MaxNumEditLayersTexturesToMerge = OutRenderThreadData.VisibleEditLayerInfos.Num();
-	}
-}
-
-int32 ALandscape::PerformLayersHeightmapsLocalMerge(const FUpdateLayersContentContext& InUpdateLayersContentContext, const FEditLayersHeightmapMergeParams& InMergeParams)
-{
-	using namespace EditLayersHeightmapLocalMerge_RenderThread;
-
-	TRACE_CPUPROFILER_EVENT_SCOPE(LandscapeLayers_PerformLayersWeightmapsLocalMerge);
-
-	FMergeInfo RenderThreadData;
-	PrepareLayersHeightmapsLocalMergeRenderThreadData(InUpdateLayersContentContext, InMergeParams, RenderThreadData);
-	
-	ENQUEUE_RENDER_COMMAND(PerformLayersHeightmapsLocalMerge)([RenderThreadData](FRHICommandListImmediate& RHICmdList)
-	{
-		FRDGBuilder GraphBuilder(RHICmdList, RDG_EVENT_NAME("PerformLayersHeightmapsLocalMerge"));
-
-		if (RenderThreadData.NeedsMerge())
-		{
-			// Prepare the GPU resources we will use during the local merge : 
-			FRDGResources RDGResources;
-			PrepareLayersHeightmapsLocalMergeRDGResources(RenderThreadData, GraphBuilder, RDGResources);
-
-			// Get a list of all external textures (heightmaps) we will manipulate during the local merge : 
-			TMap<FTexture2DResource*, FLandscapeRDGTrackedTexture> TrackedTextures;
-			GatherLayersHeightmapsLocalMergeRDGTextures(RenderThreadData, TrackedTextures);
-
-			// Start tracking those in the render graph :
-			TrackLandscapeRDGTextures(GraphBuilder, TrackedTextures);
-
-			// Process the components batch by batch in order to avoid over-allocating temporary textures :
-			for (const FComponentResolveBatchInfo& ComponentResolveBatchInfo : RenderThreadData.ComponentResolveBatchInfos)
-			{
-				RDG_EVENT_SCOPE(GraphBuilder, "Process batch %d", ComponentResolveBatchInfo.BatchIndex);
-
-				// Perform all edit layers merges, for all components to render in that batch :
-				MergeEditLayersHeightmapsForBatch(ComponentResolveBatchInfo, RenderThreadData, TrackedTextures, GraphBuilder, RDGResources);
-
-				// Correct borders of all rendered components so that they're all stitched together :
-				StitchHeightmapsForBatch(ComponentResolveBatchInfo, RenderThreadData, GraphBuilder, RDGResources);
-
-				// Finalize (compute normals) of each component to resolve :
-				FinalizeComponentsForBatch(ComponentResolveBatchInfo, RenderThreadData, TrackedTextures, GraphBuilder, RDGResources);
-			}
-
-			// Generate the mips on the entire heightmaps : 
-			GenerateHeightmapMips(RenderThreadData, TrackedTextures, GraphBuilder);
-
-			// Finally, we can put those scratch textures to good usage and update our actual heightmaps : 
-			CopyScratchToSourceHeightmaps(RenderThreadData, TrackedTextures, GraphBuilder);
-		}
-		else
-		{
-			// When there's nothing to do, we still have the obligation to output empty heightmaps :
-			ClearSourceHeightmaps(RenderThreadData, GraphBuilder);
-		}
-
-		GraphBuilder.Execute();
-
-	});
-
-	ExecuteCopyToReadbackTexture(RenderThreadData.DeferredCopyReadbackTextures);
-
-	return InMergeParams.HeightmapUpdateModes;
-}
-
 TArray<UE::Landscape::EditLayers::FEditLayerRendererState> ALandscape::GetEditLayerRendererStates(const UE::Landscape::EditLayers::FMergeContext* InMergeContext)
 {
 	return GetEditLayerRendererStatesEnableOverride(InMergeContext, {});
@@ -5571,7 +2603,7 @@
 						State.EnableTargetTypeMask(TargetTypes);
 					}
 					else
-					{
+		{
 						State.DisableTargetTypeMask(TargetTypes);
 					}
 				});
@@ -5611,9 +2643,9 @@
 	ELandscapeToolTargetTypeFlags LayerTargetTypeMask = EditLayer->GetEnabledTargetTypeMask();
 	// Then disable all types that are not in the layer's mask :
 	for (FEditLayerRendererState& LayerRendererState : RendererStates)
-	{
+{
 		LayerRendererState.DisableTargetTypeMask(~LayerTargetTypeMask);
-	}
+		}
 
 	return RendererStates;
 }
@@ -6276,15 +3308,16 @@
 	/** Given the list of components that actually need to be rendered, divide the work into batches, such that:
 	 *   - When a component is being rendered by a given batch, all components needed for rendering this given component are present in the batch
 	 *   - All components end up being rendered in at least one of the batches
+	 *  @return an empty list if the render cannot actually be processed 
 	 */
-	TArray<FRenderBatchInfo> DivideIntoBatches(const TBitArray<>& InFinalComponentsToRenderInfoBitIndices, const TArray<FComponentToRenderInfo>& InAllComponentsToRenderInfos, bool &bInOutWarnedResolution)
+	TArray<FRenderBatchInfo> DivideIntoBatches(const TBitArray<>& InFinalComponentsToRenderInfoBitIndices, const TArray<FComponentToRenderInfo>& InAllComponentsToRenderInfos, bool &bInOutWarnedResolution, bool &bInOutWarnedMergeDimensionsExceeded)
 	{
 		TRACE_CPUPROFILER_EVENT_SCOPE(DivideIntoBatches);
 
 		const int32 NumFinalComponentsToRender = InFinalComponentsToRenderInfoBitIndices.CountSetBits();
 
 		// Compute the final batch resolution we will use for this work : 
-		const int32 DesiredMaxBatchResolution = CVarLandscapeEditLayersMaxResolutionPerRenderBatch.GetValueOnGameThread();
+		const int32 DesiredMaxBatchResolution = CVarLandscapeBatchedMergeMaxResolutionPerRenderBatch.GetValueOnGameThread();
 
 		// Compute the minimal batch resolution :
 		FIntPoint MinBatchResolution(MIN_int32, MIN_int32);
@@ -6300,6 +3333,19 @@
 			MinBatchResolution = MinBatchResolution.ComponentMax(ComponentToRenderInfo.CombinedSectionRect.Size());
 		}
 
+		if (MinBatchResolution.X > GRHIGlobals.MaxTextureDimensions || MinBatchResolution.Y > GRHIGlobals.MaxTextureDimensions)
+		{
+			if (!bInOutWarnedMergeDimensionsExceeded)
+			{
+				UE_LOG(LogLandscape, Error, TEXT("Cannot render landscape edit layers because the current device does not support render targets of the required size (required resolution: %ix%i, "
+					"maximum resolution supported by the render device: %ix%i). Please reduce landscape size, use a different render device, or make sure that all edit layer renderers participating "
+					"to this landscape don't require an input area that is larger than the device's max texture dimensions."), 
+					MinBatchResolution.X, MinBatchResolution.Y, (int32)GRHIGlobals.MaxTextureDimensions, (int32)GRHIGlobals.MaxTextureDimensions);
+				bInOutWarnedMergeDimensionsExceeded = true;
+			}
+			return {};
+		}
+
 		FIntPoint MaxBatchResolution = MinBatchResolution.ComponentMax(FIntPoint(DesiredMaxBatchResolution));
 		if ((MinBatchResolution.X > DesiredMaxBatchResolution) || (MinBatchResolution.Y > DesiredMaxBatchResolution))
 		{
@@ -6309,12 +3355,12 @@
 				if (CVarSilenceMergeBatchResolutionWarning.GetValueOnGameThread())
 				{
 					// Reduce level to Display and slightly different message
-					UE_LOG(LogLandscape, Display, TEXT("Landscape edit layers merge requires a minimum batch size of resolution %ix%i, which is higher than the current desired maximum batch resolution %ix%i. Consider adjusting the maximum batch resolution (landscape.EditLayersLocalMerge.MaxResolutionPerRenderBatch) or make sure the landscape edit layers renderers in use require smaller work area."),
+					UE_LOG(LogLandscape, Display, TEXT("Landscape edit layers merge requires a minimum batch size of resolution %ix%i, which is higher than the current desired maximum batch resolution %ix%i. Consider adjusting the maximum batch resolution (landscape.BatchedMerge.MaxResolutionPerRenderBatch) or make sure the landscape edit layers renderers in use require smaller work area."),
 						MinBatchResolution.X, MinBatchResolution.Y, DesiredMaxBatchResolution, DesiredMaxBatchResolution);
 				}
 				else
 				{
-					UE_LOG(LogLandscape, Warning, TEXT("Landscape edit layers merge requires a minimum batch size of resolution %ix%i, which is higher than the current desired maximum batch resolution %ix%i. Consider adjusting the maximum batch resolution (landscape.EditLayersLocalMerge.MaxResolutionPerRenderBatch) or make sure the landscape edit layers renderers in use require smaller work area. (Quiet this warning by setting landscape.BatchedMerge.SilenceResolutionWarning to 1)."),
+					UE_LOG(LogLandscape, Warning, TEXT("Landscape edit layers merge requires a minimum batch size of resolution %ix%i, which is higher than the current desired maximum batch resolution %ix%i. Consider adjusting the maximum batch resolution (landscape.BatchedMerge.MaxResolutionPerRenderBatch) or make sure the landscape edit layers renderers in use require smaller work area. (Quiet this warning by setting landscape.BatchedMerge.SilenceResolutionWarning to 1)."),
 						MinBatchResolution.X, MinBatchResolution.Y, DesiredMaxBatchResolution, DesiredMaxBatchResolution);
 				}
 			}
@@ -6727,16 +3773,16 @@
 				if (!EnumHasAnyFlags(WeightmapTargetLayerInfos[TargetLayerIndex].Flags, EWeightmapTargetLayerFlags::Skip))
 				{ 
 				    FLandscapeEditLayersWeightmapsPerformFinalWeightBlendingPS::FParameters* PSParams = GraphBuilder.AllocParameters<FLandscapeEditLayersWeightmapsPerformFinalWeightBlendingPS::FParameters>();
-				    PSParams->RenderTargets[0] = FRenderTargetBinding(OutputTextureRef, ERenderTargetLoadAction::ENoAction, /*InMipIndex = */0, /*InArraySlice = */TargetLayerIndex);
-				    PSParams->InTargetLayerIndex = TargetLayerIndex;
-				    PSParams->InNumTargetLayers = NumTargetLayers;
-				    PSParams->InFinalWeightBlendingTargetLayerInfos = TargetLayerInfosBufferSRV;
-				    PSParams->InCurrentEditLayerWeightmaps = CurrentEditLayerTextureSRVRef;
-    
+					PSParams->RenderTargets[0] = FRenderTargetBinding(OutputTextureRef, ERenderTargetLoadAction::ENoAction, /*InMipIndex = */0, /*InArraySlice = */TargetLayerIndex);
+					PSParams->InTargetLayerIndex = TargetLayerIndex;
+					PSParams->InNumTargetLayers = NumTargetLayers;
+					PSParams->InFinalWeightBlendingTargetLayerInfos = TargetLayerInfosBufferSRV;
+					PSParams->InCurrentEditLayerWeightmaps = CurrentEditLayerTextureSRVRef;
+
 				    FLandscapeEditLayersWeightmapsPerformFinalWeightBlendingPS::PerformFinalWeightBlendingPS(
 					    RDG_EVENT_NAME("FinalWeightBlending(%s) -> %s", *TargetLayerNames[TargetLayerIndex].ToString(), *OutputResourceName),
-					    GraphBuilder, PSParams, EffectiveTextureSize);
-			    }
+					GraphBuilder, PSParams, EffectiveTextureSize);
+			}
 			}
 		};
 
@@ -7085,7 +4131,7 @@
 
 				// Only consider render items which intersect with the renderer's active state (e.g. one a given renderer, a render item (A) might affect heightmaps only and another one (B) weightmaps only, 
 				//  then if performing a heightmap merge, the RendererState's mask here will be ELandscapeToolTargetTypeFlags::Heightmap, so this allows to filter out item B: 
-				const FEditLayerTargetTypeState FilteredRenderItemTargetTypeState = RenderItem.GetTargetTypeState().Intersect(EditLayerRendererRenderInfo.RendererState.GetActiveTargetTypeState()); 
+				const FEditLayerTargetTypeState FilteredRenderItemTargetTypeState = RenderItem.GetTargetTypeState().Intersect(EditLayerRendererRenderInfo.RendererState.GetActiveTargetTypeState());
 
 				bool bOutputRenderItem = false;
 				const FOutputWorldArea& OutputWorldArea = RenderItem.GetOutputWorldArea();
@@ -7312,8 +4358,16 @@
 	verify(Info->GetLandscapeExtent(LandscapeExtent));
 
 	// Now divide the work into batches as large as possible (but fitting in the desired max batch resolution, if possible): 
-	TArray<FRenderBatchInfo> AllBatchInfos = DivideIntoBatches(FinalComponentsToRenderInfoBitIndices, AllComponentsToRenderInfos, bWarnedLayerMergeResolution);
+	TArray<FRenderBatchInfo> AllBatchInfos = DivideIntoBatches(FinalComponentsToRenderInfoBitIndices, AllComponentsToRenderInfos, bWarnedLayerMergeResolution, bWarnedGlobalMergeDimensionsExceeded);
+	
+	// Early out if we will fail to render : 
+	if (AllBatchInfos.IsEmpty())
 	{
+		check(!MergeRenderContext.IsValid());
+		return MergeRenderContext;
+	}
+
+	{
 		TRACE_CPUPROFILER_EVENT_SCOPE(PrepareBatches);
 		MergeRenderContext.RenderBatches.Reserve(AllBatchInfos.Num());
 		MergeRenderContext.TargetLayersToComponents.AddDefaulted(MergeRenderContext.AllTargetLayerNames.Num());
@@ -7568,279 +4622,6 @@
 	return MergeRenderContext;
 }
 
-int32 ALandscape::PerformLayersHeightmapsGlobalMerge(const FUpdateLayersContentContext& InUpdateLayersContentContext, const FEditLayersHeightmapMergeParams& InMergeParams)
-{
-	ULandscapeInfo* Info = GetLandscapeInfo();
-	check(Info != nullptr);
-	check(HeightmapRTList.Num() > 0);
-
-	FIntRect LandscapeExtent;
-	if (!Info->GetLandscapeExtent(LandscapeExtent.Min.X, LandscapeExtent.Min.Y, LandscapeExtent.Max.X, LandscapeExtent.Max.Y))
-	{
-		return 0;
-	}
-
-	// Use to compute top-left vertex position per Heightmap and the actual size to copy : 
-	struct FHeightmapCopyInfo
-	{
-		FHeightmapCopyInfo(UTexture2D* InTexture, const FIntPoint& InComponentVertexPosition, int32 InComponentSizeVerts, FLandscapeEditLayerReadback* InCPUReadback = nullptr)
-			: Texture(InTexture)
-			, ComponentSizeVerts(InComponentSizeVerts)
-			, SectionRect(InComponentVertexPosition, InComponentVertexPosition + FIntPoint(InComponentSizeVerts, InComponentSizeVerts))
-			, CPUReadback(InCPUReadback)
-		{}
-
-		FHeightmapCopyInfo(FHeightmapCopyInfo&&) = default;
-		FHeightmapCopyInfo& operator=(FHeightmapCopyInfo&&) = default;
-
-		void Union(const FIntPoint& InComponentVertexPosition)
-		{
-			SectionRect.Union(FIntRect(InComponentVertexPosition, InComponentVertexPosition + FIntPoint(ComponentSizeVerts, ComponentSizeVerts)));
-		}
-
-		UTexture2D* Texture;
-		int32 ComponentSizeVerts;
-		FIntRect SectionRect;
-		FLandscapeEditLayerReadback* CPUReadback;
-	};
-
-	// Calculate Top Left Lambda
-	auto GetUniqueHeightmaps = [&](const TArray<ULandscapeComponent*>& InLandscapeComponents, TArray<FHeightmapCopyInfo>& OutHeightmaps, const FIntPoint& LandscapeBaseQuads, const FGuid& InLayerGuid = FGuid())
-	{
-		const int32 ComponentSizeQuad = SubsectionSizeQuads * NumSubsections;
-		const int32 ComponentSizeVerts = (SubsectionSizeQuads + 1) * NumSubsections;
-		for (ULandscapeComponent* Component : InLandscapeComponents)
-		{
-			UTexture2D* ComponentHeightmap = Component->GetHeightmap(InLayerGuid);
-
-			int32 Index = OutHeightmaps.IndexOfByPredicate([=](const FHeightmapCopyInfo& LayerHeightmap) { return LayerHeightmap.Texture == ComponentHeightmap; });
-
-			FIntPoint ComponentSectionBase = Component->GetSectionBase() - LandscapeBaseQuads;
-			FVector2D SourcePositionOffset(FMath::RoundToInt((float)ComponentSectionBase.X / ComponentSizeQuad), FMath::RoundToInt((float)ComponentSectionBase.Y / ComponentSizeQuad));
-			FIntPoint ComponentVertexPosition = FIntPoint(static_cast<int32>(SourcePositionOffset.X * ComponentSizeVerts), static_cast<int32>(SourcePositionOffset.Y * ComponentSizeVerts));
-			ALandscapeProxy* Proxy = Component->GetLandscapeProxy();
-
-			if (Index == INDEX_NONE)
-			{
-				FLandscapeEditLayerReadback** CPUReadback = Proxy->HeightmapsCPUReadback.Find(ComponentHeightmap);
-				OutHeightmaps.Add(FHeightmapCopyInfo(ComponentHeightmap, ComponentVertexPosition, ComponentSizeVerts, CPUReadback != nullptr ? *CPUReadback : nullptr));
-			}
-			else
-			{
-				OutHeightmaps[Index].Union(ComponentVertexPosition);
-			}
-		}
-	};
-
-	FLandscapeLayersHeightmapShaderParameters ShaderParams;
-
-	bool FirstLayer = true;
-	UTextureRenderTarget2D* CombinedHeightmapAtlasRT = HeightmapRTList[(int32)EHeightmapRTType::HeightmapRT_CombinedAtlas];
-	UTextureRenderTarget2D* CombinedHeightmapNonAtlasRT = HeightmapRTList[(int32)EHeightmapRTType::HeightmapRT_CombinedNonAtlas];
-	UTextureRenderTarget2D* LandscapeScratchRT1 = HeightmapRTList[(int32)EHeightmapRTType::HeightmapRT_Scratch1];
-	UTextureRenderTarget2D* LandscapeScratchRT2 = HeightmapRTList[(int32)EHeightmapRTType::HeightmapRT_Scratch2];
-	UTextureRenderTarget2D* LandscapeScratchRT3 = HeightmapRTList[(int32)EHeightmapRTType::HeightmapRT_Scratch3];
-
-	for (FLandscapeLayer& Layer : LandscapeEditLayers)
-	{
-		check(Layer.EditLayer != nullptr);
-		const ELandscapeBlendMode LayerBlendMode = Layer.EditLayer->GetBlendMode();
-
-		// Draw each Layer's heightmaps to a Combined RT Atlas in LandscapeScratchRT1
-		ShaderParams.ApplyLayerModifiers = false;
-		ShaderParams.SetAlphaOne = false;
-		ShaderParams.LayerVisible = Layer.EditLayer->IsVisible();
-		ShaderParams.GenerateNormals = false;
-		ShaderParams.LayerBlendMode = LayerBlendMode;
-
-		if (LayerBlendMode == LSBM_AlphaBlend)
-		{
-			// For now, only Layer reserved for Landscape Splines will use the AlphaBlendMode
-			check(Layer.EditLayer->IsA<ULandscapeEditLayerSplines>());
-			ShaderParams.LayerAlpha = 1.0f;
-		}
-		else
-		{
-			check(LayerBlendMode == LSBM_AdditiveBlend);
-			ShaderParams.LayerAlpha = Layer.EditLayer->GetAlphaForTargetType(ELandscapeToolTargetType::Heightmap);
-		}
-
-		{
-			TArray<FLandscapeLayersCopyTextureParams> DeferredCopyTextures;
-			TArray<FHeightmapCopyInfo> LayerHeightmaps;
-			GetUniqueHeightmaps(InUpdateLayersContentContext.LandscapeComponentsHeightmapsToRender, LayerHeightmaps, LandscapeExtent.Min, Layer.EditLayer->GetGuid());
-			for (const FHeightmapCopyInfo& LayerHeightmap : LayerHeightmaps)
-			{
-				FLandscapeLayersCopyTextureParams& CopyTextureParams = DeferredCopyTextures.Add_GetRef(FLandscapeLayersCopyTextureParams(LayerHeightmap.Texture, LandscapeScratchRT1));
-				// Only copy the size that's actually needed : 
-				CopyTextureParams.CopySize = LayerHeightmap.SectionRect.Size();
-				// Copy from the heightmap's top-left corner to the composited texture's position :
-				CopyTextureParams.DestPosition = LayerHeightmap.SectionRect.Min;
-			}
-			ExecuteCopyLayersTexture(MoveTemp(DeferredCopyTextures));
-		}
-
-		// Convert Atlas LandscapeScratchRT1 to the world-projected NonAtlas in LandscapeScratchRT2
-		// TODO: just use this format from the beginning above...
-		DrawHeightmapComponentsToRenderTarget(FString::Printf(TEXT("%s: %s -> NonAtlas %s"), *Layer.EditLayer->GetName().ToString(), *LandscapeScratchRT1->GetName(), *LandscapeScratchRT2->GetName()),
-			InUpdateLayersContentContext.LandscapeComponentsHeightmapsToRender, LandscapeExtent.Min, LandscapeScratchRT1, nullptr, LandscapeScratchRT2, ERTDrawingType::RTAtlasToNonAtlas, true, ShaderParams);
-
-		ShaderParams.ApplyLayerModifiers = true;
-
-		// Combine Current layer NonAtlas LandscapeScratchRT2 with current result in LandscapeScratchRT3, writing final result to CombinedHeightmapNonAtlasRT
-		DrawHeightmapComponentsToRenderTarget(FString::Printf(TEXT("%s: %s += -> CombinedNonAtlas %s"), *Layer.EditLayer->GetName().ToString(), *LandscapeScratchRT2->GetName(), *CombinedHeightmapNonAtlasRT->GetName()),
-			InUpdateLayersContentContext.LandscapeComponentsHeightmapsToRender, LandscapeExtent.Min, LandscapeScratchRT2, FirstLayer ? nullptr : LandscapeScratchRT3, CombinedHeightmapNonAtlasRT, ERTDrawingType::RTNonAtlas, FirstLayer, ShaderParams);
-
-		ShaderParams.ApplyLayerModifiers = false;
-
-		if (Layer.EditLayer->IsVisible() && !InMergeParams.bSkipBrush)
-		{
-			// Draw each brushes
-			for (int32 i = 0; i < Layer.Brushes.Num(); ++i)
-			{
-				// TODO: handle conversion from float to RG8 by using material params to write correct values
-
-				FLandscapeLayerBrush& Brush = Layer.Brushes[i];
-				FLandscapeBrushParameters BrushParameters = BuildLandscapeBrushParameters(/*bIsHeightmapMerge = */true, LandscapeExtent, CombinedHeightmapNonAtlasRT);
-
-				UTextureRenderTarget2D* BrushOutputNonAtlasRT = Brush.RenderLayer(BrushParameters);
-				if (BrushOutputNonAtlasRT == nullptr) 
-				{
-					continue;
-				}
-
-				// TODO: handle conversion/handling of RT not same size as internal size
-				check((BrushOutputNonAtlasRT->SizeX == CombinedHeightmapNonAtlasRT->SizeX) && (BrushOutputNonAtlasRT->SizeY == CombinedHeightmapNonAtlasRT->SizeY));
-
-				ALandscapeBlueprintBrushBase* LandscapeBrush = Brush.GetBrush();
-				check(LandscapeBrush != nullptr); // If we managed to render, the brush should be valid
-
-				INC_DWORD_STAT(STAT_LandscapeLayersRegenerateDrawCalls); // Brush Render
-
-				PrintLayersDebugRT(FString::Printf(TEXT("LS Height: %s %s -> BrushNonAtlas %s"), *Layer.EditLayer->GetName().ToString(), *LandscapeBrush->GetName(), *BrushOutputNonAtlasRT->GetName()), BrushOutputNonAtlasRT);
-
-				// Resolve back to Combined heightmap (it's unlikely, but possible that the brush returns the same RT as input and output, if it did various operations on it, in which case the copy is useless) :
-				if (BrushOutputNonAtlasRT != CombinedHeightmapNonAtlasRT)
-				{
-					ExecuteCopyLayersTexture({ FLandscapeLayersCopyTextureParams(BrushOutputNonAtlasRT, CombinedHeightmapNonAtlasRT) });
-					PrintLayersDebugRT(FString::Printf(TEXT("LS Height: %s Component %s += -> CombinedNonAtlas %s"), *Layer.EditLayer->GetName().ToString(), *BrushOutputNonAtlasRT->GetName(), *CombinedHeightmapNonAtlasRT->GetName()), CombinedHeightmapNonAtlasRT);
-				}
-			}
-
-			// Legacy global merge support for non-blueprint-brush renderers: allow edit layer to act as blueprint brush.
-			ILandscapeBrushRenderCallAdapter_GlobalMergeLegacySupport* RenderCallAdapter = Cast<ILandscapeBrushRenderCallAdapter_GlobalMergeLegacySupport>(Layer.EditLayer);
-			if (RenderCallAdapter)
-			{
-				FLandscapeBrushParameters BrushParameters = BuildLandscapeBrushParameters(
-					/*bIsHeightmapMerge = */true, LandscapeExtent, CombinedHeightmapNonAtlasRT);
-
-				UTextureRenderTarget2D* BrushOutputNonAtlasRT = RenderCallAdapter->RenderAsBlueprintBrush(BrushParameters, GetTransform());
-
-				// Do the same conditional copy that we do for blueprint brushes
-				if (BrushOutputNonAtlasRT != CombinedHeightmapNonAtlasRT
-					&& BrushOutputNonAtlasRT
-					&& BrushOutputNonAtlasRT->SizeX == CombinedHeightmapNonAtlasRT->SizeX
-					&& BrushOutputNonAtlasRT->SizeY == CombinedHeightmapNonAtlasRT->SizeY)
-				{
-					ExecuteCopyLayersTexture({ FLandscapeLayersCopyTextureParams(BrushOutputNonAtlasRT, CombinedHeightmapNonAtlasRT) });
-					PrintLayersDebugRT(FString::Printf(TEXT("LS Height: %s Component %s += -> CombinedNonAtlas %s"), *Layer.EditLayer->GetName().ToString(), 
-						*BrushOutputNonAtlasRT->GetName(), *CombinedHeightmapNonAtlasRT->GetName()), CombinedHeightmapNonAtlasRT);
-				}
-			}
-		}
-
-		// copy CombinedHeightmapNonAtlasRT to LandscapeScratchRT3 (as a source for later layers... this is wasted on the last layer I think...)
-		// TODO: you can get the same effect for much cheaper by swapping these two pointers before the render above...
-		ExecuteCopyLayersTexture({ FLandscapeLayersCopyTextureParams(CombinedHeightmapNonAtlasRT, LandscapeScratchRT3) });
-		PrintLayersDebugRT(FString::Printf(TEXT("LS Height: %s Component %s += -> CombinedNonAtlas %s"), *Layer.EditLayer->GetName().ToString(), *CombinedHeightmapNonAtlasRT->GetName(), *LandscapeScratchRT3->GetName()), LandscapeScratchRT3);
-
-		FirstLayer = false;
-	}
-
-	// Set Alpha channel of valid areas to 1 (via shader copy to LandscapeScratchRT2)
-	ShaderParams.SetAlphaOne = true;
-	DrawHeightmapComponentsToRenderTarget(FString::Printf(TEXT("Mark Valid Area Alpha 1: %s -> %s"), *CombinedHeightmapNonAtlasRT->GetName(), *LandscapeScratchRT2->GetName()),
-		InUpdateLayersContentContext.LandscapeComponentsHeightmapsToRender, LandscapeExtent.Min, CombinedHeightmapNonAtlasRT, nullptr, LandscapeScratchRT2, ERTDrawingType::RTNonAtlas, true, ShaderParams);
-	ShaderParams.SetAlphaOne = false;
-
-	// Broadcast Event of the Full Render
-	if ((InMergeParams.HeightmapUpdateModes & Update_Heightmap_All) == Update_Heightmap_All)
-	{
-		OnEditLayersMergedDelegate.Broadcast(FOnLandscapeEditLayersMergedParams(LandscapeScratchRT3, /*InRenderAreaResolution = */LandscapeExtent.Max - LandscapeExtent.Min, /*bInIsHeightmapMerge = */true));
-	}
-
-	// compute Normals into LandscapeScratchRT1
-	ShaderParams.GenerateNormals = true;
-	ShaderParams.GridSize = GetRootComponent()->GetRelativeScale3D();
-	DrawHeightmapComponentsToRenderTarget(FString::Printf(TEXT("LS Height: %s = -> CombinedNonAtlasNormals : %s"), *LandscapeScratchRT2->GetName(), *LandscapeScratchRT1->GetName()),
-		InUpdateLayersContentContext.LandscapeComponentsHeightmapsToRender, LandscapeExtent.Min, LandscapeScratchRT2,
-		nullptr,
-		LandscapeScratchRT1, ERTDrawingType::RTNonAtlas, true, ShaderParams);
-
-	ShaderParams.GenerateNormals = false;
-
-	// convert back to atlas (TODO: we could do this on the first mip downsample instead...)
-	DrawHeightmapComponentsToRenderTarget(FString::Printf(TEXT("LS Height: %s = -> CombinedAtlasFinal : %s"), *LandscapeScratchRT1->GetName(), *CombinedHeightmapAtlasRT->GetName()),
-		InUpdateLayersContentContext.LandscapeComponentsHeightmapsToRender, LandscapeExtent.Min, LandscapeScratchRT1, nullptr, CombinedHeightmapAtlasRT, ERTDrawingType::RTNonAtlasToAtlas, true, ShaderParams);
-
-	// Downsample to generate mips...
-	DrawHeightmapComponentsToRenderTargetMips(InUpdateLayersContentContext.LandscapeComponentsHeightmapsToRender, LandscapeExtent.Min, CombinedHeightmapAtlasRT, true, ShaderParams);
-
-	// List of UTexture2D that we need to kick off readbacks for :
-	TArray<UTexture2D*> TexturesNeedingReadback;
-
-	// Copy back all Mips to original heightmap data
-	{
-		TArray<FLandscapeLayersCopyTextureParams> DeferredCopyTextures;
-		TArray<FHeightmapCopyInfo> Heightmaps;
-		GetUniqueHeightmaps(InUpdateLayersContentContext.LandscapeComponentsHeightmapsToResolve, Heightmaps, LandscapeExtent.Min);
-		for (const FHeightmapCopyInfo& Heightmap : Heightmaps)
-		{
-			check(Heightmap.CPUReadback);
-
-			const FIntPoint Mip0CopySize = Heightmap.SectionRect.Size();
-			const FIntPoint Mip0SourcePosition = Heightmap.SectionRect.Min;
-
-			// Mip 0
-			{
-				FLandscapeLayersCopyTextureParams& CopyTextureParams = DeferredCopyTextures.Add_GetRef(FLandscapeLayersCopyTextureParams(CombinedHeightmapAtlasRT, Heightmap.Texture));
-				// Only copy the size that's actually needed : 
-				CopyTextureParams.CopySize = Mip0CopySize;
-				// Copy from the composited texture's position to the top-left corner of the heightmap
-				CopyTextureParams.SourcePosition = Mip0SourcePosition;
-			}
-
-			// Other Mips
-			uint8 MipIndex = 1;
-			for (int32 MipRTIndex = (int32)EHeightmapRTType::HeightmapRT_Mip1; MipRTIndex < (int32)EHeightmapRTType::HeightmapRT_Count; ++MipRTIndex)
-			{
-				UTextureRenderTarget2D* RenderTargetMip = HeightmapRTList[MipRTIndex];
-				if (RenderTargetMip != nullptr)
-				{
-					FLandscapeLayersCopyTextureParams& CopyTextureParams = DeferredCopyTextures.Add_GetRef(FLandscapeLayersCopyTextureParams(RenderTargetMip, Heightmap.Texture));
-					CopyTextureParams.CopySize.X = Mip0CopySize.X >> MipIndex;
-					CopyTextureParams.CopySize.Y = Mip0CopySize.Y >> MipIndex;
-					CopyTextureParams.SourcePosition.X = Mip0SourcePosition.X >> MipIndex;
-					CopyTextureParams.SourcePosition.Y = Mip0SourcePosition.Y >> MipIndex;
-					CopyTextureParams.DestMip = MipIndex;
-
-					++MipIndex;
-				}
-			}
-
-			check(!TexturesNeedingReadback.Contains(Heightmap.Texture));
-			TexturesNeedingReadback.Add(Heightmap.Texture);
-		}
-		ExecuteCopyLayersTexture(MoveTemp(DeferredCopyTextures));
-	}
-
-	// Prepare the UTexture2D readbacks we'll need to perform :
-	TArray<FLandscapeLayersCopyReadbackTextureParams> DeferredCopyReadbackTextures = PrepareLandscapeLayersCopyReadbackTextureParams(InUpdateLayersContentContext.MapHelper, TexturesNeedingReadback, /*bWeightmaps = */false);
-	ExecuteCopyToReadbackTexture(DeferredCopyReadbackTextures);
-
-	return InMergeParams.HeightmapUpdateModes;
-}
-
 int32 ALandscape::PerformLayersHeightmapsBatchedMerge(const FUpdateLayersContentContext& InUpdateLayersContentContext, const FEditLayersHeightmapMergeParams& InMergeParams)
 {
 	using namespace UE::Landscape::EditLayers;
@@ -8308,17 +5089,6 @@
 	return true;
 }
 
-
-FLandscapeBrushParameters ALandscape::BuildLandscapeBrushParameters(bool bInIsHeightmapMerge, const FIntRect& InRenderAreaExtents, UTextureRenderTarget2D* InRenderTarget, FName InWeightmapLayerName)
-{
-	FTransform RenderAreaWorldTransform = GetTransform();
-	FVector OffsetVector(InRenderAreaExtents.Min.X, InRenderAreaExtents.Min.Y, 0.f);
-	FVector Translation = RenderAreaWorldTransform.TransformFVector4(OffsetVector);
-	RenderAreaWorldTransform.SetTranslation(Translation);
-	FIntPoint RenderAreaSize = InRenderAreaExtents.Max - InRenderAreaExtents.Min;
-	return FLandscapeBrushParameters(bInIsHeightmapMerge, RenderAreaWorldTransform, RenderAreaSize, InRenderTarget, InWeightmapLayerName);
-}
-
 int32 ALandscape::RegenerateLayersHeightmaps(const FUpdateLayersContentContext& InUpdateLayersContentContext)
 {
 	TRACE_CPUPROFILER_EVENT_SCOPE(LandscapeLayers_RegenerateLayersHeightmaps);
@@ -8369,23 +5139,7 @@
 		MergeParams.bForceRender = bForceRender;
 		MergeParams.bSkipBrush = bSkipBrush;
 
-		switch (CurrentEditLayersMergeMode)
-		{
-		case ELandscapeEditLayersMergeMode::GlobalMerge:
-		{
-			return PerformLayersHeightmapsGlobalMerge(InUpdateLayersContentContext, MergeParams);
-		}
-		case ELandscapeEditLayersMergeMode::LocalMerge:
-		{
-			return PerformLayersHeightmapsLocalMerge(InUpdateLayersContentContext, MergeParams);
-		}
-		case ELandscapeEditLayersMergeMode::BatchedMerge:
-		{
-			return PerformLayersHeightmapsBatchedMerge(InUpdateLayersContentContext, MergeParams);
-		}
-		default:
-			check(false);
-		}
+		return PerformLayersHeightmapsBatchedMerge(InUpdateLayersContentContext, MergeParams);
 	}
 
 	return 0;
@@ -9077,207 +5831,6 @@
 	return bChanged;
 }
 
-void ALandscape::PrepareComponentDataToExtractMaterialLayersCS(const TArray<ULandscapeComponent*>& InLandscapeComponents, const ULandscapeEditLayerBase* InEditLayer, int32 InCurrentWeightmapToProcessIndex, const FIntPoint& InLandscapeBase, FLandscapeTexture2DResource* InOutTextureData,
-	TArray<FLandscapeLayerWeightmapExtractMaterialLayersComponentData>& OutComponentData, TMap<ULandscapeLayerInfoObject*, int32>& OutLayerInfoObjects)
-{
-	ULandscapeInfo* Info = GetLandscapeInfo();
-
-	if (Info == nullptr)
-	{
-		return;
-	}
-
-	TArray<FLandscapeLayersCopyTextureParams> DeferredCopyTextures;
-
-	const int32 LocalComponentSizeQuad = SubsectionSizeQuads * NumSubsections;
-	const int32 LocalComponentSizeVerts = (SubsectionSizeQuads + 1) * NumSubsections;
-	for (ULandscapeComponent* Component : InLandscapeComponents)
-	{
-		FLandscapeLayerComponentData* ComponentLayerData = Component->GetLayerData(InEditLayer->GetGuid());
-
-		if (ComponentLayerData != nullptr)
-		{
-			if (ComponentLayerData->WeightmapData.Textures.IsValidIndex(InCurrentWeightmapToProcessIndex) && ComponentLayerData->WeightmapData.TextureUsages.IsValidIndex(InCurrentWeightmapToProcessIndex))
-			{
-				UTexture2D* LayerWeightmap = ComponentLayerData->WeightmapData.Textures[InCurrentWeightmapToProcessIndex];
-				check(LayerWeightmap != nullptr);
-
-				const ULandscapeWeightmapUsage* LayerWeightmapUsage = ComponentLayerData->WeightmapData.TextureUsages[InCurrentWeightmapToProcessIndex];
-				check(LayerWeightmapUsage != nullptr);
-
-				FIntPoint ComponentSectionBase = Component->GetSectionBase() - InLandscapeBase;
-				FVector2D SourcePositionOffset(FMath::RoundToInt((float)ComponentSectionBase.X / LocalComponentSizeQuad), FMath::RoundToInt((float)ComponentSectionBase.Y / LocalComponentSizeQuad));
-				FIntPoint SourceComponentVertexPosition = FIntPoint(static_cast<int32>(SourcePositionOffset.X * LocalComponentSizeVerts), static_cast<int32>(SourcePositionOffset.Y * LocalComponentSizeVerts));
-
-				FLandscapeLayersCopyTextureParams& CopyTextureParams = DeferredCopyTextures.Add_GetRef(FLandscapeLayersCopyTextureParams(*LayerWeightmap->GetName(), LayerWeightmap->GetResource(), FString::Printf(TEXT("%s WeightmapScratchTexture"), *InEditLayer->GetName().ToString()), InOutTextureData));
-				// Only copy the size that's actually needed : 
-				CopyTextureParams.CopySize.X = LayerWeightmap->GetResource()->GetSizeX();
-				CopyTextureParams.CopySize.Y = LayerWeightmap->GetResource()->GetSizeY();
-				// Copy from the top-left corner of the weightmap to the composited texture's position
-				CopyTextureParams.DestPosition = SourceComponentVertexPosition;
-				PrintLayersDebugTextureResource(FString::Printf(TEXT("LS Weight: %s WeightmapScratchTexture %s"), *InEditLayer->GetName().ToString(), TEXT("WeightmapScratchTextureResource")), InOutTextureData, 0, false);
-
-				for (const FWeightmapLayerAllocationInfo& WeightmapLayerAllocation : ComponentLayerData->WeightmapData.LayerAllocations)
-				{
-					if (WeightmapLayerAllocation.LayerInfo != nullptr && WeightmapLayerAllocation.IsAllocated() && ComponentLayerData->WeightmapData.Textures[WeightmapLayerAllocation.WeightmapTextureIndex] == LayerWeightmap)
-					{
-						const ULandscapeComponent* DestComponent = LayerWeightmapUsage->ChannelUsage[WeightmapLayerAllocation.WeightmapTextureChannel];
-						check(DestComponent);
-
-						FIntPoint DestComponentSectionBase = DestComponent->GetSectionBase() - InLandscapeBase;
-
-						// Compute component top left vertex position from section base info
-						FVector2D DestPositionOffset(FMath::RoundToInt((float)DestComponentSectionBase.X / LocalComponentSizeQuad), FMath::RoundToInt((float)DestComponentSectionBase.Y / LocalComponentSizeQuad));
-
-						FLandscapeLayerWeightmapExtractMaterialLayersComponentData Data{
-							SourceComponentVertexPosition,
-							0,
-							WeightmapLayerAllocation.WeightmapTextureChannel,
-							FIntPoint(static_cast<int32>(DestPositionOffset.X * LocalComponentSizeVerts), static_cast<int32>(DestPositionOffset.Y * LocalComponentSizeVerts))
-						};
-
-						if (WeightmapLayerAllocation.LayerInfo == ALandscapeProxy::VisibilityLayer)
-						{
-							int32& NewLayerInfoObjectIndex = OutLayerInfoObjects.FindOrAdd(ALandscapeProxy::VisibilityLayer);
-							NewLayerInfoObjectIndex = 0;
-						}
-						else
-						{
-							for (int32 LayerInfoSettingsIndex = 0; LayerInfoSettingsIndex < Info->Layers.Num(); ++LayerInfoSettingsIndex)
-							{
-								const FLandscapeInfoLayerSettings& InfoLayerSettings = Info->Layers[LayerInfoSettingsIndex];
-
-								if (InfoLayerSettings.LayerInfoObj == WeightmapLayerAllocation.LayerInfo)
-								{
-									Data.DestinationPaintLayerIndex = LayerInfoSettingsIndex + 1; // due to visibility layer that is at 0
-									int32& NewLayerinfoObjectIndex = OutLayerInfoObjects.FindOrAdd(WeightmapLayerAllocation.LayerInfo);
-									NewLayerinfoObjectIndex = LayerInfoSettingsIndex + 1;
-
-									break;
-								}
-							}
-						}
-
-						OutComponentData.Add(Data);
-					}
-				}
-			}
-		}
-	}
-
-	ExecuteCopyLayersTexture(MoveTemp(DeferredCopyTextures));
-}
-
-void ALandscape::PrepareComponentDataToPackMaterialLayersCS(int32 InCurrentWeightmapToProcessIndex, const FIntPoint& InLandscapeBase, const TArray<ULandscapeComponent*>& InAllLandscapeComponents, TArray<UTexture2D*>& OutProcessedWeightmaps,
-	TArray<FLandscapeEditLayerReadback*>& OutProcessedCPUReadbacks, TArray<FLandscapeLayerWeightmapPackMaterialLayersComponentData>& OutComponentData)
-{
-	ULandscapeInfo* Info = GetLandscapeInfo();
-
-	if (Info == nullptr)
-	{
-		return;
-	}
-
-	// Compute a mapping of all textures for the asked index and their usage
-	TMap<UTexture2D*, ULandscapeWeightmapUsage*> WeightmapsToProcess;
-
-	for (ULandscapeComponent* Component : InAllLandscapeComponents)
-	{
-		const TArray<UTexture2D*>& ComponentWeightmapTextures = Component->GetWeightmapTextures();
-		const TArray<ULandscapeWeightmapUsage*>& ComponentWeightmapTexturesUsage = Component->GetWeightmapTexturesUsage();
-
-		if (ComponentWeightmapTextures.IsValidIndex(InCurrentWeightmapToProcessIndex))
-		{
-			UTexture2D* ComponentWeightmapTexture = ComponentWeightmapTextures[InCurrentWeightmapToProcessIndex];
-			ULandscapeWeightmapUsage* ComponentWeightmapTextureUsage = ComponentWeightmapTexturesUsage[InCurrentWeightmapToProcessIndex];
-			check(ComponentWeightmapTextureUsage != nullptr);
-
-			// It's possible the texture (and its usage, hence) has already been processed by a previous call to PrepareComponentDataToPackMaterialLayersCS with a different InCurrentWeightmapToProcessIndex
-			//  since a texture can be shared by multiple components :
-			if (!OutProcessedWeightmaps.Contains(ComponentWeightmapTexture) && (WeightmapsToProcess.Find(ComponentWeightmapTexture) == nullptr))
-			{
-				WeightmapsToProcess.Add(ComponentWeightmapTexture, ComponentWeightmapTextureUsage);
-				OutProcessedWeightmaps.Add(ComponentWeightmapTexture);
-
-				FLandscapeEditLayerReadback** CPUReadback = Component->GetLandscapeProxy()->WeightmapsCPUReadback.Find(ComponentWeightmapTexture);
-				check(CPUReadback != nullptr);
-
-				OutProcessedCPUReadbacks.Add(*CPUReadback);
-			}
-		}
-	}
-
-	TArray<const FWeightmapLayerAllocationInfo*> AlreadyProcessedAllocation;
-
-	// Build for each texture what each channel should contains
-	for (auto& ItPair : WeightmapsToProcess)
-	{
-		FLandscapeLayerWeightmapPackMaterialLayersComponentData Data;
-
-		for (int32 WeightmapChannelIndex = 0; WeightmapChannelIndex < ULandscapeWeightmapUsage::NumChannels; ++WeightmapChannelIndex)
-		{
-			UTexture2D* ComponentWeightmapTexture = ItPair.Key;
-			ULandscapeWeightmapUsage* ComponentWeightmapTextureUsage = ItPair.Value;
-
-			// Clear out data to known values
-			Data.ComponentVertexPositionX[WeightmapChannelIndex] = INDEX_NONE;
-			Data.ComponentVertexPositionY[WeightmapChannelIndex] = INDEX_NONE;
-			Data.SourcePaintLayerIndex[WeightmapChannelIndex] = INDEX_NONE;
-			Data.WeightmapChannelToProcess[WeightmapChannelIndex] = INDEX_NONE;
-
-			if (ComponentWeightmapTextureUsage->ChannelUsage[WeightmapChannelIndex] != nullptr)
-			{
-				const ULandscapeComponent* ChannelComponent = ComponentWeightmapTextureUsage->ChannelUsage[WeightmapChannelIndex];
-
-				const TArray<FWeightmapLayerAllocationInfo>& ChannelLayerAllocations = ChannelComponent->GetWeightmapLayerAllocations();
-				const TArray<UTexture2D*>& ChannelComponentWeightmapTextures = ChannelComponent->GetWeightmapTextures();
-
-				for (const FWeightmapLayerAllocationInfo& ChannelLayerAllocation : ChannelLayerAllocations)
-				{
-					if (ChannelLayerAllocation.LayerInfo != nullptr && !AlreadyProcessedAllocation.Contains(&ChannelLayerAllocation) && ChannelComponentWeightmapTextures[ChannelLayerAllocation.WeightmapTextureIndex] == ComponentWeightmapTexture)
-					{
-						FIntPoint ComponentSectionBase = ChannelComponent->GetSectionBase() - InLandscapeBase;
-
-						// Compute component top left vertex position from section base info
-						int32 LocalComponentSizeQuad = ChannelComponent->SubsectionSizeQuads * NumSubsections;
-						int32 LocalComponentSizeVerts = (ChannelComponent->SubsectionSizeQuads + 1) * NumSubsections;
-						FVector2D PositionOffset(FMath::RoundToInt((float)ComponentSectionBase.X / LocalComponentSizeQuad), FMath::RoundToInt((float)ComponentSectionBase.Y / LocalComponentSizeQuad));
-
-						Data.ComponentVertexPositionX[WeightmapChannelIndex] = static_cast<int32>(PositionOffset.X * LocalComponentSizeVerts);
-						Data.ComponentVertexPositionY[WeightmapChannelIndex] = static_cast<int32>(PositionOffset.Y * LocalComponentSizeVerts);
-
-						Data.WeightmapChannelToProcess[WeightmapChannelIndex] = ChannelLayerAllocation.WeightmapTextureChannel;
-
-						AlreadyProcessedAllocation.Add(&ChannelLayerAllocation);
-
-						if (ChannelLayerAllocation.LayerInfo == ALandscapeProxy::VisibilityLayer)
-						{
-							Data.SourcePaintLayerIndex[WeightmapChannelIndex] = 0; // Always store after the last weightmap index
-						}
-						else
-						{
-							for (int32 LayerInfoSettingsIndex = 0; LayerInfoSettingsIndex < Info->Layers.Num(); ++LayerInfoSettingsIndex)
-							{
-								const FLandscapeInfoLayerSettings& LayerInfo = Info->Layers[LayerInfoSettingsIndex];
-
-								if (ChannelLayerAllocation.LayerInfo == LayerInfo.LayerInfoObj) //-V1051
-								{
-									Data.SourcePaintLayerIndex[WeightmapChannelIndex] = LayerInfoSettingsIndex + 1; // due to visibility layer that is at 0
-									break;
-								}
-							}
-						}
-
-						break;
-					}
-				}
-			}
-		}
-
-		OutComponentData.Add(Data);
-	}
-}
-
 void ALandscape::ReallocateLayersWeightmaps(FUpdateLayersContentContext& InUpdateLayersContentContext, const TArray<ULandscapeLayerInfoObject*>& InBrushRequiredAllocations, 
 	const TMap<ULandscapeComponent*, TArray<ULandscapeLayerInfoObject*>>* InPerComponentAllocations, TSet<ULandscapeComponent*>* InRestrictTextureSharingToComponents)
 {
@@ -9433,7 +5986,7 @@
 	for (ULandscapeComponent* Component : InUpdateLayersContentContext.LandscapeComponentsWeightmapsToResolve)
 	{
 		// If requested, don't allow the component to use textures that have a weightmap in another component than those provided when reallocating weightmaps :
-		NewCreatedTextures.Append(Component->ReallocateWeightmaps(/*DataInterface = */nullptr, /*InEditLayerGuid = */FGuid(), /*bInSaveToTransactionBuffer = */false, /*bool bInForceReallocate = */false,
+		NewCreatedTextures.Append(Component->ReallocateWeightmaps(/*DataInterface = */nullptr, /*InEditLayerGuid = */FGuid(), /*bInSaveToTransactionBuffer = */false, /*bool bInForceReallocate = */false, 
 			/*InTargetProxy = */nullptr, InRestrictTextureSharingToComponents));
 	}
 
@@ -9452,45 +6005,45 @@
 
 	// Clean-up unused weightmap CPUReadback resources
 	Info->ForEachLandscapeProxy([](ALandscapeProxy* Proxy)
+	{
+		TArray<UTexture2D*, TInlineAllocator<64>> EntriesToRemoveFromMap;
+		for (auto& Pair : Proxy->WeightmapsCPUReadback)
 		{
-			TArray<UTexture2D*, TInlineAllocator<64>> EntriesToRemoveFromMap;
-			for (auto& Pair : Proxy->WeightmapsCPUReadback)
+			UTexture2D* WeightmapTextureKey = Pair.Key;
+			bool IsTextureReferenced = false;
+			for (ULandscapeComponent* Component : Proxy->LandscapeComponents)
 			{
-				UTexture2D* WeightmapTextureKey = Pair.Key;
-				bool IsTextureReferenced = false;
-				for (ULandscapeComponent* Component : Proxy->LandscapeComponents)
+				for (UTexture2D* WeightmapTexture : Component->GetWeightmapTextures(false))
 				{
-					for (UTexture2D* WeightmapTexture : Component->GetWeightmapTextures(false))
+					if (WeightmapTexture == WeightmapTextureKey)
 					{
-						if (WeightmapTexture == WeightmapTextureKey)
-						{
-							IsTextureReferenced = true;
-							break;
-						}
-					}
-				}
-				if (!IsTextureReferenced)
-				{
-					EntriesToRemoveFromMap.Add(WeightmapTextureKey);
-				}
-			}
-
-			if (EntriesToRemoveFromMap.Num())
-			{
-				for (UTexture2D* OldWeightmapTexture : EntriesToRemoveFromMap)
-				{
-					if (FLandscapeEditLayerReadback** CPUReadbackToDelete = Proxy->WeightmapsCPUReadback.Find(OldWeightmapTexture))
-					{
-						check(*CPUReadbackToDelete);
-						delete* CPUReadbackToDelete;
-						Proxy->WeightmapsCPUReadback.Remove(OldWeightmapTexture);
+						IsTextureReferenced = true;
+						break;
 					}
 				}
 			}
+			if (!IsTextureReferenced)
+			{
+				EntriesToRemoveFromMap.Add(WeightmapTextureKey);
+			}
+		}
 
-			return true;
-		});
+		if (EntriesToRemoveFromMap.Num())
+		{
+			for (UTexture2D* OldWeightmapTexture : EntriesToRemoveFromMap)
+			{
+				if (FLandscapeEditLayerReadback** CPUReadbackToDelete = Proxy->WeightmapsCPUReadback.Find(OldWeightmapTexture))
+				{
+					check(*CPUReadbackToDelete);
+					delete* CPUReadbackToDelete;
+					Proxy->WeightmapsCPUReadback.Remove(OldWeightmapTexture);
+				}
+			}
+		}
 
+		return true;
+	});
+
 	int32 NumInvalidated = ComponentsWithChangedAllocs.Num();
 
 	// Very spammy logging for interactive edits.
@@ -9502,7 +6055,7 @@
 
 	for (ULandscapeComponent* Component : ComponentsWithChangedAllocs)
 	{
-		FLandscapeEditLayerComponentReadbackResult* ComponentReadbackResult = UE::Landscape::Private::FindOrAddByComponent(InUpdateLayersContentContext.AllLandscapeComponentReadbackResults, Component, ELandscapeLayerUpdateMode::Update_None);
+		FLandscapeEditLayerComponentReadbackResult* ComponentReadbackResult =  UE::Landscape::Private::FindOrAddByComponent(InUpdateLayersContentContext.AllLandscapeComponentReadbackResults, Component, ELandscapeLayerUpdateMode::Update_None);
 		ComponentReadbackResult->bCleared = true;
 		ComponentReadbackResult->bModified = true;
 	}
@@ -9513,69 +6066,6 @@
 	InUpdateLayersContentContext.Refresh(FUpdateLayersContentContext::ERefreshFlags::RefreshWeightmapInfos | FUpdateLayersContentContext::ERefreshFlags::RefreshMapHelper);
 }
 
-void ALandscape::InitializeLayersWeightmapResources()
-{
-	ULandscapeInfo* Info = GetLandscapeInfo();
-
-	if (Info == nullptr)
-	{
-		return;
-	}
-
-	// Destroy existing resource
-	TArray<FTextureResource*> ResourceToDestroy;
-	ResourceToDestroy.Add(CombinedLayersWeightmapAllMaterialLayersResource);
-	ResourceToDestroy.Add(CurrentLayersWeightmapAllMaterialLayersResource);
-	ResourceToDestroy.Add(WeightmapScratchExtractLayerTextureResource);
-	ResourceToDestroy.Add(WeightmapScratchPackLayerTextureResource);
-
-	for (FTextureResource* Resource : ResourceToDestroy)
-	{
-		if (Resource != nullptr)
-		{
-			ENQUEUE_RENDER_COMMAND(LandscapeLayers_Cmd_ReleaseResources)(
-				[Resource](FRHICommandList& RHICmdList)
-			{
-				TRACE_CPUPROFILER_EVENT_SCOPE(LandscapeLayers_RT_ReleaseResources);
-				Resource->ReleaseResource();
-				delete Resource;
-			});
-		}
-	}
-
-	// Create resources
-
-	int32 LayerCount = Info->Layers.Num() + 1; // due to visibility being stored at 0
-
-	// Use the 1st one to compute the resource as they are all the same anyway
-	UTextureRenderTarget2D* FirstWeightmapRT = WeightmapRTList[(int32)EWeightmapRTType::WeightmapRT_Scratch1];
-
-	CombinedLayersWeightmapAllMaterialLayersResource = new FLandscapeTexture2DArrayResource(FirstWeightmapRT->SizeX, FirstWeightmapRT->SizeY, LayerCount, PF_G8, 1, /*bInNeedUAVs = */ true, /*bInNeedSRV = */ false);
-	BeginInitResource(CombinedLayersWeightmapAllMaterialLayersResource);
-
-	CurrentLayersWeightmapAllMaterialLayersResource = new FLandscapeTexture2DArrayResource(FirstWeightmapRT->SizeX, FirstWeightmapRT->SizeY, LayerCount, PF_G8, 1, /*bInNeedUAVs = */ true, /*bInNeedSRV = */ false);
-	BeginInitResource(CurrentLayersWeightmapAllMaterialLayersResource);
-
-	WeightmapScratchExtractLayerTextureResource = new FLandscapeTexture2DResource(FirstWeightmapRT->SizeX, FirstWeightmapRT->SizeY, PF_B8G8R8A8, 1, /*bInNeedUAVs = */ false, /*bInNeedSRV = */ false);
-	BeginInitResource(WeightmapScratchExtractLayerTextureResource);
-
-	int32 MipCount = 0;
-
-	for (int32 MipRTIndex = (int32)EWeightmapRTType::WeightmapRT_Mip0; MipRTIndex < (int32)EWeightmapRTType::WeightmapRT_Count; ++MipRTIndex)
-	{
-		if (WeightmapRTList[MipRTIndex] != nullptr)
-		{
-			++MipCount;
-		}
-	}
-
-	// Format for UAV can't be PF_B8G8R8A8 on Windows 7 so use PF_R8G8B8A8
-	// We make the final copy out of this to a PF_R8G8B8A8 target with CopyTexturePS() instead of CopyLayersTexture() because a pixel shader will automatically handle the channel swizzling (where a RHICopyTexture won't)
-	WeightmapScratchPackLayerTextureResource = new FLandscapeTexture2DResource(FirstWeightmapRT->SizeX, FirstWeightmapRT->SizeY, PF_R8G8B8A8, MipCount, /*bInNeedUAVs = */ true, /*bInNeedSRV = */ false);
-	BeginInitResource(WeightmapScratchPackLayerTextureResource);
-}
-
-// Little struct that holds information common to PerformLayersWeightmapsLocalMerge, PerformLayersWeightmapsBatchedMerge and PerformLayersWeightmapsGlobalMerge
 struct FEditLayersWeightmapMergeParams
 {
 	int32 WeightmapUpdateModes;
@@ -9583,1410 +6073,6 @@
 	bool bSkipBrush;
 };
 
-// Render-thread version of the data / functions we need for the local merge of edit layers : 
-namespace EditLayersWeightmapLocalMerge_RenderThread
-{
-	struct FEditLayerInfo
-	{
-		int32 SourceWeightmapTextureIndex = INDEX_NONE; // The index in VisibleEditLayerWeightmapTextures of the texture to read from for this layer
-		int32 SourceWeightmapTextureChannel = INDEX_NONE; // The channel of the texture to read from for this layer
-		ELandscapeEditLayerWeightmapBlendMode BlendMode = ELandscapeEditLayerWeightmapBlendMode::Num; // See ELandscapeEditLayerWeightmapBlendMode
-		float Alpha = 1.0f; // Alpha value to be used in the blend
-	};
-
-	struct FComponentPaintLayerRenderInfo
-	{
-		// Name of the paint layer for debug purposes :
-		FString Name;
-		// Describes how to access each visible edit layer's weightmap and how to blend it in the final weightmap for this paint layer :
-		TArray<FEditLayerInfo> VisibleEditLayerInfos;
-		// Global index of this paint layer in the paint layer infos array : 
-		int32 PaintLayerInfoIndex = INDEX_NONE;
-	};
-
-	struct FComponentRenderInfo
-	{
-		// Name of the component for debug purposes :
-		FString Name;
-		// The information needed to render each of the component's paint layers :
-		TArray<FComponentPaintLayerRenderInfo> PaintLayerRenderInfos;
-		// The unique textures referenced by this component's visible edit layer's weightmaps for all paint layers :
-		TArray<FTexture2DResource*> VisibleEditLayerWeightmapTextures;
-	};
-
-	// For a given FTextureResolveBatchInfo, allows to identify a FComponentRenderInfo/FComponentPaintLayerRenderInfo pair (useful when recombining the weightmaps into the final -packed- weightmap)
-	struct FComponentAndPaintLayerRenderInfoIdentifier
-	{
-		FComponentAndPaintLayerRenderInfoIdentifier() = default;
-		FComponentAndPaintLayerRenderInfoIdentifier(int32 InComponentIndex, int32 InPaintLayerIndex)
-			: ComponentIndex(InComponentIndex)
-			, PaintLayerIndex(InPaintLayerIndex)
-		{}
-		bool IsValid() const { return (ComponentIndex != INDEX_NONE) && (PaintLayerIndex != INDEX_NONE); }
-
-		bool operator == (const FComponentAndPaintLayerRenderInfoIdentifier& InOther) const { return (InOther.ComponentIndex == ComponentIndex) && (InOther.PaintLayerIndex == PaintLayerIndex); }
-		bool operator != (const FComponentAndPaintLayerRenderInfoIdentifier& InOther) const { return !(*this == InOther); }
-
-		int32 ComponentIndex = INDEX_NONE; // Index of a FComponentRenderInfo in ComponentToRenderInfos
-		int32 PaintLayerIndex = INDEX_NONE; // Index of a FComponentPaintLayerRenderInfo in FComponentRenderInfo::PaintLayerRenderInfos
-	};
-
-	struct FTextureResolveInfo
-	{
-		FTextureResolveInfo(int32 InNumComponentsToRender)
-			: ComponentToRenderInfoBitIndices(false, InNumComponentsToRender)
-		{}
-
-		void ValidatePerChannelSourceInfo(int32 InChannelIndex, const FComponentAndPaintLayerRenderInfoIdentifier& InComponentAndPaintLayerIdentifier)
-		{
-			for (int32 Index = 0; Index < 4; ++Index)
-			{
-				if (Index == InChannelIndex)
-				{
-					// Channel shouldn't be already assigned
-					check(!PerChannelSourceWeightmapsIdentifiers[Index].IsValid());
-				}
-				else
-				{
-					// There should be no duplicates in channels : that would indicate that there are 2 identical component/layer info pair for 2 different texture channels
-					check(PerChannelSourceWeightmapsIdentifiers[Index] != InComponentAndPaintLayerIdentifier);
-				}
-			}
-		}
-
-		void SetPerChannelSourceInfo(int32 InChannelIndex, const FComponentAndPaintLayerRenderInfoIdentifier& InComponentAndPaintLayerIdentifier)
-		{
-			ValidatePerChannelSourceInfo(InChannelIndex, InComponentAndPaintLayerIdentifier);
-			PerChannelSourceWeightmapsIdentifiers[InChannelIndex] = InComponentAndPaintLayerIdentifier;
-			ComponentToRenderInfoBitIndices[InComponentAndPaintLayerIdentifier.ComponentIndex] = true;
-		}
-
-	public:
-		// Index of this texture in TextureToResolveInfos :
-		int32 TextureToResolveInfoIndex = INDEX_NONE;
-		// Texture that was created or updated that needs resolving : 
-		FTexture2DResource* Texture = nullptr;
-		// List of the 4 identifiers (one per weightmap channel) of a component/paint layer association in a given FTextureResolveBatchInfo that will be used 
-		//  to recombine the individual weightmaps into the final -packed- one
-		TStaticArray<FComponentAndPaintLayerRenderInfoIdentifier, 4> PerChannelSourceWeightmapsIdentifiers;
-		// List of(up to) 4 unique component render info indices that are needed for reconstructing the 4 channels of this texture (index into ComponentToRenderInfos) :
-		//  It's a bit array (1 bit per component to render info) to vastly optimize the division of texture resolve infos into batches, which is a O(N^2) operation : 
-		TBitArray<> ComponentToRenderInfoBitIndices;
-		// CPU readback utility to bring back the result on the CPU : 
-		FLandscapeEditLayerReadback* CPUReadback = nullptr;
-	};
-
-	// Because of weightmaps being shared between one component and another (within a given landcape proxy), we have to group the components to render into batches
-	//  where we'll render all of the paint layers into individual, 1-channel, scratch textures, which we'll then be able to re-assemble into the final, packed, weightmaps :
-	struct FTextureResolveBatchInfo
-	{
-		FTextureResolveBatchInfo(int32 InNumComponentsToRender, int32 InBatchIndex)
-			: BatchIndex(InBatchIndex)
-			, ComponentToRenderInfoBitIndices(false, InNumComponentsToRender)
-		{}
-
-		void AddTexture(const FTextureResolveInfo& InTextureResolveInfo)
-		{
-			check(!TextureToResolveInfoIndices.Contains(InTextureResolveInfo.TextureToResolveInfoIndex));
-			TextureToResolveInfoIndices.Add(InTextureResolveInfo.TextureToResolveInfoIndex);
-			// Remember all the unique components that this texture needs for resolving:
-			ComponentToRenderInfoBitIndices.CombineWithBitwiseOR(InTextureResolveInfo.ComponentToRenderInfoBitIndices, EBitwiseOperatorFlags::MinSize);
-		}
-
-	public:
-		// Index of this batch in TextureResolveBatchInfos
-		int32 BatchIndex = INDEX_NONE;
-
-		// Indices (in ComponentToRenderInfos) of the components whose weightmaps we need to render within this batch in order to produce (and then resolve) the textures in TextureToResolveInfos:
-		//  It's a bit array (1 bit per component to render info) to vastly optimize the division of texture resolve infos into batches, which is a O(N^2) operation : 
-		TBitArray<> ComponentToRenderInfoBitIndices;
-		// Indices (TextureToResolveInfos) of textures that need to be resolved / read back on the CPU :
-		TArray<int32> TextureToResolveInfoIndices;
-	};
-
-	// Description of the entire merge pass :
-	struct FMergeInfo
-	{
-		bool NeedsMerge() const
-		{
-			// If no edit layer or if no paint layer present on any edit layer, we've got nothing to do :
-			bool bNeedsMerge = (MaxNumEditLayersTexturesToMerge > 0) && (MaxNumWeightmapArraysPerResolveTextureBatch > 0);
-			check(!bNeedsMerge || !PaintLayerInfos.IsEmpty()); // If we need merging, we must have at least one paint layer
-			return bNeedsMerge;
-		}
-
-	public:
-		// Maximum number of visible edit layers that have to be merged for a single FComponentRenderInfo : 
-		int32 MaxNumEditLayersTexturesToMerge = 0;
-
-		// Maximum number of weightmap arrays that is needed for a given FTextureResolveBatchInfo (1 per FComponentRenderInfo in the batch) :
-		int32 MaxNumWeightmapArraysPerResolveTextureBatch = 0;
-
-		// Number of vertices per component : 
-		FIntPoint ComponentSizeVerts = 0;
-
-		// Number of sub sections for this landscape :
-		uint32 NumSubsections = 1;
-
-		// Number of mips for the weightmaps of this landscape
-		int32 NumMips = 0;
-
-		// List of batches of FTextureResolveInfo that needs to be resolved in the same pass. This allows massive saves on transient resources on large landscapes because those can be re-cycled from one pass to another :
-		TArray<FTextureResolveBatchInfo> TextureResolveBatchInfos;
-
-		// List of infos for each component that needs its paint layers to be rendered in order to be resolved :
-		TArray<FComponentRenderInfo> ComponentToRenderInfos;
-
-		// List of infos for each texture that needs to be resolved :
-		TArray<FTextureResolveInfo> TextureToResolveInfos;
-
-		// List of infos for each individual paint layer involved in the merge operation (including the visibility layer)
-		TArray<FLandscapeEditLayerWeightmapPaintLayerInfo> PaintLayerInfos;
-
-		// Not truly render-thread data because it references UTextures but it's just because FLandscapeEditLayerReadback were historically game-thread initiated so for as long as we'll use those for readback, we need to store this here : 
-		TArray<FLandscapeLayersCopyReadbackTextureParams> DeferredCopyReadbackTextures;
-	};
-
-	// Render graph intermediate resources :
-	struct FRDGResources
-	{
-		// Contains info about each individual paint layer 
-		FRDGBufferRef PaintLayerInfosBuffer;
-		FRDGBufferSRVRef PaintLayerInfosBufferSRV;
-
-		// Texture array that can be reused from one component / one paint layer to another and that contains the list of all edit layers textures that need merging in a given pass :
-		FRDGTextureRef EditLayersWeightmapsTextureArray;
-		FRDGTextureSRVRef EditLayersWeightmapsTextureArraySRV;
-
-		// List of temporary scratch texture arrays that store the output for all (edit layer-merged) active paint layer (one per landscape component) until they can be packed onto the final weightmap textures. 
-		//  Can be reused from one batch to another :
-		TArray<FRDGTextureRef> ScratchPaintLayerWeightmapTextureArrays;
-		TArray<FRDGTextureSRVRef> ScratchPaintLayerWeightmapTextureArraysSRV;
-
-		// Single scratch texture for the weightmap finalize operation (since we cannot directly write to the final weightmaps because they were not created with TexCreate_RenderTargetable) :
-		FRDGTextureRef ScratchFinalWeightmapTexture;
-		TArray<FRDGTextureSRVRef> ScratchFinalWeightmapTextureMipsSRV; // One SRV per mip level
-
-		// Single structured buffer that will contain all possible FLandscapeEditLayerWeightmapMergeInfo we might need during the entire merge operation (this is to avoid
-		//  too many individual buffer allocations/uploads when many components need to be merged : CPU optimization) :
-		FRDGBufferRef EditLayerMergeInfosBuffer;
-		FRDGBufferSRVRef EditLayerMergeInfosBufferSRV;
-		int32 CurrentEditLayerMergeInfosBufferIndex = 0;
-
-		// Single structured buffer that will contain all possible paint layer info indices we might need during the entire merge operation (this is to avoid
-		//  too many individual buffer allocations/uploads when many components need to be merged : CPU optimization) :
-		FRDGBufferRef PaintLayerInfoIndicesBuffer;
-		FRDGBufferSRVRef PaintLayerInfoIndicesBufferSRV;
-		int32 CurrentPaintLayerInfoIndicesBufferIndex = 0;
-
-		// Dummy buffers :
-		FRDGTextureSRVRef BlackDummyArraySRV;
-	};
-
-	TArray<FLandscapeEditLayerWeightmapMergeInfo> PrepareEditLayerWeightmapMergeInfosBufferData(const FMergeInfo& InLocalMergeInfo)
-	{
-		TArray<FLandscapeEditLayerWeightmapMergeInfo> EditLayerMergeInfos;
-
-		// Batch by batch
-		for (const FTextureResolveBatchInfo& TextureResolveBatchInfo : InLocalMergeInfo.TextureResolveBatchInfos)
-		{
-			// Component by component
-			for (TConstSetBitIterator<> BitIt(TextureResolveBatchInfo.ComponentToRenderInfoBitIndices); BitIt; ++BitIt)
-			{
-				int32 ComponentRenderInfoIndex = BitIt.GetIndex();
-
-				// Paint layer by paint layer
-				for (const FComponentPaintLayerRenderInfo& ComponentPaintLayerRenderInfo : InLocalMergeInfo.ComponentToRenderInfos[ComponentRenderInfoIndex].PaintLayerRenderInfos)
-				{
-					// Update the edit layers merge info big buffer : 
-					for (const FEditLayerInfo& EditLayerInfo : ComponentPaintLayerRenderInfo.VisibleEditLayerInfos)
-					{
-						FLandscapeEditLayerWeightmapMergeInfo& EditLayerMergeInfo = EditLayerMergeInfos.Emplace_GetRef();
-						EditLayerMergeInfo.SourceWeightmapTextureIndex = EditLayerInfo.SourceWeightmapTextureIndex;
-						EditLayerMergeInfo.SourceWeightmapTextureChannel = EditLayerInfo.SourceWeightmapTextureChannel;
-						EditLayerMergeInfo.BlendMode = EditLayerInfo.BlendMode;
-						EditLayerMergeInfo.Alpha = EditLayerInfo.Alpha;
-					}
-				}
-			}
-		}
-
-		return EditLayerMergeInfos;
-	}
-
-	TArray<uint32> PreparePaintLayerInfoIndicesBufferData(const FMergeInfo& InLocalMergeInfo)
-	{
-		TArray<uint32> PaintLayerInfoIndices;
-
-		// Batch by batch
-		for (const FTextureResolveBatchInfo& TextureResolveBatchInfo : InLocalMergeInfo.TextureResolveBatchInfos)
-		{
-			// Texture by texture
-			for (int32 TextureResolveInfoIndex : TextureResolveBatchInfo.TextureToResolveInfoIndices)
-			{
-				// Output channel by output channel
-				const FTextureResolveInfo& TextureResolveInfo = InLocalMergeInfo.TextureToResolveInfos[TextureResolveInfoIndex];
-				for (int32 ChannelIndex = 0; ChannelIndex < 4; ++ChannelIndex)
-				{
-					const FComponentAndPaintLayerRenderInfoIdentifier& ComponentAndPaintLayerRenderInfoIdentifier = TextureResolveInfo.PerChannelSourceWeightmapsIdentifiers[ChannelIndex];
-					if (ComponentAndPaintLayerRenderInfoIdentifier.IsValid())
-					{
-						check(InLocalMergeInfo.ComponentToRenderInfos.IsValidIndex(ComponentAndPaintLayerRenderInfoIdentifier.ComponentIndex)); // this identifier must point to a valid render info in the FTextureResolveBatchInfo
-						const FComponentRenderInfo& ComponentRenderInfo = InLocalMergeInfo.ComponentToRenderInfos[ComponentAndPaintLayerRenderInfoIdentifier.ComponentIndex];
-
-						check(ComponentRenderInfo.PaintLayerRenderInfos.IsValidIndex(ComponentAndPaintLayerRenderInfoIdentifier.PaintLayerIndex)); // this identifier must point to a valid paint layer render info in the FComponentRenderInfo
-
-						// Update the PaintLayerInfoIndices big buffer : 
-						for (const FComponentPaintLayerRenderInfo& PaintLayerInfo : ComponentRenderInfo.PaintLayerRenderInfos)
-						{
-							check(PaintLayerInfo.PaintLayerInfoIndex < InLocalMergeInfo.PaintLayerInfos.Num());
-							PaintLayerInfoIndices.Add(PaintLayerInfo.PaintLayerInfoIndex);
-						}
-					}
-				}
-			}
-		}
-
-		return PaintLayerInfoIndices;
-	}
-
-	void PrepareLayersWeightmapsLocalMergeRDGResources(const FMergeInfo& InLocalMergeInfo, FRDGBuilder& GraphBuilder, FRDGResources& OutResources)
-	{
-		{
-			// Upload paint layer infos buffer once and for all since it's unchanged from one component to another : 
-			OutResources.PaintLayerInfosBuffer = CreateStructuredBuffer(GraphBuilder, TEXT("LandscapePaintLayerInfosBuffer"), InLocalMergeInfo.PaintLayerInfos);
-			OutResources.PaintLayerInfosBufferSRV = GraphBuilder.CreateSRV(FRDGBufferSRVDesc(OutResources.PaintLayerInfosBuffer));
-		}
-
-		{
-			// Allocate a texture array that can contain all edit layers textures to merge for any given component to render (this will be reused from one component to another) :
-			int32 SizeZ = InLocalMergeInfo.MaxNumEditLayersTexturesToMerge;
-			check(SizeZ > 0);
-
-			// TODO [chris.tchou] this texture does not have to be a render target, but RDG does not support transient shader-resource-only/copy-populated textures yet -- see issue https://jira.it.epicgames.com/browse/UE-162198
-			FRDGTextureDesc Desc = FRDGTextureDesc::Create2DArray(InLocalMergeInfo.ComponentSizeVerts, PF_B8G8R8A8, FClearValueBinding::None, TexCreate_ShaderResource | TexCreate_RenderTargetable, static_cast<uint16>(SizeZ), /*InNumMips = */1, /*InNumSamples = */1);
-			OutResources.EditLayersWeightmapsTextureArray = GraphBuilder.CreateTexture(Desc, TEXT("LandscapeEditLayersWeightmapsTextureArray"));
-			OutResources.EditLayersWeightmapsTextureArraySRV = GraphBuilder.CreateSRV(FRDGTextureSRVDesc::Create(OutResources.EditLayersWeightmapsTextureArray));
-		}
-
-		{
-			// Allocate as many texture arrays as needed for a given texture resolve batch (this will be reused from one batch to another) :
-			OutResources.ScratchPaintLayerWeightmapTextureArrays.Reserve(InLocalMergeInfo.MaxNumWeightmapArraysPerResolveTextureBatch);
-			OutResources.ScratchPaintLayerWeightmapTextureArraysSRV.Reserve(InLocalMergeInfo.MaxNumWeightmapArraysPerResolveTextureBatch);
-			// Each texture array (reusable from batch to batch) will contain at most a number of slices equals to the total number of active paint layers (including the visibility layer) : 
-			int32 SizeZ = InLocalMergeInfo.PaintLayerInfos.Num();
-			check(SizeZ > 0);
-			for (int32 Index = 0; Index < InLocalMergeInfo.MaxNumWeightmapArraysPerResolveTextureBatch; ++Index)
-			{
-				FRDGTextureDesc Desc = FRDGTextureDesc::Create2DArray(InLocalMergeInfo.ComponentSizeVerts, PF_G8, FClearValueBinding::None, TexCreate_ShaderResource | TexCreate_RenderTargetable | TexCreate_TargetArraySlicesIndependently, static_cast<uint16>(SizeZ), /*InNumMips = */1, /*InNumSamples = */1);
-				FRDGTextureRef& TextureRef = OutResources.ScratchPaintLayerWeightmapTextureArrays.Add_GetRef(GraphBuilder.CreateTexture(Desc, TEXT("LandscapeEditLayersScratchPaintLayerWeightmapTextureArray")));
-				OutResources.ScratchPaintLayerWeightmapTextureArraysSRV.Add(GraphBuilder.CreateSRV(FRDGTextureSRVDesc::Create(TextureRef)));
-			}
-		}
-
-		{
-			// Allocate a single scratch texture will all of its mip for each individual texture we want to resolve :
-			check(InLocalMergeInfo.NumMips > 0);
-			FRDGTextureDesc Desc = FRDGTextureDesc::Create2D(InLocalMergeInfo.ComponentSizeVerts, PF_B8G8R8A8, FClearValueBinding::None, TexCreate_ShaderResource | TexCreate_RenderTargetable, static_cast<uint8>(InLocalMergeInfo.NumMips), /*InNumSamples = */1);
-			OutResources.ScratchFinalWeightmapTexture = GraphBuilder.CreateTexture(Desc, TEXT("LandscapeEditLayersScratchFinalWeightmapTexture"));
-			for (int32 MipLevel = 0; MipLevel < InLocalMergeInfo.NumMips; ++MipLevel)
-			{
-				OutResources.ScratchFinalWeightmapTextureMipsSRV.Add(GraphBuilder.CreateSRV(FRDGTextureSRVDesc::CreateForMipLevel(OutResources.ScratchFinalWeightmapTexture, MipLevel)));
-			}
-		}
-
-		{
-			// Allocate a single structured buffer that will contain all possible FLandscapeEditLayerWeightmapMergeInfo we might need during the entire merge operation. 
-			//  Although CreateStructuredBuffer can be given a callback to provide its initial data, we need to build the source data array upfront, since it's not compatible with RDG immediate mode:
-			OutResources.EditLayerMergeInfosBuffer = CreateStructuredBuffer(GraphBuilder, TEXT("LandscapeEditLayersMergeInfosBuffer"), PrepareEditLayerWeightmapMergeInfosBufferData(InLocalMergeInfo));
-			OutResources.EditLayerMergeInfosBufferSRV = GraphBuilder.CreateSRV(FRDGBufferSRVDesc(OutResources.EditLayerMergeInfosBuffer));
-		}
-
-		{
-			// Allocate a single structured buffer that will contain all possible paint layer info indices we might need during the entire merge operation.
-			//  Although CreateStructuredBuffer can be given a callback to provide its initial data, we need to build the source data array upfront, since it's not compatible with RDG immediate mode:
-			OutResources.PaintLayerInfoIndicesBuffer = CreateStructuredBuffer(GraphBuilder, TEXT("LandscapeEditLayersPaintLayerInfoIndicesBuffer"), PreparePaintLayerInfoIndicesBufferData(InLocalMergeInfo));
-			OutResources.PaintLayerInfoIndicesBufferSRV = GraphBuilder.CreateSRV(FRDGBufferSRVDesc(OutResources.PaintLayerInfoIndicesBuffer));
-		}
-
-		{
-			// Dummy buffers for avoiding missing shader bindings
-			OutResources.BlackDummyArraySRV = GraphBuilder.CreateSRV(FRDGTextureSRVDesc::Create(GSystemTextures.GetBlackArrayDummy(GraphBuilder)));
-		}
-	}
-
-	// Gather all textures we will want to write into or read from in the render graph and output them in OutTrackedTextures:
-	void GatherLayersWeightmapsLocalMergeRDGTextures(const FMergeInfo& InLocalMergeInfo, TMap<FTexture2DResource*, FLandscapeRDGTrackedTexture>& OutTrackedTextures)
-	{
-		// Gather all source weightmaps :
-		for (const FComponentRenderInfo& ComponentRenderInfo : InLocalMergeInfo.ComponentToRenderInfos)
-		{
-			for (FTexture2DResource* VisibleEditLayersWeightmapTexture : ComponentRenderInfo.VisibleEditLayerWeightmapTextures)
-			{
-				check(VisibleEditLayersWeightmapTexture != nullptr);
-				FLandscapeRDGTrackedTexture* TrackedTexture = OutTrackedTextures.Find(VisibleEditLayersWeightmapTexture);
-				if (TrackedTexture == nullptr)
-				{
-					TrackedTexture = &(OutTrackedTextures.Add(VisibleEditLayersWeightmapTexture, FLandscapeRDGTrackedTexture(VisibleEditLayersWeightmapTexture)));
-				}
-				TrackedTexture->bNeedsSRV = true;
-			}
-		}
-
-		// Gather all destination weightmaps :
-		for (const FTextureResolveInfo& TextureResolveInfo : InLocalMergeInfo.TextureToResolveInfos)
-		{
-			FLandscapeRDGTrackedTexture* TrackedTexture = OutTrackedTextures.Find(TextureResolveInfo.Texture);
-			if (ensure(TrackedTexture == nullptr)) // Resolved textures should only be registered once
-			{
-				TrackedTexture = &(OutTrackedTextures.Add(TextureResolveInfo.Texture, FLandscapeRDGTrackedTexture(TextureResolveInfo.Texture)));
-			}
-
-			TrackedTexture->bNeedsSRV = true;
-		}
-	}
-
-	void MergeEditLayersWeightmapsForBatch(const FTextureResolveBatchInfo& InTextureResolveBatchInfo, const FMergeInfo& InLocalMergeInfo,
-		const TMap<FTexture2DResource*, FLandscapeRDGTrackedTexture>& InTrackedTextures, FRDGBuilder& GraphBuilder, FRDGResources& RDGResources)
-	{
-		RDG_EVENT_SCOPE(GraphBuilder, "Merge paint layers on %d components for batch %d", InTextureResolveBatchInfo.ComponentToRenderInfoBitIndices.CountSetBits(), InTextureResolveBatchInfo.BatchIndex);
-
-		// For each component in the batch, perform the edit layers merge and write the resulting scratch weightmap :
-		int32 IndexInBatch = 0;
-		for (TConstSetBitIterator<> BitIt(InTextureResolveBatchInfo.ComponentToRenderInfoBitIndices); BitIt; ++BitIt, ++IndexInBatch)
-		{
-			int32 ComponentRenderInfoIndex = BitIt.GetIndex();
-
-			const FComponentRenderInfo& ComponentRenderInfo = InLocalMergeInfo.ComponentToRenderInfos[ComponentRenderInfoIndex];
-			RDG_EVENT_SCOPE(GraphBuilder, "Component %s", *ComponentRenderInfo.Name);
-
-			// Prepare a texture array of texture array that will contain all the input textures we'll need for this component, regardless of the paint layer
-			//  (done once per component since various paint layers could share the same texture) :
-			{
-				int32 NumTextures = ComponentRenderInfo.VisibleEditLayerWeightmapTextures.Num();
-				for (int32 TextureIndex = 0; TextureIndex < NumTextures; ++TextureIndex)
-				{
-					const FTexture2DResource* PackedWeightmap = ComponentRenderInfo.VisibleEditLayerWeightmapTextures[TextureIndex];
-
-					const FLandscapeRDGTrackedTexture* TrackedTexture = InTrackedTextures.Find(PackedWeightmap);
-					check(TrackedTexture != nullptr);
-
-					// We need to copy the (portion of the) layer's texture to the texture array : 
-					FRHICopyTextureInfo CopyTextureInfo;
-					CopyTextureInfo.Size = FIntVector(InLocalMergeInfo.ComponentSizeVerts.X, InLocalMergeInfo.ComponentSizeVerts.Y, 0);
-					CopyTextureInfo.DestSliceIndex = TextureIndex;
-
-					AddCopyTexturePass(GraphBuilder, TrackedTexture->ExternalTextureRef, RDGResources.EditLayersWeightmapsTextureArray, CopyTextureInfo);
-				}
-			}
-
-			// We should have a single output scratch texture array reserved for this component in RDGResources.ScratchPaintLayerWeightmapTextureArrays already :
-			check(RDGResources.ScratchPaintLayerWeightmapTextureArrays.IsValidIndex(IndexInBatch));
-			FRDGTextureRef ScratchTextureArrayRef = RDGResources.ScratchPaintLayerWeightmapTextureArrays[IndexInBatch];
-
-			// Paint layer by paint layer, merge the weightmaps from all the corresponding edit layers onto the corresponding scratch texture using the MergeEditLayers PS : 
-			int32 NumComponentPaintLayers = ComponentRenderInfo.PaintLayerRenderInfos.Num();
-			for (int32 ComponentPaintLayerIndex = 0; ComponentPaintLayerIndex < NumComponentPaintLayers; ++ComponentPaintLayerIndex)
-			{
-				const FComponentPaintLayerRenderInfo& ComponentPaintLayerRenderInfo = ComponentRenderInfo.PaintLayerRenderInfos[ComponentPaintLayerIndex];
-				RDG_EVENT_SCOPE(GraphBuilder, "Merge %d edit layers for paint layer %s", ComponentPaintLayerRenderInfo.VisibleEditLayerInfos.Num(), *ComponentPaintLayerRenderInfo.Name);
-
-				FLandscapeLayersWeightmapsMergeEditLayersPS::FParameters* MergeEditLayersPSParams = GraphBuilder.AllocParameters<FLandscapeLayersWeightmapsMergeEditLayersPS::FParameters>();
-				// We'll write to a single slice of the texture array for this component, since we're acting paint layer by paint layer here :
-				MergeEditLayersPSParams->RenderTargets[0] = FRenderTargetBinding(ScratchTextureArrayRef, ERenderTargetLoadAction::ENoAction, /*InMipIndex = */0, /*InArraySlice = */static_cast<int16>(ComponentPaintLayerIndex));
-				MergeEditLayersPSParams->InNumEditLayers = ComponentPaintLayerRenderInfo.VisibleEditLayerInfos.Num();
-				MergeEditLayersPSParams->InPackedWeightmaps = RDGResources.EditLayersWeightmapsTextureArraySRV;
-				MergeEditLayersPSParams->InEditLayersMergeInfos = RDGResources.EditLayerMergeInfosBufferSRV;
-
-				// EditLayerMergeInfosBuffer is already uploaded but we need to tell the shader where we're currently at in that big buffer :
-				MergeEditLayersPSParams->InStartIndexInEditLayersMergeInfos = RDGResources.CurrentEditLayerMergeInfosBufferIndex; 
-				// Update CurrentEditLayerMergeInfosBufferIndex so that the next paint layer starts after : 
-				RDGResources.CurrentEditLayerMergeInfosBufferIndex += ComponentPaintLayerRenderInfo.VisibleEditLayerInfos.Num();
-
-				FLandscapeLayersWeightmapsMergeEditLayersPS::MergeEditLayers(GraphBuilder, MergeEditLayersPSParams);
-			}
-		}
-	}
-
-	void FinalizeSingleWeightmap(const FTextureResolveBatchInfo& InTextureResolveBatchInfo, const FMergeInfo& InLocalMergeInfo,
-		const FTextureResolveInfo& InTextureResolveInfo, FRDGBuilder& GraphBuilder, FRDGResources& RDGResources)
-	{
-		RDG_EVENT_SCOPE(GraphBuilder, "Finalize Weightmap %s", *InTextureResolveInfo.Texture->GetTextureName().ToString());
-
-		FLandscapeLayersWeightmapsFinalizeWeightmapPS::FParameters* FinalizeWeightmapPSParams = GraphBuilder.AllocParameters<FLandscapeLayersWeightmapsFinalizeWeightmapPS::FParameters>();
-		FinalizeWeightmapPSParams->RenderTargets[0] = FRenderTargetBinding(RDGResources.ScratchFinalWeightmapTexture, ERenderTargetLoadAction::ENoAction);
-		FinalizeWeightmapPSParams->InPerChannelPaintLayerIndexInWeightmaps = FUintVector4(ForceInitToZero);
-		FinalizeWeightmapPSParams->InPerChannelStartPaintLayerIndex = FUintVector4(ForceInitToZero);
-		FinalizeWeightmapPSParams->InPerChannelNumPaintLayers = FUintVector4(ForceInitToZero);
-		FinalizeWeightmapPSParams->InPaintLayerInfoIndices = RDGResources.PaintLayerInfoIndicesBufferSRV;
-		FinalizeWeightmapPSParams->InPaintLayerInfos = RDGResources.PaintLayerInfosBufferSRV;
-		FinalizeWeightmapPSParams->InValidTextureChannelsMask = 0;
-
-		check(InTextureResolveInfo.PerChannelSourceWeightmapsIdentifiers.Num() == 4);
-		for (int32 ChannelIndex = 0; ChannelIndex < 4; ++ChannelIndex)
-		{
-			const FComponentAndPaintLayerRenderInfoIdentifier& ComponentAndPaintLayerRenderInfoIdentifier = InTextureResolveInfo.PerChannelSourceWeightmapsIdentifiers[ChannelIndex];
-			if (ComponentAndPaintLayerRenderInfoIdentifier.IsValid())
-			{
-				// Indicate this channel will need to be processed :
-				FinalizeWeightmapPSParams->InValidTextureChannelsMask |= (1 << ChannelIndex);
-
-				check(InLocalMergeInfo.ComponentToRenderInfos.IsValidIndex(ComponentAndPaintLayerRenderInfoIdentifier.ComponentIndex)); // this identifier must point to a valid render info in the FTextureResolveBatchInfo
-				const FComponentRenderInfo& ComponentRenderInfo = InLocalMergeInfo.ComponentToRenderInfos[ComponentAndPaintLayerRenderInfoIdentifier.ComponentIndex];
-
-				check(ComponentRenderInfo.PaintLayerRenderInfos.IsValidIndex(ComponentAndPaintLayerRenderInfoIdentifier.PaintLayerIndex)); // this identifier must point to a valid paint layer render info in the FComponentRenderInfo
-
-				// The paint layer to process on this texture channel : 
-				FinalizeWeightmapPSParams->InPerChannelPaintLayerIndexInWeightmaps[ChannelIndex] = ComponentAndPaintLayerRenderInfoIdentifier.PaintLayerIndex;
-				// The total number of paint layers for this component/paint layer (for weight blending in-between paint layers):
-				FinalizeWeightmapPSParams->InPerChannelNumPaintLayers[ChannelIndex] = ComponentRenderInfo.PaintLayerRenderInfos.Num();
-				// The index at which we'll find the first paint layer info index in the PaintLayerInfoIndices big buffer for this channel : use the index where we're currently at in that big buffer :
-				FinalizeWeightmapPSParams->InPerChannelStartPaintLayerIndex[ChannelIndex] = RDGResources.CurrentPaintLayerInfoIndicesBufferIndex;
-				// And update the big buffer current index so that the next channel starts at the right location in the big buffer : 
-				RDGResources.CurrentPaintLayerInfoIndicesBufferIndex += ComponentRenderInfo.PaintLayerRenderInfos.Num();
-
-				// We should have a single output scratch texture array reserved for this component in RDGResources.ScratchPaintLayerWeightmapTextureArrays already :
-				check(InTextureResolveBatchInfo.ComponentToRenderInfoBitIndices[ComponentAndPaintLayerRenderInfoIdentifier.ComponentIndex]); // This component should have been rendered in that batch!
-				int32 ScratchPaintLayerWeightmapTextureIndex = InTextureResolveBatchInfo.ComponentToRenderInfoBitIndices.CountSetBits(0, ComponentAndPaintLayerRenderInfoIdentifier.ComponentIndex);
-				check(RDGResources.ScratchPaintLayerWeightmapTextureArrays.IsValidIndex(ScratchPaintLayerWeightmapTextureIndex));
-				FinalizeWeightmapPSParams->InPerChannelPaintLayerWeightmaps[ChannelIndex] = RDGResources.ScratchPaintLayerWeightmapTextureArraysSRV[ScratchPaintLayerWeightmapTextureIndex];
-			}
-			else
-			{
-				FinalizeWeightmapPSParams->InPerChannelPaintLayerWeightmaps[ChannelIndex] = RDGResources.BlackDummyArraySRV;
-			}
-		}
-
-		FLandscapeLayersWeightmapsFinalizeWeightmapPS::FinalizeWeightmap(GraphBuilder, FinalizeWeightmapPSParams);
-	}
-
-	void GenerateSingleWeightmapMips(const FTextureResolveInfo& InTextureResolveInfo, const FMergeInfo& InLocalMergeInfo,
-		FRDGBuilder& GraphBuilder, FRDGResources& RDGResources)
-	{
-		RDG_EVENT_SCOPE(GraphBuilder, "Generate mips for Weightmap %s", *InTextureResolveInfo.Texture->GetTextureName().ToString());
-		check(RDGResources.ScratchFinalWeightmapTextureMipsSRV.Num() == InLocalMergeInfo.NumMips);
-
-		FIntPoint CurrentMipSize = InLocalMergeInfo.ComponentSizeVerts;
-		for (int32 MipLevel = 1; MipLevel < InLocalMergeInfo.NumMips; ++MipLevel)
-		{
-			CurrentMipSize.X >>= 1;
-			CurrentMipSize.Y >>= 1;
-
-			// Read from scratch weightmap texture (mip N - 1) -> write to scratch weightmap texture (mip N) :
-			FLandscapeLayersWeightmapsGenerateMipsPS::FParameters* GenerateMipsPSParams = GraphBuilder.AllocParameters<FLandscapeLayersWeightmapsGenerateMipsPS::FParameters>();
-			GenerateMipsPSParams->RenderTargets[0] = FRenderTargetBinding(RDGResources.ScratchFinalWeightmapTexture, ERenderTargetLoadAction::ENoAction, static_cast<uint8>(MipLevel));
-			GenerateMipsPSParams->InCurrentMipSize = FUintVector2(CurrentMipSize.X, CurrentMipSize.Y);
-			GenerateMipsPSParams->InNumSubsections = InLocalMergeInfo.NumSubsections;
-			GenerateMipsPSParams->InSourceWeightmap = RDGResources.ScratchFinalWeightmapTextureMipsSRV[MipLevel - 1];
-
-			FLandscapeLayersWeightmapsGenerateMipsPS::GenerateSingleMip(GraphBuilder, GenerateMipsPSParams);
-		}
-	}
-
-	void CopyScratchToSourceWeightmap(const FTextureResolveInfo& InTextureResolveInfo, const TMap<FTexture2DResource*, FLandscapeRDGTrackedTexture>& InTrackedTextures,
-		FRDGBuilder& GraphBuilder, FRDGResources& RDGResources)
-	{
-		RDG_EVENT_SCOPE(GraphBuilder, "Copy scratch to Weightmap %s", *InTextureResolveInfo.Texture->GetTextureName().ToString());
-
-		const FLandscapeRDGTrackedTexture* TrackedTexture = InTrackedTextures.Find(InTextureResolveInfo.Texture);
-		check((TrackedTexture != nullptr) && (TrackedTexture->ExternalTextureRef != nullptr));
-
-		FRHICopyTextureInfo CopyTextureInfo;
-		// We want to copy all mips : 
-		CopyTextureInfo.NumMips = TrackedTexture->ExternalTextureRef->Desc.NumMips;
-
-		AddCopyTexturePass(GraphBuilder, RDGResources.ScratchFinalWeightmapTexture, TrackedTexture->ExternalTextureRef, CopyTextureInfo);
-	}
-
-	void FinalizeAndResolveWeightmapsForBatch(const FTextureResolveBatchInfo& InTextureResolveBatchInfo, const FMergeInfo& InLocalMergeInfo,
-		const TMap<FTexture2DResource*, FLandscapeRDGTrackedTexture>& InTrackedTextures, FRDGBuilder& GraphBuilder, FRDGResources& RDGResources)
-	{
-		RDG_EVENT_SCOPE(GraphBuilder, "Finalize %d weightmaps for batch %d", InTextureResolveBatchInfo.TextureToResolveInfoIndices.Num(), InTextureResolveBatchInfo.BatchIndex);
-
-		for (int32 TextureResolveInfoIndex : InTextureResolveBatchInfo.TextureToResolveInfoIndices)
-		{
-			const FTextureResolveInfo& TextureResolveInfo = InLocalMergeInfo.TextureToResolveInfos[TextureResolveInfoIndex];
-			RDG_EVENT_SCOPE(GraphBuilder, "Finalize / resolve %s", *TextureResolveInfo.Texture->GetTextureName().ToString());
-
-			// Finalize the weightmap to the scratch texture (cannot directly write to the texture because it's not render-targetable) :
-			FinalizeSingleWeightmap(InTextureResolveBatchInfo, InLocalMergeInfo, TextureResolveInfo, GraphBuilder, RDGResources);
-
-			// Generate mips for this scratch texture :
-			GenerateSingleWeightmapMips(TextureResolveInfo, InLocalMergeInfo, GraphBuilder, RDGResources);
-
-			// And finally, copy to the output texture :
-			CopyScratchToSourceWeightmap(TextureResolveInfo, InTrackedTextures, GraphBuilder, RDGResources);
-		}
-	}
-
-} // End namespace EditLayersWeightmapLocalMerge_RenderThread
-
-void ALandscape::PrepareLayersWeightmapsLocalMergeRenderThreadData(const FUpdateLayersContentContext& InUpdateLayersContentContext, const FEditLayersWeightmapMergeParams& InMergeParams, EditLayersWeightmapLocalMerge_RenderThread::FMergeInfo& OutRenderThreadData)
-{
-	using namespace EditLayersWeightmapLocalMerge_RenderThread;
-	using namespace UE::Landscape::Private;
-
-	TRACE_CPUPROFILER_EVENT_SCOPE(LandscapeLayers_PrepareLayersWeightmapsLocalMergeRenderThreadData);
-
-	ULandscapeInfo* LandscapeInfo = GetLandscapeInfo();
-	check(LandscapeInfo != nullptr);
-
-	// Number of vertices for each landscape component :
-	int32 ComponentSizeVerts = (SubsectionSizeQuads + 1) * NumSubsections;
-	OutRenderThreadData.ComponentSizeVerts = FIntPoint(ComponentSizeVerts, ComponentSizeVerts);
-	OutRenderThreadData.NumMips = FMath::CeilLogTwo(ComponentSizeVerts) + 1;
-	OutRenderThreadData.NumSubsections = NumSubsections;
-
-	// Lookup table to retrieve, for a given paint layer, its index in OutRenderThreadData.PaintLayerInfos (we don't keep UObjects in OutRenderThreadData because it's a render-thread struct) :
-	TMap<ULandscapeLayerInfoObject*, int32> PaintLayerToPaintLayerInfoIndex;
-
-	// Lookup table to retrieve, for a given rendered component/paint layer, its FComponentRenderInfo/FComponentPaintLayerRenderInfo pair identifier (indices)
-	TMap<TPair<ULandscapeComponent*, ULandscapeLayerInfoObject*>, FComponentAndPaintLayerRenderInfoIdentifier> ComponentAndPaintLayerToRenderInfoIndex;
-
-	// Prepare per-landscape component render data : 
-	{
-		TRACE_CPUPROFILER_EVENT_SCOPE(LandscapeLayers_PrepareWeightmapComponentRenderAndPaintLayerInfos);
-		for (ULandscapeComponent* Component : InUpdateLayersContentContext.LandscapeComponentsWeightmapsToRender)
-		{
-			FComponentRenderInfo& ComponentRenderInfo = OutRenderThreadData.ComponentToRenderInfos.Emplace_GetRef();
-			ComponentRenderInfo.Name = Component->GetName();
-
-			// Build, for each of the paint layers, a list of the weightmaps from each of the edit layers that affect the final result :
-			TMap<ULandscapeLayerInfoObject*, int32> PaintLayerToComponentPaintLayerRenderInfoIndex;
-			for (const FLandscapeLayer& Layer : LandscapeEditLayers)
-			{
-				check(Layer.EditLayer != nullptr);
-				if (Layer.EditLayer->IsVisible() && !InMergeParams.bSkipBrush)
-				{
-					// Retrieve the input textures/channels that are needed for this component and this edit layer :
-					const TArray<UTexture2D*>& ComponentEditLayerTextures = Component->GetWeightmapTextures(Layer.EditLayer->GetGuid());
-					const TArray<FWeightmapLayerAllocationInfo>& ComponentEditLayerAllocations = Component->GetWeightmapLayerAllocations(Layer.EditLayer->GetGuid());
-					int32 NumAllocations = ComponentEditLayerAllocations.Num();
-					for (int32 AllocIndex = 0; AllocIndex < NumAllocations; ++AllocIndex)
-					{
-						const FWeightmapLayerAllocationInfo& ComponentEditLayerAllocation = ComponentEditLayerAllocations[AllocIndex];
-						if (ComponentEditLayerAllocation.LayerInfo != nullptr)
-						{
-							// First, take note of the global paint layer information if it has not already been done :
-							int32* PaintLayerInfoIndex = PaintLayerToPaintLayerInfoIndex.Find(ComponentEditLayerAllocation.LayerInfo);
-							if (PaintLayerInfoIndex == nullptr)
-							{
-								FLandscapeEditLayerWeightmapPaintLayerInfo NewPaintLayerInfo;
-								if (ComponentEditLayerAllocation.LayerInfo->GetBlendMethod() == ELandscapeTargetLayerBlendMethod::FinalWeightBlending)
-								{
-									NewPaintLayerInfo.Flags = ELandscapeEditLayerWeightmapPaintLayerFlags::IsWeightBlended;
-								}
-								if (ComponentEditLayerAllocation.LayerInfo == ALandscapeProxy::VisibilityLayer)
-								{
-									NewPaintLayerInfo.Flags = ELandscapeEditLayerWeightmapPaintLayerFlags::IsVisibilityLayer;
-								}
-
-								int32 NewPaintLayerInfoIndex = OutRenderThreadData.PaintLayerInfos.Add(NewPaintLayerInfo);
-								PaintLayerInfoIndex = &PaintLayerToPaintLayerInfoIndex.Add(ComponentEditLayerAllocation.LayerInfo, NewPaintLayerInfoIndex);
-								check(NewPaintLayerInfoIndex == *PaintLayerInfoIndex);
-							}
-
-							int32* ComponentPaintLayerRenderInfoIndex = PaintLayerToComponentPaintLayerRenderInfoIndex.Find(ComponentEditLayerAllocation.LayerInfo);
-							// Add one entry for each paint layer if not already there :
-							if (ComponentPaintLayerRenderInfoIndex == nullptr)
-							{
-								FComponentPaintLayerRenderInfo NewPaintLayerRenderInfo;
-								NewPaintLayerRenderInfo.Name = ComponentEditLayerAllocation.LayerInfo->LayerName.ToString();
-								// Remember which index this paint layer corresponds to in the global paint layer info array : 
-								NewPaintLayerRenderInfo.PaintLayerInfoIndex = *PaintLayerInfoIndex;
-
-								// Add it to the list and add an entry for it in the map:
-								int32 NewPaintLayerRenderInfoIndex = ComponentRenderInfo.PaintLayerRenderInfos.Add(NewPaintLayerRenderInfo);
-								ComponentPaintLayerRenderInfoIndex = &PaintLayerToComponentPaintLayerRenderInfoIndex.Add(ComponentEditLayerAllocation.LayerInfo, NewPaintLayerRenderInfoIndex);
-								check(NewPaintLayerRenderInfoIndex == *ComponentPaintLayerRenderInfoIndex);
-
-								// The next step will need an identifier for this rendered component/rendered paint layer pair:
-								TPair<ULandscapeComponent*, ULandscapeLayerInfoObject*> ComponentAndPaintLayer(Component, ComponentEditLayerAllocation.LayerInfo);
-								check(!ComponentAndPaintLayerToRenderInfoIndex.Contains(ComponentAndPaintLayer));
-								ComponentAndPaintLayerToRenderInfoIndex.Add(ComponentAndPaintLayer, FComponentAndPaintLayerRenderInfoIdentifier(OutRenderThreadData.ComponentToRenderInfos.Num() - 1, NewPaintLayerRenderInfoIndex));
-							}
-
-							FComponentPaintLayerRenderInfo& ComponentPaintLayerRenderInfo = ComponentRenderInfo.PaintLayerRenderInfos[*ComponentPaintLayerRenderInfoIndex];
-
-							// Add the texture we'll need to read from if not already there : 
-							check(ComponentEditLayerTextures.IsValidIndex(ComponentEditLayerAllocation.WeightmapTextureIndex));
-							UTexture2D* ComponentEditLayerTexture = ComponentEditLayerTextures[ComponentEditLayerAllocation.WeightmapTextureIndex];
-							int32 TextureIndexInVisibleEditLayerTextures = ComponentRenderInfo.VisibleEditLayerWeightmapTextures.AddUnique(ComponentEditLayerTexture->GetResource()->GetTexture2DResource());
-
-							const ELandscapeToolTargetType TargetType = UE::Landscape::IsVisibilityLayer(ComponentEditLayerAllocation.LayerInfo) ? ELandscapeToolTargetType::Visibility : ELandscapeToolTargetType::Weightmap;
-
-							// Add an entry for each edit layer that participates to this paint layer : 
-							FEditLayerInfo NewEditLayerInfo;
-							NewEditLayerInfo.SourceWeightmapTextureIndex = TextureIndexInVisibleEditLayerTextures;
-							NewEditLayerInfo.SourceWeightmapTextureChannel = ComponentEditLayerAllocation.WeightmapTextureChannel;
-							NewEditLayerInfo.Alpha = (ComponentEditLayerAllocation.LayerInfo == ALandscapeProxy::VisibilityLayer) ? 1.0f : Layer.EditLayer->GetAlphaForTargetType(ELandscapeToolTargetType::Weightmap); // visibility can't affect or be affected by other paint layer weights
-							const bool* BlendSubstractive = Layer.EditLayer->WeightmapLayerAllocationBlend.Find(ComponentEditLayerAllocation.LayerInfo);
-							NewEditLayerInfo.BlendMode = ((BlendSubstractive != nullptr) && *BlendSubstractive) ? ELandscapeEditLayerWeightmapBlendMode::Subtractive : ELandscapeEditLayerWeightmapBlendMode::Additive;
-
-							ComponentPaintLayerRenderInfo.VisibleEditLayerInfos.Add(NewEditLayerInfo);
-						}
-					}
-				}
-			}
-
-			// Keep track of the maximum number of weightmaps we'll need to merge for a given component :
-			OutRenderThreadData.MaxNumEditLayersTexturesToMerge = FMath::Max(ComponentRenderInfo.VisibleEditLayerWeightmapTextures.Num(), OutRenderThreadData.MaxNumEditLayersTexturesToMerge);
-		}
-	}
-
-	// Collect all UTexture2D that we need to kick off readbacks for and create a FTextureResolveInfo for each:
-	TSet<UTexture2D*> ProcessedTextures;
-	TArray<UTexture2D*> TexturesNeedingReadback;
-	{
-		TRACE_CPUPROFILER_EVENT_SCOPE(LandscapeLayers_PrepareWeightmapTextureResolveInfos);
-		for (ULandscapeComponent* LandscapeComponentToResolve : InUpdateLayersContentContext.LandscapeComponentsWeightmapsToResolve)
-		{
-			ALandscapeProxy* Proxy = LandscapeComponentToResolve->GetLandscapeProxy();
-			check(Proxy != nullptr);
-			const TArray<UTexture2D*>& ComponentBaseWeightmapTextures = LandscapeComponentToResolve->GetWeightmapTextures();
-			for (UTexture2D* ComponentWeightmap : ComponentBaseWeightmapTextures)
-			{
-				bool bIsAlreadyInSet = false;
-				ProcessedTextures.Add(ComponentWeightmap, &bIsAlreadyInSet);
-				if (!bIsAlreadyInSet)
-				{
-					FTexture2DResource* WeightmapResource = ComponentWeightmap->GetResource()->GetTexture2DResource();
-
-					FTextureResolveInfo NewTextureResolveInfo(OutRenderThreadData.ComponentToRenderInfos.Num());
-					NewTextureResolveInfo.TextureToResolveInfoIndex = OutRenderThreadData.TextureToResolveInfos.Num(); // It will be added at the end of the array 
-					NewTextureResolveInfo.Texture = WeightmapResource;
-
-					bool bDoResolve = false;
-
-					// Start preparing the information we need for resolving this texture : we'll need to know how to reconstruct each of its channels :
-					TObjectPtr<ULandscapeWeightmapUsage>* Usage = Proxy->WeightmapUsageMap.Find(ComponentWeightmap);
-					check(Usage && *Usage);
-					// Iterate over all components that participate to this texture's data :
-					for (ULandscapeComponent* SourceComponent : (*Usage)->GetUniqueValidComponents())
-					{
-						// Iterate over all of its allocations to find out which corresponds to this texture :
-						const TArray<UTexture2D*>& SourceComponentBaseWeightmapTextures = SourceComponent->GetWeightmapTextures();
-						const TArray<FWeightmapLayerAllocationInfo>& SourceComponentBaseWeightmapAllocations = SourceComponent->GetWeightmapLayerAllocations();
-						for (const FWeightmapLayerAllocationInfo& AllocationInfo : SourceComponentBaseWeightmapAllocations)
-						{
-							check(SourceComponentBaseWeightmapTextures.IsValidIndex(AllocationInfo.WeightmapTextureIndex));
-							UTexture2D* SourceComponentWeightmap = SourceComponentBaseWeightmapTextures[AllocationInfo.WeightmapTextureIndex];
-							// Same texture, we'll need this allocation
-							if (AllocationInfo.IsAllocated() && (SourceComponentWeightmap == ComponentWeightmap))
-							{
-								TPair<ULandscapeComponent*, ULandscapeLayerInfoObject*> ComponentAndPaintLayer(SourceComponent, AllocationInfo.LayerInfo);
-								FComponentAndPaintLayerRenderInfoIdentifier* ComponentAndPaintLayerIdentifier = ComponentAndPaintLayerToRenderInfoIndex.Find(ComponentAndPaintLayer);
-								// All components needed to recompose this weightmap should be in ComponentToRenderInfos and have a valid paint layer info there:
-								check((ComponentAndPaintLayerIdentifier != nullptr)
-									&& OutRenderThreadData.ComponentToRenderInfos.IsValidIndex(ComponentAndPaintLayerIdentifier->ComponentIndex)
-									&& OutRenderThreadData.ComponentToRenderInfos[ComponentAndPaintLayerIdentifier->ComponentIndex].PaintLayerRenderInfos.IsValidIndex(ComponentAndPaintLayerIdentifier->PaintLayerIndex));
-								NewTextureResolveInfo.SetPerChannelSourceInfo(AllocationInfo.WeightmapTextureChannel, *ComponentAndPaintLayerIdentifier);
-
-								// At least one channel to resolve : we need to resolve the texture : 
-								bDoResolve = true;
-							}
-						}
-					}
-
-					if (bDoResolve)
-					{
-						// Setup the CPU readback if it does not already exist:
-						FLandscapeEditLayerReadback** CPUReadback = Proxy->WeightmapsCPUReadback.Find(ComponentWeightmap);
-						if (CPUReadback == nullptr)
-						{
-							// Lazily create the readback objects as required (ReallocateLayersWeightmaps might have created new weightmaps)
-							FLandscapeEditLayerReadback* NewCPUReadback = new FLandscapeEditLayerReadback();
-							bool bHashIsValid = true;
-							const uint64 Hash = ULandscapeTextureHash::CalculateTextureHash64(ComponentWeightmap, ELandscapeTextureType::Weightmap, bHashIsValid);
-							NewCPUReadback->SetHash(Hash);
-							CPUReadback = &Proxy->WeightmapsCPUReadback.Add(ComponentWeightmap, NewCPUReadback);
-						}
-						check(*CPUReadback != nullptr);
-
-						// Register the CPU readback and add to our list of textures to resolve :
-						NewTextureResolveInfo.CPUReadback = *CPUReadback;
-						OutRenderThreadData.TextureToResolveInfos.Add(NewTextureResolveInfo);
-
-						TexturesNeedingReadback.Add(ComponentWeightmap);
-					}
-				}
-			}
-		}
-		check(OutRenderThreadData.TextureToResolveInfos.Num() == TexturesNeedingReadback.Num());
-	}
-
-	// Prepare the texture resolve batches: 
-	{
-		TRACE_CPUPROFILER_EVENT_SCOPE(LandscapeLayers_PrepareWeightmapTextureResolveInfoBatches);
-
-		int32 MaxComponentsPerResolveBatch = CVarLandscapeEditLayersMaxComponentsPerWeightmapResolveBatch.GetValueOnGameThread();
-
-		// Copy the texture infos because TextureToResolveInfos indices need to remain stable at this point :
-		TArray<FTextureResolveInfo> RemainingTextureToResolveInfos = OutRenderThreadData.TextureToResolveInfos;
-		if (!RemainingTextureToResolveInfos.IsEmpty())
-		{
-			TBitArray<TInlineAllocator<1>> TempBitArray;
-			TempBitArray.Reserve(OutRenderThreadData.ComponentToRenderInfos.Num());
-
-			while (!RemainingTextureToResolveInfos.IsEmpty())
-			{
-				const FTextureResolveInfo& TextureResolveInfo = RemainingTextureToResolveInfos.Pop(EAllowShrinking::No);
-
-				int32 BestBatchIndex = INDEX_NONE;
-				int32 MinNumComponents = MAX_int32;
-
-				// Iterate through all all batches and try to find which would be able to accept it and amongst those, which it would share the most components to render with:
-				int32 NumBatches = OutRenderThreadData.TextureResolveBatchInfos.Num();
-				for (int32 BatchIndex = 0; BatchIndex < NumBatches; ++BatchIndex)
-				{
-					const FTextureResolveBatchInfo& Batch = OutRenderThreadData.TextureResolveBatchInfos[BatchIndex];
-					TempBitArray = TBitArray<>::BitwiseOR(Batch.ComponentToRenderInfoBitIndices, TextureResolveInfo.ComponentToRenderInfoBitIndices, EBitwiseOperatorFlags::MinSize);
-
-					// If after adding its components, the batch still has less than MaxComponentsPerResolveBatch components to render, it can accept it
-					int32 NumComponentsAfter = TempBitArray.CountSetBits();
-					if (NumComponentsAfter <= MaxComponentsPerResolveBatch)
-					{
-						// Is this the best candidate so far?
-						if (NumComponentsAfter < MinNumComponents)
-						{
-							BestBatchIndex = BatchIndex;
-							MinNumComponents = NumComponentsAfter;
-						}
-
-						// If the number of components after addition of this texture is unchanged, it's a perfect match, we won't ever find a better batch so just stop there for this texture:
-						if (NumComponentsAfter == Batch.ComponentToRenderInfoBitIndices.CountSetBits())
-						{
-							break;
-						}
-					}
-				}
-
-				// If we have found a batch, just add the texture to it, otherwise, add a new batch:
-				FTextureResolveBatchInfo& SelectedBatch = (BestBatchIndex != INDEX_NONE) ? OutRenderThreadData.TextureResolveBatchInfos[BestBatchIndex]
-					: OutRenderThreadData.TextureResolveBatchInfos.Add_GetRef(FTextureResolveBatchInfo(OutRenderThreadData.ComponentToRenderInfos.Num(), /*InBatchIndex = */OutRenderThreadData.TextureResolveBatchInfos.Num()));
-
-				SelectedBatch.AddTexture(TextureResolveInfo);
-				check(SelectedBatch.ComponentToRenderInfoBitIndices.CountSetBits() <= MaxComponentsPerResolveBatch);
-
-				// Keep track of the maximum number of scratch texture arrays we'll need for any given batch :
-				OutRenderThreadData.MaxNumWeightmapArraysPerResolveTextureBatch = FMath::Max(SelectedBatch.ComponentToRenderInfoBitIndices.CountSetBits(), OutRenderThreadData.MaxNumWeightmapArraysPerResolveTextureBatch);
-			}
-		}
-	}
-
-	// Finalize :
-	{
-		OutRenderThreadData.DeferredCopyReadbackTextures = PrepareLandscapeLayersCopyReadbackTextureParams(InUpdateLayersContentContext.MapHelper, TexturesNeedingReadback, /*bWeightmaps = */true);
-
-		// Finally, update the material instances to take into account potentially new material combinations : 
-		UpdateLayersMaterialInstances(InUpdateLayersContentContext.LandscapeComponentsWeightmapsToResolve);
-	}
-}
-
-int32 ALandscape::PerformLayersWeightmapsLocalMerge(FUpdateLayersContentContext& InUpdateLayersContentContext, const FEditLayersWeightmapMergeParams& InMergeParams)
-{
-	using namespace EditLayersWeightmapLocalMerge_RenderThread;
-
-	TRACE_CPUPROFILER_EVENT_SCOPE(LandscapeLayers_PerformLayersWeightmapsLocalMerge);
-
-	// We need to perform layer reallocations before doing anything, since additional weightmaps might be added in the process, which might result in new components to render/resolve : 
-	// TODO [jonathan.bard] : this is only possible to do so in local merge, since we don't yet support BP brushes, for which we need to call Render() in order to be able to know which layer allocations
-	//  they need. As such, the global merge path is broken and cannot be fixed unless the BP brush interface is changed in order to inform the system, by advance, about what paint layer it needs
-	// Make sure we have proper textures+allocations for all the final weightmaps we're about to resolve :
-	ReallocateLayersWeightmaps(InUpdateLayersContentContext, /*InBrushRequiredAllocations = */TArray<ULandscapeLayerInfoObject*>(), /*InPerComponentAllocations = */nullptr, /*InRestrictTextureSharingToComponents = */nullptr);
-
-	FMergeInfo RenderThreadData;
-	PrepareLayersWeightmapsLocalMergeRenderThreadData(InUpdateLayersContentContext, InMergeParams, RenderThreadData);
-
-	if (RenderThreadData.NeedsMerge())
-	{
-		ENQUEUE_RENDER_COMMAND(PerformLayersWeightmapsLocalMerge)([RenderThreadData](FRHICommandListImmediate& RHICmdList)
-		{
-			FRDGBuilder GraphBuilder(RHICmdList, RDG_EVENT_NAME("PerformLayersWeightmapsLocalMerge"));
-
-			// Prepare the GPU resources we will use during the local merge : 
-			FRDGResources RDGResources;
-			PrepareLayersWeightmapsLocalMergeRDGResources(RenderThreadData, GraphBuilder, RDGResources);
-
-			// Get a list of all external textures (weightmaps) we will manipulate during the local merge : 
-			TMap<FTexture2DResource*, FLandscapeRDGTrackedTexture> TrackedTextures;
-			GatherLayersWeightmapsLocalMergeRDGTextures(RenderThreadData, TrackedTextures);
-
-			// Start tracking those in the render graph :
-			TrackLandscapeRDGTextures(GraphBuilder, TrackedTextures);
-
-			// Process the components batch by batch in order to avoid over-allocating temporary textures :
-			for (const FTextureResolveBatchInfo& TextureResolveBatchInfo : RenderThreadData.TextureResolveBatchInfos)
-			{
-				RDG_EVENT_SCOPE(GraphBuilder, "Process batch %d", TextureResolveBatchInfo.BatchIndex);
-
-				// Perform all edit layers merges, for all paint layers of all the components in that batch :
-				MergeEditLayersWeightmapsForBatch(TextureResolveBatchInfo, RenderThreadData, TrackedTextures, GraphBuilder, RDGResources);
-
-				// Pack the temporary weightmaps and generate mips on the final texture :
-				FinalizeAndResolveWeightmapsForBatch(TextureResolveBatchInfo, RenderThreadData, TrackedTextures, GraphBuilder, RDGResources);
-			}
-
-			GraphBuilder.Execute();
-		});
-
-		ExecuteCopyToReadbackTexture(RenderThreadData.DeferredCopyReadbackTextures);
-	}
-
-	return InMergeParams.WeightmapUpdateModes;
-}
-
-int32 ALandscape::PerformLayersWeightmapsGlobalMerge(FUpdateLayersContentContext& InUpdateLayersContentContext, const FEditLayersWeightmapMergeParams& InMergeParams)
-{
-	ULandscapeInfo* Info = GetLandscapeInfo();
-	check(Info != nullptr);
-	check(WeightmapRTList.Num() > 0);
-
-	FIntRect LandscapeExtent;
-	if (!Info->GetLandscapeExtent(LandscapeExtent.Min.X, LandscapeExtent.Min.Y, LandscapeExtent.Max.X, LandscapeExtent.Max.Y))
-	{
-		return 0;
-	}
-
-	TArray<ULandscapeLayerInfoObject*> BrushRequiredAllocations;
-	int32 LayerCount = Info->Layers.Num() + 1; // due to visibility being stored at 0
-
-	if (InMergeParams.WeightmapUpdateModes || InMergeParams.bForceRender)
-	{
-		UTextureRenderTarget2D* LandscapeScratchRT1 = WeightmapRTList[(int32)EWeightmapRTType::WeightmapRT_Scratch1];
-		UTextureRenderTarget2D* LandscapeScratchRT2 = WeightmapRTList[(int32)EWeightmapRTType::WeightmapRT_Scratch2];
-		UTextureRenderTarget2D* LandscapeScratchRT3 = WeightmapRTList[(int32)EWeightmapRTType::WeightmapRT_Scratch3];
-		UTextureRenderTarget2D* EmptyRT = WeightmapRTList[(int32)EWeightmapRTType::WeightmapRT_Scratch_RGBA];
-		FLandscapeLayersWeightmapShaderParameters PSShaderParams;
-		FString SourceDebugName;
-		FString DestDebugName;
-		ClearLayersWeightmapTextureResource(TEXT("ClearRT RGBA"), EmptyRT->GameThread_GetRenderTargetResource());
-		ClearLayersWeightmapTextureResource(TEXT("ClearRT R"), LandscapeScratchRT1->GameThread_GetRenderTargetResource());
-
-		{
-			TArray<FLandscapeLayersCopyTextureParams> DeferredCopyTextures;
-			for (int32 LayerIndex = 0; LayerIndex < LayerCount; ++LayerIndex)
-			{
-				SourceDebugName = LandscapeScratchRT1->GetName();
-				DestDebugName = FString::Printf(TEXT("Weight: Clear CombinedProcLayerWeightmapAllLayersResource %d, "), LayerIndex);
-
-				FLandscapeLayersCopyTextureParams& CopyTextureParams = DeferredCopyTextures.Add_GetRef(FLandscapeLayersCopyTextureParams(SourceDebugName, LandscapeScratchRT1->GameThread_GetRenderTargetResource(), DestDebugName, CombinedLayersWeightmapAllMaterialLayersResource));
-				CopyTextureParams.DestArrayIndex = LayerIndex;
-				CopyTextureParams.SourceAccess = ERHIAccess::SRVMask;
-				CopyTextureParams.DestAccess = ERHIAccess::UAVMask;
-			}
-
-			ExecuteCopyLayersTexture(MoveTemp(DeferredCopyTextures));
-		}
-
-		bool bHasWeightmapData = false;
-		bool bFirstLayer = true;
-
-		for (FLandscapeLayer& Layer : LandscapeEditLayers)
-		{
-			check(Layer.EditLayer != nullptr);
-			int8 CurrentWeightmapToProcessIndex = 0;
-			bool HasFoundWeightmapToProcess = true; // try processing at least once
-
-			TMap<ULandscapeLayerInfoObject*, int32> LayerInfoObjects; // <LayerInfoObj, LayerIndex>
-
-			// Determine if some brush want to write to layer that we have currently no data on
-			if (Layer.EditLayer->IsVisible() && !InMergeParams.bSkipBrush)
-			{
-				for (int32 LayerInfoSettingsIndex = 0; LayerInfoSettingsIndex < Info->Layers.Num(); ++LayerInfoSettingsIndex)
-				{
-					const FLandscapeInfoLayerSettings& InfoLayerSettings = Info->Layers[LayerInfoSettingsIndex];
-
-					// It is possible that no layer info has been assigned so that InfoLayerSettings.LayerInfoObj == nullptr. In that case we don't consider it the layer here.
-					if (InfoLayerSettings.LayerInfoObj != nullptr)
-					{
-						for (int32 i = 0; i < Layer.Brushes.Num(); ++i)
-						{
-							FLandscapeLayerBrush& Brush = Layer.Brushes[i];
-							TOptional<int32> LayerInfoSettingsAllocatedIndex;
-
-							// Note that we need to check visibility before we check by name, or else a name match could
-							//  cause us to allocate the visibility layer at the wrong index.
-							if (UE::Landscape::IsVisibilityLayer(InfoLayerSettings.LayerInfoObj))
-							{
-								if (Brush.AffectsVisibilityLayer() && !LayerInfoObjects.Contains(InfoLayerSettings.LayerInfoObj))
-								{
-									LayerInfoSettingsAllocatedIndex = GetVisibilityLayerAllocationIndex();
-								}
-							}
-							else if (Brush.AffectsWeightmapLayer(InfoLayerSettings.GetLayerName()) && !LayerInfoObjects.Contains(InfoLayerSettings.LayerInfoObj))
-							{
-								LayerInfoSettingsAllocatedIndex = LayerInfoSettingsIndex + 1; // due to visibility layer that is at 0
-							}
-
-							if (LayerInfoSettingsAllocatedIndex.IsSet())
-							{
-								LayerInfoObjects.Add(InfoLayerSettings.LayerInfoObj, LayerInfoSettingsAllocatedIndex.GetValue());
-								bHasWeightmapData = true;
-							}
-						}
-
-						// Legacy global merge support for non-blueprint-brush renderers: allow edit layer to act as blueprint brush.
-						ILandscapeBrushRenderCallAdapter_GlobalMergeLegacySupport* RenderCallAdapter = Cast<ILandscapeBrushRenderCallAdapter_GlobalMergeLegacySupport>(Layer.EditLayer);
-						if (RenderCallAdapter && !LayerInfoObjects.Contains(InfoLayerSettings.LayerInfoObj))
-						{
-							TOptional<int32> LayerInfoSettingsAllocatedIndex;
-							if (UE::Landscape::IsVisibilityLayer(InfoLayerSettings.LayerInfoObj))
-							{
-								if (RenderCallAdapter->AffectsVisibilityLayerAsBlueprintBrush())
-								{
-									LayerInfoSettingsAllocatedIndex = GetVisibilityLayerAllocationIndex();
-								}
-							}
-							else if (RenderCallAdapter->AffectsWeightmapLayerAsBlueprintBrush(InfoLayerSettings.GetLayerName()))
-							{
-								LayerInfoSettingsAllocatedIndex = LayerInfoSettingsIndex + 1; // due to visibility layer that is at 0
-							}
-
-							if (LayerInfoSettingsAllocatedIndex.IsSet())
-							{
-								LayerInfoObjects.Add(InfoLayerSettings.LayerInfoObj, LayerInfoSettingsAllocatedIndex.GetValue());
-								bHasWeightmapData = true;
-							}
-						}
-					}
-				}
-			}
-
-			// Track the layers that we have cleared (use a TBitArray in case we get more than 64 layers!)
-			TBitArray<> ClearedLayers(0, Info->Layers.Num() + 1);
-
-			// Loop until there is no more weightmap texture to process
-			while (HasFoundWeightmapToProcess)
-			{
-				SourceDebugName = EmptyRT->GetName();
-				DestDebugName = FString::Printf(TEXT("Weight: %s Clear WeightmapScratchExtractLayerTextureResource"), *Layer.EditLayer->GetName().ToString());
-
-				ExecuteCopyLayersTexture({ FLandscapeLayersCopyTextureParams(SourceDebugName, EmptyRT->GameThread_GetRenderTargetResource(), DestDebugName, WeightmapScratchExtractLayerTextureResource) });
-
-				// Prepare compute shader data
-				TArray<FLandscapeLayerWeightmapExtractMaterialLayersComponentData> ComponentsData;
-				PrepareComponentDataToExtractMaterialLayersCS(InUpdateLayersContentContext.LandscapeComponentsWeightmapsToRender, Layer.EditLayer, CurrentWeightmapToProcessIndex, LandscapeExtent.Min, WeightmapScratchExtractLayerTextureResource,
-					ComponentsData, LayerInfoObjects);
-
-				HasFoundWeightmapToProcess = ComponentsData.Num() > 0;
-
-				// Clear the current atlas if required
-				if (CurrentWeightmapToProcessIndex == 0)
-				{
-					ClearLayersWeightmapTextureResource(TEXT("ClearRT"), LandscapeScratchRT1->GameThread_GetRenderTargetResource());
-				}
-
-				TArray<FLandscapeLayersCopyTextureParams> DeferredCopyTextures;
-				// Important: for performance reason we only clear the layer we will write to, the other one might contain data but they will not be read during the blend phase
-				if (ClearedLayers.CountSetBits() < LayerInfoObjects.Num())
-				{
-					for (auto& ItPair : LayerInfoObjects)
-					{
-						// Only clear the layers that we haven't already cleared
-						const int32 LayerIndex = ItPair.Value;
-						FBitReference ClearedLayerBit = ClearedLayers[LayerIndex];
-						if (!ClearedLayerBit)
-						{
-							ClearedLayerBit = true;
-
-							SourceDebugName = LandscapeScratchRT1->GetName();
-							DestDebugName = FString::Printf(TEXT("Weight: %s Clear CurrentProcLayerWeightmapAllLayersResource %d, "), *Layer.EditLayer->GetName().ToString(), LayerIndex);
-
-							FLandscapeLayersCopyTextureParams& CopyTextureParams = DeferredCopyTextures.Add_GetRef(FLandscapeLayersCopyTextureParams(SourceDebugName, LandscapeScratchRT1->GameThread_GetRenderTargetResource(), DestDebugName, CurrentLayersWeightmapAllMaterialLayersResource));
-							CopyTextureParams.DestAccess = ERHIAccess::UAVMask;
-							CopyTextureParams.DestArrayIndex = LayerIndex;
-						}
-					}
-
-					ExecuteCopyLayersTexture(MoveTemp(DeferredCopyTextures));
-				}
-
-				// Perform the compute shader
-				if (ComponentsData.Num() > 0)
-				{
-					PrintLayersDebugTextureResource(FString::Printf(TEXT("LS Weight: %s WeightmapScratchTexture %s"), *Layer.EditLayer->GetName().ToString(), TEXT("WeightmapScratchTextureResource")), WeightmapScratchExtractLayerTextureResource, 0, false);
-
-					FLandscapeLayerWeightmapExtractMaterialLayersComputeShaderParameters CSExtractLayersShaderParams;
-					CSExtractLayersShaderParams.AtlasWeightmapsPerLayer = CurrentLayersWeightmapAllMaterialLayersResource;
-					CSExtractLayersShaderParams.ComponentWeightmapResource = WeightmapScratchExtractLayerTextureResource;
-					CSExtractLayersShaderParams.ComputeShaderResource = new FLandscapeLayerWeightmapExtractMaterialLayersComputeShaderResource(ComponentsData);
-					CSExtractLayersShaderParams.ComponentSize = (SubsectionSizeQuads + 1) * NumSubsections;
-
-					BeginInitResource(CSExtractLayersShaderParams.ComputeShaderResource);
-
-					FLandscapeLayerWeightmapExtractMaterialLayersCSDispatch_RenderThread CSDispatch(CSExtractLayersShaderParams);
-
-					ENQUEUE_RENDER_COMMAND(LandscapeLayers_Cmd_ExtractLayers)(
-						[CSDispatch](FRHICommandListImmediate& RHICmdList) mutable
-					{
-						TRACE_CPUPROFILER_EVENT_SCOPE(LandscapeLayers_RT_ExtractLayers);
-						CSDispatch.ExtractLayers(RHICmdList);
-					});
-
-					++CurrentWeightmapToProcessIndex;
-					bHasWeightmapData = true; // at least 1 CS was executed, so we can continue the processing
-				}
-			}
-
-			// If we have data in at least one weight map layer
-			if (LayerInfoObjects.Num() > 0)
-			{
-				for (auto& LayerInfoObject : LayerInfoObjects)
-				{
-					int32 LayerIndex = LayerInfoObject.Value;
-					ULandscapeLayerInfoObject* LayerInfoObj = LayerInfoObject.Key;
-
-					RHI_BREADCRUMB_EVENT_GAMETHREAD_F("Layer", "LS Weight: %s PaintLayer: %s", Layer.EditLayer->GetName(), LayerInfoObj->LayerName);
-
-					// Copy the layer we are working on
-					SourceDebugName = FString::Printf(TEXT("Weight: %s PaintLayer: %s, CurrentProcLayerWeightmapAllLayersResource"), *Layer.EditLayer->GetName().ToString(), *LayerInfoObj->LayerName.ToString());
-					DestDebugName = LandscapeScratchRT1->GetName();
-
-					{
-						FLandscapeLayersCopyTextureParams CopyTextureParams(SourceDebugName, CurrentLayersWeightmapAllMaterialLayersResource, DestDebugName, LandscapeScratchRT1->GameThread_GetRenderTargetResource());
-						CopyTextureParams.SourceAccess = ERHIAccess::UAVMask;
-						CopyTextureParams.SourceArrayIndex = LayerIndex;
-						ExecuteCopyLayersTexture({ MoveTemp(CopyTextureParams) });
-					}
-					PrintLayersDebugRT(FString::Printf(TEXT("LS Weight: %s CurrentProcLayerWeightmapAllLayersResource -> Paint Layer RT %s"), *Layer.EditLayer->GetName().ToString(), *LandscapeScratchRT1->GetName()), LandscapeScratchRT1, 0, false);
-
-					const ELandscapeToolTargetType TargetType = UE::Landscape::IsVisibilityLayer(LayerInfoObj) ? ELandscapeToolTargetType::Visibility : ELandscapeToolTargetType::Weightmap;
-
-					PSShaderParams.ApplyLayerModifiers = true;
-					PSShaderParams.LayerVisible = Layer.EditLayer->IsVisible();
-					PSShaderParams.LayerAlpha = Layer.EditLayer->GetAlphaForTargetType(TargetType);
-					checkf((TargetType != ELandscapeToolTargetType::Visibility) || (PSShaderParams.LayerAlpha == 1.0f), TEXT("Visibility can't affect or be affected by other paint layer weights"));
-
-					DrawWeightmapComponentsToRenderTarget(FString::Printf(TEXT("LS Weight: %s PaintLayer: %s, %s += -> %s"), *Layer.EditLayer->GetName().ToString(), *LayerInfoObj->LayerName.ToString(), *LandscapeScratchRT1->GetName(), *LandscapeScratchRT2->GetName()),
-						InUpdateLayersContentContext.LandscapeComponentsWeightmapsToRender, LandscapeExtent.Min, LandscapeScratchRT1, nullptr, LandscapeScratchRT2, ERTDrawingType::RTAtlas, true, PSShaderParams, 0);
-
-					PSShaderParams.ApplyLayerModifiers = false;
-
-					// Combined Layer data with current stack
-					SourceDebugName = FString::Printf(TEXT("Weight: %s PaintLayer: %s CombinedProcLayerWeightmap"), *Layer.EditLayer->GetName().ToString(), *LayerInfoObj->LayerName.ToString());
-					DestDebugName = LandscapeScratchRT1->GetName();
-
-					{
-						FLandscapeLayersCopyTextureParams CopyTextureParams(SourceDebugName, CombinedLayersWeightmapAllMaterialLayersResource, DestDebugName, LandscapeScratchRT1->GameThread_GetRenderTargetResource());
-						CopyTextureParams.SourceAccess = ERHIAccess::UAVMask;
-						CopyTextureParams.SourceArrayIndex = LayerIndex;
-						ExecuteCopyLayersTexture({ MoveTemp(CopyTextureParams) });
-					}
-					PrintLayersDebugRT(FString::Printf(TEXT("LS Weight: %s CombinedProcLayerWeightmap -> Paint Layer RT %s"), *Layer.EditLayer->GetName().ToString(), *LandscapeScratchRT1->GetName()), LandscapeScratchRT1, 0, false);
-
-					// Combine with current status and copy back to the combined 2d resource array
-					PSShaderParams.OutputAsSubstractive = false;
-
-					if (!bFirstLayer)
-					{
-						const bool* BlendSubstractive = Layer.EditLayer->GetWeightmapLayerAllocationBlend().Find(LayerInfoObj);
-						PSShaderParams.OutputAsSubstractive = BlendSubstractive != nullptr ? *BlendSubstractive : false;
-					}
-
-					DrawWeightmapComponentsToRenderTarget(FString::Printf(TEXT("LS Weight: %s PaintLayer: %s, %s += -> Combined %s"), *Layer.EditLayer->GetName().ToString(), *LayerInfoObj->LayerName.ToString(), *LandscapeScratchRT2->GetName(), *LandscapeScratchRT3->GetName()),
-						InUpdateLayersContentContext.LandscapeComponentsWeightmapsToRender, LandscapeExtent.Min, LandscapeScratchRT2, bFirstLayer ? nullptr : LandscapeScratchRT1, LandscapeScratchRT3, ERTDrawingType::RTAtlasToNonAtlas, true, PSShaderParams, 0);
-
-					PSShaderParams.OutputAsSubstractive = false;
-
-					SourceDebugName = FString::Printf(TEXT("Weight: %s PaintLayer: %s %s"), *Layer.EditLayer->GetName().ToString(), *LayerInfoObj->LayerName.ToString(), *LandscapeScratchRT3->GetName());
-					DestDebugName = TEXT("CombinedProcLayerWeightmap");
-
-					// Handle brush blending
-					if (Layer.EditLayer->IsVisible() && !InMergeParams.bSkipBrush)
-					{
-						// Draw each brush				
-						for (int32 i = 0; i < Layer.Brushes.Num(); ++i)
-						{
-							// TODO: handle conversion/handling of RT not same size as internal size
-
-							FLandscapeLayerBrush& Brush = Layer.Brushes[i];
-							FLandscapeBrushParameters BrushParameters = BuildLandscapeBrushParameters(/*bIsHeightmapMerge = */false, LandscapeExtent, LandscapeScratchRT3, LayerInfoObj->LayerName);
-
-							UTextureRenderTarget2D* BrushOutputRT = Brush.RenderLayer(BrushParameters);
-							if (BrushOutputRT == nullptr || BrushOutputRT->SizeX != LandscapeScratchRT3->SizeX || BrushOutputRT->SizeY != LandscapeScratchRT3->SizeY)
-							{
-								continue;
-							}
-
-							ALandscapeBlueprintBrushBase* LandscapeBrush = Brush.GetBrush();
-							check(LandscapeBrush != nullptr); // If we managed to render, the brush should be valid
-
-							BrushRequiredAllocations.AddUnique(LayerInfoObj);
-
-							INC_DWORD_STAT(STAT_LandscapeLayersRegenerateDrawCalls); // Brush RenderInitialize
-
-							PrintLayersDebugRT(FString::Printf(TEXT("LS Weight: %s %s -> Brush %s"), *Layer.EditLayer->GetName().ToString(), *LandscapeBrush->GetName(), *BrushOutputRT->GetName()), BrushOutputRT);
-
-							// Copy result back if brush did not edit things in place.
-
-							// Resolve back to Combined heightmap (it's unlikely, but possible that the brush returns the same RT as input and output, if it did various operations on it, in which case the copy is useless) :
-							if (BrushOutputRT != LandscapeScratchRT3)
-							{
-								SourceDebugName = FString::Printf(TEXT("Weight: %s PaintLayer: %s Brush: %s"), *Layer.EditLayer->GetName().ToString(), *LayerInfoObj->LayerName.ToString(), *BrushOutputRT->GetName());
-								DestDebugName = LandscapeScratchRT3->GetName();
-								ExecuteCopyLayersTexture({ FLandscapeLayersCopyTextureParams(SourceDebugName, BrushOutputRT->GameThread_GetRenderTargetResource(), DestDebugName, LandscapeScratchRT3->GameThread_GetRenderTargetResource()) });
-								PrintLayersDebugRT(FString::Printf(TEXT("LS Weight: %s Component %s += -> Combined %s"), *Layer.EditLayer->GetName().ToString(), *BrushOutputRT->GetName(), *LandscapeScratchRT3->GetName()), LandscapeScratchRT3);
-							}
-						}
-
-						// Legacy global merge support for non-blueprint-brush renderers: allow edit layer to act as blueprint brush.
-						ILandscapeBrushRenderCallAdapter_GlobalMergeLegacySupport* RenderCallAdapter = Cast<ILandscapeBrushRenderCallAdapter_GlobalMergeLegacySupport>(Layer.EditLayer);
-						if (RenderCallAdapter)
-						{
-							FLandscapeBrushParameters BrushParameters = BuildLandscapeBrushParameters(
-								/*bIsHeightmapMerge = */false, LandscapeExtent, LandscapeScratchRT3, LayerInfoObj->LayerName);
-							
-							UTextureRenderTarget2D* BrushOutputRT = RenderCallAdapter->RenderAsBlueprintBrush(BrushParameters, GetTransform());
-
-							if (BrushOutputRT != nullptr && BrushOutputRT->SizeX == LandscapeScratchRT3->SizeX 
-								&& BrushOutputRT->SizeY == LandscapeScratchRT3->SizeY)
-							{
-								BrushRequiredAllocations.AddUnique(LayerInfoObj);
-
-								// Same conditional copy as for blueprint brushes
-								if (BrushOutputRT != LandscapeScratchRT3)
-								{
-									SourceDebugName = FString::Printf(TEXT("Weight: %s PaintLayer: %s Brush: %s"), *Layer.EditLayer->GetName().ToString(), *LayerInfoObj->LayerName.ToString(), *BrushOutputRT->GetName());
-									DestDebugName = LandscapeScratchRT3->GetName();
-									ExecuteCopyLayersTexture({ FLandscapeLayersCopyTextureParams(SourceDebugName, BrushOutputRT->GameThread_GetRenderTargetResource(), DestDebugName, LandscapeScratchRT3->GameThread_GetRenderTargetResource()) });
-									PrintLayersDebugRT(FString::Printf(TEXT("LS Weight: %s Component %s += -> Combined %s"), *Layer.EditLayer->GetName().ToString(), *BrushOutputRT->GetName(), *LandscapeScratchRT3->GetName()), LandscapeScratchRT3);
-								}
-							}
-						}
-
-						PrintLayersDebugRT(FString::Printf(TEXT("LS Weight: %s CombinedPostBrushProcLayerWeightmap -> Paint Layer RT %s"), *Layer.EditLayer->GetName().ToString(), *LandscapeScratchRT3->GetName()), LandscapeScratchRT3, 0, false);
-
-						SourceDebugName = FString::Printf(TEXT("Weight: %s PaintLayer: %s %s"), *Layer.EditLayer->GetName().ToString(), *LayerInfoObj->LayerName.ToString(), *LandscapeScratchRT3->GetName());
-						DestDebugName = TEXT("CombinedProcLayerWeightmap");
-
-						FLandscapeLayersCopyTextureParams CopyTextureParams(SourceDebugName, LandscapeScratchRT3->GameThread_GetRenderTargetResource(), DestDebugName, CombinedLayersWeightmapAllMaterialLayersResource);
-						CopyTextureParams.DestAccess = ERHIAccess::UAVMask;
-						CopyTextureParams.DestArrayIndex = LayerIndex;
-						ExecuteCopyLayersTexture({ MoveTemp(CopyTextureParams) });
-					}
-
-					DrawWeightmapComponentsToRenderTarget(FString::Printf(TEXT("LS Weight: %s Combined Scratch No Border to %s Combined Scratch with Border"), *LandscapeScratchRT3->GetName(), *LandscapeScratchRT1->GetName()),
-						InUpdateLayersContentContext.LandscapeComponentsWeightmapsToRender, LandscapeExtent.Min, LandscapeScratchRT3, nullptr, LandscapeScratchRT1, ERTDrawingType::RTNonAtlasToAtlas, true, PSShaderParams, 0);
-
-					FLandscapeLayersCopyTextureParams CopyTextureParams(SourceDebugName, LandscapeScratchRT1->GameThread_GetRenderTargetResource(), DestDebugName, CombinedLayersWeightmapAllMaterialLayersResource);
-					CopyTextureParams.DestAccess = ERHIAccess::UAVMask;
-					CopyTextureParams.DestArrayIndex = LayerIndex;
-					ExecuteCopyLayersTexture({ MoveTemp(CopyTextureParams) });
-				}
-
-				PSShaderParams.ApplyLayerModifiers = false;
-			}
-
-			bFirstLayer = false;
-		}
-
-		ReallocateLayersWeightmaps(InUpdateLayersContentContext, BrushRequiredAllocations, /*InPerComponentAllocations = */nullptr, /*InRestrictTextureSharingToComponents = */nullptr);
-
-		// List of UTexture2D that we need to kick off readbacks for :
-		TArray<UTexture2D*> TexturesNeedingReadback;
-
-		if (bHasWeightmapData)
-		{
-			// Lazily create CPU read back objects as required
-			for (ULandscapeComponent* Component : InUpdateLayersContentContext.LandscapeComponentsWeightmapsToRender)
-			{
-				const TArray<UTexture2D*>& ComponentWeightmapTextures = Component->GetWeightmapTextures();
-
-				for (UTexture2D* WeightmapTexture : ComponentWeightmapTextures)
-				{
-					ALandscapeProxy* Proxy = Component->GetLandscapeProxy();
-					FLandscapeEditLayerReadback** CPUReadback = Proxy->WeightmapsCPUReadback.Find(WeightmapTexture);
-
-					if (CPUReadback == nullptr)
-					{
-						FLandscapeEditLayerReadback* NewCPUReadback = new FLandscapeEditLayerReadback();
-						bool bHashIsValid = true;
-						const uint64 Hash = ULandscapeTextureHash::CalculateTextureHash64(WeightmapTexture, ELandscapeTextureType::Weightmap, bHashIsValid);
-						NewCPUReadback->SetHash(Hash);
-						Proxy->WeightmapsCPUReadback.Add(WeightmapTexture, NewCPUReadback);
-					}
-				}
-			}
-
-			int8 CurrentWeightmapToProcessIndex = 0;
-			bool HasFoundWeightmapToProcess = true; // try processing at least once
-
-			TArray<float> WeightmapLayerWeightBlend;
-			TArray<UTexture2D*> ProcessedWeightmaps;
-			TArray<FLandscapeEditLayerReadback*> ProcessedCPUReadbackTextures;
-			int32 NextTextureIndexToProcess = 0;
-
-			// Generate the component data from the weightmap allocation that were done earlier and weight blend them if required (i.e renormalize)
-			while (HasFoundWeightmapToProcess)
-			{
-				TArray<FLandscapeLayerWeightmapPackMaterialLayersComponentData> PackLayersComponentsData;
-				PrepareComponentDataToPackMaterialLayersCS(CurrentWeightmapToProcessIndex, LandscapeExtent.Min, InUpdateLayersContentContext.LandscapeComponentsWeightmapsToRender, ProcessedWeightmaps, ProcessedCPUReadbackTextures, PackLayersComponentsData);
-				HasFoundWeightmapToProcess = PackLayersComponentsData.Num() > 0;
-
-				// Perform the compute shader
-				if (PackLayersComponentsData.Num() > 0)
-				{
-					// Compute the weightblend mode of each layer for the compute shader
-					if (WeightmapLayerWeightBlend.Num() != LayerCount)
-					{
-						WeightmapLayerWeightBlend.SetNum(LayerCount);
-
-						for (int32 LayerInfoSettingsIndex = 0; LayerInfoSettingsIndex < Info->Layers.Num(); ++LayerInfoSettingsIndex)
-						{
-							const FLandscapeInfoLayerSettings& LayerInfo = Info->Layers[LayerInfoSettingsIndex];
-							WeightmapLayerWeightBlend[LayerInfoSettingsIndex + 1] = (LayerInfo.LayerInfoObj != nullptr) ? ((LayerInfo.LayerInfoObj->GetBlendMethod() == ELandscapeTargetLayerBlendMethod::FinalWeightBlending) ? 1.0f : 0.0f) : 1.0f;
-						}
-
-						WeightmapLayerWeightBlend[0] = 0.0f; // Blend of Visibility 
-					}
-
-					TArray<FVector2f> WeightmapTextureOutputOffset;
-
-					// Compute each weightmap location so compute shader will be able to output at expected location
-					const int32 WeightmapSizeX = WeightmapScratchPackLayerTextureResource->GetSizeX();
-					const int32 WeightmapSizeY = WeightmapScratchPackLayerTextureResource->GetSizeY();
-					const int32 ComponentSize = (SubsectionSizeQuads + 1) * NumSubsections;
-
-					float ComponentY = 0;
-					float ComponentX = 0;
-
-					for (int32 i = 0; i < PackLayersComponentsData.Num(); ++i)
-					{
-						if (ComponentX + ComponentSize > WeightmapSizeX)
-						{
-							ComponentY += ComponentSize;
-							ComponentX = 0;
-						}
-
-						// This should never happen as it would be a bug in the algo
-						check(ComponentX + ComponentSize <= WeightmapSizeX);
-						check(ComponentY + ComponentSize <= WeightmapSizeY);
-
-						WeightmapTextureOutputOffset.Add(FVector2f(ComponentX, ComponentY));
-						ComponentX += ComponentSize;
-					}
-
-					// Clear Pack texture
-					SourceDebugName = *EmptyRT->GetName();
-					DestDebugName = TEXT("Weight: Clear WeightmapScratchPackLayerTextureResource");
-
-					CopyTexturePS(SourceDebugName, EmptyRT->GameThread_GetRenderTargetResource(), DestDebugName, WeightmapScratchPackLayerTextureResource);
-
-					ENQUEUE_RENDER_COMMAND(LandscapeLayers_TransitionPackLayerResources)(
-						[this](FRHICommandListImmediate& RHICmdList) mutable
-					{
-						RHICmdList.Transition(FRHITransitionInfo(CombinedLayersWeightmapAllMaterialLayersResource->TextureRHI, ERHIAccess::UAVMask, ERHIAccess::SRVMask));
-						RHICmdList.Transition(FRHITransitionInfo(WeightmapScratchPackLayerTextureResource->TextureRHI, ERHIAccess::RTV, ERHIAccess::UAVMask));
-					});
-
-					FLandscapeLayerWeightmapPackMaterialLayersComputeShaderParameters CSPackLayersShaderParams;
-					CSPackLayersShaderParams.AtlasWeightmapsPerLayer = CombinedLayersWeightmapAllMaterialLayersResource;
-					CSPackLayersShaderParams.ComponentWeightmapResource = WeightmapScratchPackLayerTextureResource;
-					CSPackLayersShaderParams.ComputeShaderResource = new FLandscapeLayerWeightmapPackMaterialLayersComputeShaderResource(PackLayersComponentsData, WeightmapLayerWeightBlend, WeightmapTextureOutputOffset);
-					CSPackLayersShaderParams.ComponentSize = ComponentSize;
-					BeginInitResource(CSPackLayersShaderParams.ComputeShaderResource);
-
-					FLandscapeLayerWeightmapPackMaterialLayersCSDispatch_RenderThread CSDispatch(CSPackLayersShaderParams);
-
-					ENQUEUE_RENDER_COMMAND(LandscapeLayers_Cmd_PackLayers)(
-						[CSDispatch](FRHICommandListImmediate& RHICmdList) mutable
-					{
-						TRACE_CPUPROFILER_EVENT_SCOPE(LandscapeLayers_RT_PackLayers);
-						CSDispatch.PackLayers(RHICmdList);
-					});
-
-					UTextureRenderTarget2D* CurrentRT = WeightmapRTList[(int32)EWeightmapRTType::WeightmapRT_Mip0];
-
-					SourceDebugName = TEXT("WeightmapScratchTexture");
-					DestDebugName = CurrentRT->GetName();
-
-					ENQUEUE_RENDER_COMMAND(LandscapeLayers_TransitionCopyResources)(
-						[this, CurrentRT](FRHICommandListImmediate& RHICmdList) mutable
-					{
-						RHICmdList.Transition(FRHITransitionInfo(CombinedLayersWeightmapAllMaterialLayersResource->TextureRHI, ERHIAccess::SRVMask, ERHIAccess::UAVMask));
-						RHICmdList.Transition(FRHITransitionInfo(WeightmapScratchPackLayerTextureResource->TextureRHI, ERHIAccess::UAVMask, ERHIAccess::SRVMask));
-						RHICmdList.Transition(FRHITransitionInfo(CurrentRT->GetRenderTargetResource()->TextureRHI, ERHIAccess::SRVMask, ERHIAccess::RTV));
-					});
-
-					CopyTexturePS(SourceDebugName, WeightmapScratchPackLayerTextureResource, DestDebugName, CurrentRT->GameThread_GetRenderTargetResource());
-
-					ENQUEUE_RENDER_COMMAND(LandscapeLayers_TransitionMip0)(
-						[this, CurrentRT](FRHICommandListImmediate& RHICmdList) mutable
-					{
-						RHICmdList.Transition(FRHITransitionInfo(CurrentRT->GetRenderTargetResource()->TextureRHI, ERHIAccess::RTV, ERHIAccess::SRVMask));
-						RHICmdList.Transition(FRHITransitionInfo(WeightmapScratchPackLayerTextureResource->TextureRHI, ERHIAccess::SRVMask, ERHIAccess::RTV));
-					});
-					DrawWeightmapComponentToRenderTargetMips(WeightmapTextureOutputOffset, CurrentRT, true, PSShaderParams);
-
-					int32 StartTextureIndex = NextTextureIndexToProcess;
-
-					TArray<FLandscapeLayersCopyTextureParams> DeferredCopyTextures;
-					for (; NextTextureIndexToProcess < ProcessedWeightmaps.Num(); ++NextTextureIndexToProcess)
-					{
-						UTexture2D* WeightmapTexture = ProcessedWeightmaps[NextTextureIndexToProcess];
-						if (!InUpdateLayersContentContext.WeightmapsToResolve.Contains(WeightmapTexture))
-						{
-							continue;
-						}
-
-						FTextureResource* WeightmapTextureResource = WeightmapTexture ? WeightmapTexture->GetResource() : nullptr;
-						if (WeightmapTextureResource == nullptr)
-						{
-							continue;
-						}
-
-						const int32 TextureSizeX = WeightmapTextureResource->GetSizeX();
-						const int32 TextureSizeY = WeightmapTextureResource->GetSizeY();
-
-						FIntPoint TextureTopLeftPositionInAtlas(
-							static_cast<int32>(WeightmapTextureOutputOffset[NextTextureIndexToProcess - StartTextureIndex].X),
-							static_cast<int32>(WeightmapTextureOutputOffset[NextTextureIndexToProcess - StartTextureIndex].Y));
-
-						int32 CurrentMip = 0;
-
-						for (int32 MipRTIndex = (int32)EWeightmapRTType::WeightmapRT_Mip0; MipRTIndex < (int32)EWeightmapRTType::WeightmapRT_Count; ++MipRTIndex)
-						{
-							CurrentRT = WeightmapRTList[MipRTIndex];
-
-							if (CurrentRT != nullptr)
-							{
-								SourceDebugName = CurrentRT->GetName();
-								DestDebugName = FString::Printf(TEXT("Weightmap Mip: %d"), CurrentMip);
-
-								FLandscapeLayersCopyTextureParams& CopyTextureParams = DeferredCopyTextures.Add_GetRef(FLandscapeLayersCopyTextureParams(SourceDebugName, CurrentRT->GameThread_GetRenderTargetResource(), DestDebugName, WeightmapTexture->GetResource()));
-								// Only copy the size that's actually needed : 
-								CopyTextureParams.CopySize.X = TextureSizeX >> CurrentMip;
-								CopyTextureParams.CopySize.Y = TextureSizeY >> CurrentMip;
-								// Copy from the composited texture's position to the top-left corner of the heightmap
-								CopyTextureParams.SourcePosition.X = TextureTopLeftPositionInAtlas.X >> CurrentMip;
-								CopyTextureParams.SourcePosition.Y = TextureTopLeftPositionInAtlas.Y >> CurrentMip;
-								CopyTextureParams.DestMip = static_cast<uint8>(CurrentMip);
-								++CurrentMip;
-							}
-						}
-
-						check(!TexturesNeedingReadback.Contains(WeightmapTexture));
-						TexturesNeedingReadback.Add(WeightmapTexture);
-					}
-
-					ExecuteCopyLayersTexture(MoveTemp(DeferredCopyTextures));
-				}
-
-				++CurrentWeightmapToProcessIndex;
-			}
-		}
-
-		// Prepare the UTexture2D readbacks we'll need to perform :
-		TArray<FLandscapeLayersCopyReadbackTextureParams> DeferredCopyReadbackTextures = PrepareLandscapeLayersCopyReadbackTextureParams(InUpdateLayersContentContext.MapHelper, TexturesNeedingReadback, /*bWeightmaps = */true);
-		ExecuteCopyToReadbackTexture(DeferredCopyReadbackTextures);
-
-		UpdateLayersMaterialInstances(InUpdateLayersContentContext.LandscapeComponentsWeightmapsToResolve);
-	}
-
-	return InMergeParams.WeightmapUpdateModes;
-}
-
 int32 ALandscape::PerformLayersWeightmapsBatchedMerge(FUpdateLayersContentContext& InUpdateLayersContentContext, const FEditLayersWeightmapMergeParams& InMergeParams)
 {
 	using namespace UE::Landscape;
@@ -11219,7 +6305,7 @@
 					const FWeightmapResolveInfo& WeightmapResolveInfo = It.Value;
 					FRDGTextureRef DestinationTextureRef = GraphBuilder.RegisterExternalTexture(CreateRenderTarget(WeightmapResolveInfo.TextureResource->TextureRHI, TEXT("DestinationTexture")));
 
-					RDG_EVENT_SCOPE(GraphBuilder, "%sFinalize %s", WeightmapResolveInfo.bIsFinalResolve ? TEXT("") : TEXT("(partially) "), *WeightmapResolveInfo.TextureResourceDebugName);
+					RDG_EVENT_SCOPE(GraphBuilder, "%sFinalize %s", WeightmapResolveInfo.bIsFinalResolve ? TEXT("") : TEXT("(partially)"), *WeightmapResolveInfo.TextureResourceDebugName);
 					{
 						RDG_EVENT_SCOPE(GraphBuilder, "Pack %d channels %s", FMath::CountBits(WeightmapResolveInfo.ChannelMask), WeightmapResolveInfo.bIsAdditiveResolve ? TEXT("(additive)") : TEXT(""));
 
@@ -11514,23 +6600,7 @@
 		MergeParams.bForceRender = bForceRender;
 		MergeParams.bSkipBrush = bSkipBrush;
 
-		switch (CurrentEditLayersMergeMode)
-		{
-		case ELandscapeEditLayersMergeMode::GlobalMerge:
-		{
-			return PerformLayersWeightmapsGlobalMerge(InUpdateLayersContentContext, MergeParams);
-		}
-		case ELandscapeEditLayersMergeMode::LocalMerge:
-		{
-			return PerformLayersWeightmapsLocalMerge(InUpdateLayersContentContext, MergeParams);
-		}
-		case ELandscapeEditLayersMergeMode::BatchedMerge:
-		{
-			return PerformLayersWeightmapsBatchedMerge(InUpdateLayersContentContext, MergeParams);
-		}
-		default:
-			check(false);
-		}
+		return PerformLayersWeightmapsBatchedMerge(InUpdateLayersContentContext, MergeParams);
 	}
 
 	return 0;
@@ -12198,8 +7268,6 @@
 	}
 	
 	// Expecting to be already initialized for BatchMerge.
-	check(CurrentEditLayersMergeMode == ELandscapeEditLayersMergeMode::BatchedMerge);
-	check(GetEditLayersMergeMode() == ELandscapeEditLayersMergeMode::BatchedMerge);
 	check(bLandscapeLayersAreInitialized);
 
 	if (!CanUpdateLayersContent())
@@ -12259,8 +7327,6 @@
 	}
 
 	// Expecting to be already initialized for BatchMerge.
-	check(CurrentEditLayersMergeMode == ELandscapeEditLayersMergeMode::BatchedMerge);
-	check(GetEditLayersMergeMode() == ELandscapeEditLayersMergeMode::BatchedMerge);
 	check(bLandscapeLayersAreInitialized);
 
 	if (!CanUpdateLayersContent())
@@ -12364,14 +7430,6 @@
 	// Make sure Update doesn't dirty Landscape packages when not in Landscape Ed Mode
 	FLandscapeDirtyOnlyInModeScope DirtyOnlyInMode(LandscapeInfo);
 
-	// If we went from local merge to global merge or vice versa, we need to reinitialize layers : 
-	ELandscapeEditLayersMergeMode EditLayersMergeMode = GetEditLayersMergeMode();
-	if (CurrentEditLayersMergeMode != EditLayersMergeMode)
-	{
-		RequestLayersInitialization(/*bInRequestContentUpdate = */true, /*bForceLayerResourceReset = */true);
-		CurrentEditLayersMergeMode = EditLayersMergeMode;
-	}
-
 	if (!bLandscapeLayersAreInitialized)
 	{
 		InitializeLayers();
@@ -12727,12 +7785,8 @@
 
 void ALandscape::InitializeLayers()
 {
-	if (CreateLayersRenderingResource(bLandscapeLayersForceResourceReset))
-	{
-		InitializeLandscapeLayersWeightmapUsage();
-		bLandscapeLayersAreInitialized = true;
-		bLandscapeLayersForceResourceReset = false;
-	}
+	InitializeLandscapeLayersWeightmapUsage();
+	bLandscapeLayersAreInitialized = true;
 }
 
 void ALandscape::OnPreSave()
@@ -12805,38 +7859,9 @@
 
 #endif
 
-void ALandscapeProxy::BeginDestroy()
-{
-	// At GC time, it is important not to rely on other objects as there's no guarantee about destruction order so we mustn't 
-	//  do anything that resolves the soft object ptr LandscapeActorRef here, in particular
-
-#if WITH_EDITORONLY_DATA
-	// Prevent destruction until all render resources are done destroying
-	ReleaseResourceFence.BeginFence();
-#endif
-
-	Super::BeginDestroy();
-}
-
-bool ALandscapeProxy::IsReadyForFinishDestroy()
-{
-	bool bReadyForFinishDestroy = Super::IsReadyForFinishDestroy();
-
-#if WITH_EDITORONLY_DATA
-	if (bReadyForFinishDestroy)
-	{
-		bReadyForFinishDestroy = ReleaseResourceFence.IsFenceComplete();
-	}
-#endif
-
-	return bReadyForFinishDestroy;
-}
-
 void ALandscapeProxy::FinishDestroy()
 {
 #if WITH_EDITORONLY_DATA
-	check(ReleaseResourceFence.IsFenceComplete());
-
 	for (auto& ItPair : HeightmapsCPUReadback)
 	{
 		FLandscapeEditLayerReadback* HeightmapCPUReadback = ItPair.Value;
@@ -13118,54 +8143,6 @@
 
 #endif
 
-void ALandscape::BeginDestroy()
-{
-#if WITH_EDITOR
-	if (CombinedLayersWeightmapAllMaterialLayersResource != nullptr)
-	{
-		BeginReleaseResource(CombinedLayersWeightmapAllMaterialLayersResource);
-	}
-
-	if (CurrentLayersWeightmapAllMaterialLayersResource != nullptr)
-	{
-		BeginReleaseResource(CurrentLayersWeightmapAllMaterialLayersResource);
-	}
-
-	if (WeightmapScratchExtractLayerTextureResource != nullptr)
-	{
-		BeginReleaseResource(WeightmapScratchExtractLayerTextureResource);
-	}
-
-	if (WeightmapScratchPackLayerTextureResource != nullptr)
-	{
-		BeginReleaseResource(WeightmapScratchPackLayerTextureResource);
-	}
-
-	// Use ResourceFence from base class		
-#endif
-
-	Super::BeginDestroy();
-}
-
-void ALandscape::FinishDestroy()
-{
-#if WITH_EDITORONLY_DATA
-	check(ReleaseResourceFence.IsFenceComplete());
-
-	delete CombinedLayersWeightmapAllMaterialLayersResource;
-	delete CurrentLayersWeightmapAllMaterialLayersResource;
-	delete WeightmapScratchExtractLayerTextureResource;
-	delete WeightmapScratchPackLayerTextureResource;
-
-	CombinedLayersWeightmapAllMaterialLayersResource = nullptr;
-	CurrentLayersWeightmapAllMaterialLayersResource = nullptr;
-	WeightmapScratchExtractLayerTextureResource = nullptr;
-	WeightmapScratchPackLayerTextureResource = nullptr;
-#endif
-
-	Super::FinishDestroy();
-}
-
 bool ALandscape::IsUpToDate() const
 {
 	if (!FApp::CanEverRender())
@@ -13284,13 +8261,13 @@
 {
 	if (LandscapeEditLayers.IsValidIndex(InLayerIndex))
 	{
-		return LandscapeEditLayers[InLayerIndex].EditLayer;
+			return LandscapeEditLayers[InLayerIndex].EditLayer;
 	}
 	return nullptr;
 }
 
 TArrayView<const FLandscapeLayer> ALandscape::GetLayersConst() const
-{ 
+{
 	return MakeArrayView(LandscapeEditLayers); 
 }
 
@@ -13705,10 +8682,10 @@
 
 	// Call Request Update on all components...
 	GetLandscapeInfo()->ForAllLandscapeComponents([](ULandscapeComponent* LandscapeComponent)
-		{
-			LandscapeComponent->RequestWeightmapUpdate(false, false);
-			LandscapeComponent->RequestHeightmapUpdate(false, false);
-		});
+	{
+		LandscapeComponent->RequestWeightmapUpdate(false, false);
+		LandscapeComponent->RequestHeightmapUpdate(false, false);
+	});
 
 	const bool bLocalIntermediateRender = true;
 	ForceUpdateLayersContent(bLocalIntermediateRender);
@@ -13721,10 +8698,10 @@
 		TSet<UTexture2D*> ProcessedHeightmaps;
 		FScopedSetLandscapeEditingLayer ScopeEditingLayer(this, LandscapeEditLayers[InLayerIndex - 1].EditLayer->GetGuid());
 		GetLandscapeInfo()->ForAllLandscapeComponents([&](ULandscapeComponent* LandscapeComponent)
-			{
-				SlowTask.EnterProgressFrame(1.f);
-				LandscapeComponent->CopyFinalLayerIntoEditingLayer(DataInterface, ProcessedHeightmaps);
-			});
+		{
+			SlowTask.EnterProgressFrame(1.f);
+			LandscapeComponent->CopyFinalLayerIntoEditingLayer(DataInterface, ProcessedHeightmaps);
+		});
 	}
 
 	TArray<ALandscapeBlueprintBrushBase*> BrushesToMove;
@@ -13810,7 +8787,7 @@
 
 // Deprecated
 void ALandscape::ClearLayer(const FGuid& InLayerGuid, TSet<TObjectPtr<ULandscapeComponent>>* InComponents, ELandscapeClearMode InClearMode, bool bMarkPackageDirty)
-{
+	{
 	ClearEditLayer(InLayerGuid, InComponents, InClearMode == ELandscapeClearMode::Clear_Weightmap ? ELandscapeToolTargetTypeFlags::Weightmap : ELandscapeToolTargetTypeFlags::Heightmap, bMarkPackageDirty);
 }
 PRAGMA_ENABLE_DEPRECATION_WARNINGS
@@ -13906,10 +8883,10 @@
 			{
 				if (LayerSettings.LayerInfoObj != ALandscapeProxy::VisibilityLayer)
 				{
-					Component->DeleteLayer(LayerSettings.LayerInfoObj, LandscapeEdit);
-				}
+				Component->DeleteLayer(LayerSettings.LayerInfoObj, LandscapeEdit);
 			}
 		}
+	}
 
 		if (EnumHasAllFlags(InClearMode, ELandscapeToolTargetTypeFlags::Visibility))
 		{
@@ -13932,29 +8909,29 @@
 	if (VisibleEditLayer)
 	{
 		for (ULandscapeEditLayerBase* EditLayer : GetEditLayers())
-		{
+			{
 			bool bDesiredVisible = (EditLayer == VisibleEditLayer);
 			if (EditLayer->IsVisible() != bDesiredVisible)
-			{
+				{
 				EditLayer->SetVisible(bDesiredVisible, /* bInModify = */ true);
-			}
 		}
 	}
 }
+}
 
 void ALandscape::ShowAllLayers()
 {
 	if (LandscapeEditLayers.Num() > 0)
 	{
 		for (ULandscapeEditLayerBase* EditLayer : GetEditLayers())
-		{
+			{
 			if (EditLayer->IsVisible() != true)
-			{
+				{
 				EditLayer->SetVisible(/*bInVisible = */ true, /*bInModify = */ true);
-			}
-		}
 	}
 }
+}
+}
 
 LANDSCAPE_API extern bool GDisableUpdateLandscapeMaterialInstances;
 
@@ -14581,25 +9558,6 @@
 #endif
 }
 
-UTextureRenderTarget2D* FLandscapeLayerBrush::RenderLayer(const FLandscapeBrushParameters& InParameters)
-{
-#if WITH_EDITORONLY_DATA
-	TRACE_CPUPROFILER_EVENT_SCOPE(LandscapeLayers_LayerBrushRender);
-	if ((BlueprintBrush == nullptr)
-		|| ((InParameters.LayerType == ELandscapeToolTargetType::Heightmap) && !AffectsHeightmap())
-		|| ((InParameters.LayerType == ELandscapeToolTargetType::Weightmap) && !AffectsWeightmapLayer(InParameters.WeightmapLayerName))
-		|| ((InParameters.LayerType == ELandscapeToolTargetType::Visibility) && !AffectsVisibilityLayer()))
-	{
-		return nullptr;
-	}
-
-	return BlueprintBrush->Execute(InParameters);
-#else  // WITH_EDITORONLY_DATA
-	return nullptr;
-#endif // !WITH_EDITORONLY_DATA
-
-}
-
 void LandscapeDumpSelectiveLayerRender(const TArray<FString>& Args)
 {
 	// Example usages:
@@ -14925,6 +9883,33 @@
 	}
 }
 
+void LandscapeForceLayersFullUpdate(const TArray<FString>& Args)
+{
+	const TArray<FString>& StringFilters = Args;
 
+	UWorld* World = GEditor->GetEditorWorldContext().World();
+	if (!World)
+	{
+		return;
+	}
+
+	bool bAtLeastOneLandscapeUpdated = false;
+	for (ALandscape* Landscape : TActorRange<ALandscape>(World))
+	{
+		const FString ActorName = Landscape->GetActorNameOrLabel();
+		if (StringFilters.IsEmpty() || StringFilters.FindByPredicate([ActorName](const FString& InStringFilter){ return ActorName.Find(InStringFilter) != INDEX_NONE; }))
+		{
+			UE_LOG(LogLandscape, Display, TEXT("LandscapeForceLayersFullUpdate : %s"), *ActorName);
+			Landscape->ForceLayersFullUpdate();
+			bAtLeastOneLandscapeUpdated = true;
+		}
+	}
+
+	if (!bAtLeastOneLandscapeUpdated)
+	{
+		UE_LOG(LogLandscape, Display, TEXT("LandscapeForceLayersFullUpdate : no landscape updated"));
+	}
+}
+
 #endif // WITH_EDITOR
 
