Landscape : removed global and local merge code and deprecated anything that was only available for this.

diff --git a/Engine/Source/Runtime/Landscape/Public/GlobalMergeLegacySupportUtil.h b/Engine/Source/Runtime/Landscape/Public/GlobalMergeLegacySupportUtil.h
--- a/Engine/Source/Runtime/Landscape/Public/GlobalMergeLegacySupportUtil.h
+++ b/Engine/Source/Runtime/Landscape/Public/GlobalMergeLegacySupportUtil.h
@@ -13,9 +13,12 @@
 class UTextureRenderTarget2D;
 struct FLandscapeBrushParameters;
 
+class UE_DEPRECATED(5.7, "Global merge has been deprecated so this is no longer used. There's no good way of deprecating UInterfaces ATM, hence this forward-declaration") ULandscapeBrushRenderCallAdapter_GlobalMergeLegacySupport {};
+class UE_DEPRECATED(5.7, "Global merge has been deprecated so this is no longer used. There's no good way of deprecating UInterfaces ATM, hence this forward-declaration") ILandscapeBrushRenderCallAdapter_GlobalMergeLegacySupport {};
+
 // UObject for the UInterface
 UINTERFACE(MinimalAPI)
-class ULandscapeBrushRenderCallAdapter_GlobalMergeLegacySupport : public UInterface
+class ULandscapeBrushRenderCallAdapter_GlobalMergeLegacySupport_DEPRECATED : public UInterface
 {
 	GENERATED_BODY()
 };
@@ -27,7 +30,8 @@ class ULandscapeBrushRenderCallAdapter_GlobalMergeLegacySupport : public UInterf
  * 
  * This interface is likely to be removed once global merge is no longer used.
  */
-class ILandscapeBrushRenderCallAdapter_GlobalMergeLegacySupport
+ 
+class ILandscapeBrushRenderCallAdapter_GlobalMergeLegacySupport_DEPRECATED
 {
 	GENERATED_BODY()
 
@@ -38,7 +42,7 @@ public:
 	 * is called. This function calls the virtual InitializeAsBlueprintBrush and RenderLayerAsBlueprintBrush
 	 * functions, which are equivalent to a blueprint brush actor's Initialize_Native and RenderLayer_Native.
 	 */
-	UTextureRenderTarget2D* RenderAsBlueprintBrush(const FLandscapeBrushParameters& InParameters, const FTransform& LandscapeTransform);
+	UTextureRenderTarget2D* RenderAsBlueprintBrush(const FLandscapeBrushParameters& InParameters, const FTransform& LandscapeTransform) { return nullptr; }
 
 	// Called in the same places as the equivalent methods on FLandscapeLayerBrush.
 	virtual bool AffectsHeightmapAsBlueprintBrush() const { return false; }
@@ -54,7 +58,7 @@ protected:
 
 	// Overridable, equivalent to RenderLayer_Native on a blueprint brush actor.
 	virtual UTextureRenderTarget2D* RenderLayerAsBlueprintBrush(const FLandscapeBrushParameters& InParameters)
-		PURE_VIRTUAL(ILandscapeBrushRenderCallAdapter_GlobalMergeLegacySupport::InitializeAsBlueprintBrush, return nullptr; );
+		PURE_VIRTUAL(ILandscapeBrushRenderCallAdapter_GlobalMergeLegacySupport_DEPRECATED::InitializeAsBlueprintBrush, return nullptr; );
 
 	// Used in RenderAsBlueprintBrush to cache things the same way that blueprint brushes do
 	FTransform CurrentRenderAreaWorldTransform;
