Fix crash in linux and mac when deleting landscape spline.

diff --git a/Engine/Source/Runtime/Landscape/Classes/LandscapeSplineControlPoint.h b/Engine/Source/Runtime/Landscape/Classes/LandscapeSplineControlPoint.h
--- a/Engine/Source/Runtime/Landscape/Classes/LandscapeSplineControlPoint.h
+++ b/Engine/Source/Runtime/Landscape/Classes/LandscapeSplineControlPoint.h
@@ -287,6 +287,12 @@ public:
 	//~ End UObject Interface
 #endif // WITH_EDITOR
 
+	/** Use GetOuterSafe instead of GetOuterULandscapeSplinesComponent if there is any danger that this point has been deleted (via DeleteSplinePoints).
+	  * This will return nullptr in that case, because the outer is no longer a ULandscapeSplineComponent */
+	ULandscapeSplinesComponent* GetOuterSafe() const;
+	/** Use GetOuterChecked when the call is always expected to succeed.  Will assert if it fails. */
+	ULandscapeSplinesComponent* GetOuterChecked() const;
+
 	friend class FLandscapeToolSplines;
 	friend class ULandscapeInfo;
 };
diff --git a/Engine/Source/Runtime/Landscape/Classes/LandscapeSplineSegment.h b/Engine/Source/Runtime/Landscape/Classes/LandscapeSplineSegment.h
--- a/Engine/Source/Runtime/Landscape/Classes/LandscapeSplineSegment.h
+++ b/Engine/Source/Runtime/Landscape/Classes/LandscapeSplineSegment.h
@@ -16,6 +16,7 @@ class UStaticMesh;
 
 //Forward declarations
 class ULandscapeSplineControlPoint;
+class ULandscapeSplinesComponent;
 
 USTRUCT()
 struct FLandscapeSplineInterpPoint
@@ -384,6 +385,12 @@ private:
 public:
 	//~ End UObject Interface
 
+	/** Use GetOuterSafe instead of GetOuterULandscapeSplinesComponent if there is any danger that this segment has been deleted (via DeleteSplinePoints).
+	  * This will return nullptr in that case, because the outer is no longer a ULandscapeSplineComponent */
+	ULandscapeSplinesComponent* GetOuterSafe() const;
+	/** Use GetOuterChecked when the call is always expected to succeed.  Will assert if it fails. */
+	ULandscapeSplinesComponent* GetOuterChecked() const;
+
 	friend class FLandscapeToolSplines;
 	friend class ULandscapeInfo;
 };
diff --git a/Engine/Source/Runtime/Landscape/Private/LandscapeSplines.cpp b/Engine/Source/Runtime/Landscape/Private/LandscapeSplines.cpp
--- a/Engine/Source/Runtime/Landscape/Private/LandscapeSplines.cpp
+++ b/Engine/Source/Runtime/Landscape/Private/LandscapeSplines.cpp
@@ -1963,7 +1963,8 @@ FLandscapeSplineSegmentConnection& FLandscapeSplineConnection::GetFarConnection(
 #if WITH_EDITOR
 bool ULandscapeSplineControlPoint::SupportsForeignSplineMesh() const
 {
-	return bPlaceSplineMeshesInStreamingLevels && GetOuterULandscapeSplinesComponent()->GetSplineOwner()->SupportsForeignSplineMesh();
+	ULandscapeSplinesComponent* SplineComponent = GetOuterSafe();
+	return bPlaceSplineMeshesInStreamingLevels && SplineComponent && SplineComponent->GetSplineOwner()->SupportsForeignSplineMesh();
 }
 
 FName ULandscapeSplineControlPoint::GetBestConnectionTo(FVector Destination) const
@@ -2041,7 +2042,11 @@ void ULandscapeSplineControlPoint::SetSplineSelected(bool bInSelected)
 		return;
 	}
 	bSelected = bInSelected;
-	GetOuterULandscapeSplinesComponent()->MarkRenderStateDirty();
+	ULandscapeSplinesComponent* SplineComponent = GetOuterSafe();
+	if (SplineComponent)
+	{
+		SplineComponent->MarkRenderStateDirty();
+	}
 
 	UpdateSplineMeshSelectionState();
 }
@@ -2179,18 +2184,21 @@ TMap<ULandscapeSplinesComponent*, UControlPointMeshComponent*> ULandscapeSplineC
 {
 	TMap<ULandscapeSplinesComponent*, UControlPointMeshComponent*> ForeignMeshComponentsMap;
 
-	ULandscapeSplinesComponent* OuterSplines = GetOuterULandscapeSplinesComponent();
-	TArray<ULandscapeSplinesComponent*> SplineComponents = OuterSplines->GetAllStreamingSplinesComponents();
-
-	for (ULandscapeSplinesComponent* SplineComponent : SplineComponents)
+	ULandscapeSplinesComponent* OuterSplines = GetOuterSafe();
+	if (OuterSplines)
 	{
-		if (SplineComponent != OuterSplines)
+		TArray<ULandscapeSplinesComponent*> SplineComponents = OuterSplines->GetAllStreamingSplinesComponents();
+
+		for (ULandscapeSplinesComponent* SplineComponent : SplineComponents)
 		{
-			auto* ForeignMeshComponent = SplineComponent->GetForeignMeshComponent(this);
-			if (ForeignMeshComponent)
+			if (SplineComponent != OuterSplines)
 			{
-				ForeignMeshComponent->Modify(false);
-				ForeignMeshComponentsMap.Add(SplineComponent, ForeignMeshComponent);
+				auto* ForeignMeshComponent = SplineComponent->GetForeignMeshComponent(this);
+				if (ForeignMeshComponent)
+				{
+					ForeignMeshComponent->Modify(false);
+					ForeignMeshComponentsMap.Add(SplineComponent, ForeignMeshComponent);
+				}
 			}
 		}
 	}
@@ -2202,7 +2210,7 @@ void ULandscapeSplineControlPoint::UpdateSplinePoints(bool bUpdateCollision, boo
 {
 	Modify();
 
-	ULandscapeSplinesComponent* OuterSplines = GetOuterULandscapeSplinesComponent();
+	ULandscapeSplinesComponent* OuterSplines = GetOuterChecked();
 
 	auto ForeignMeshComponentsMap = GetForeignMeshComponents();
 
@@ -2569,7 +2577,7 @@ void ULandscapeSplineControlPoint::DeleteSplinePoints()
 {
 	Modify();
 
-	ULandscapeSplinesComponent* OuterSplines = CastChecked<ULandscapeSplinesComponent>(GetOuter());
+	ULandscapeSplinesComponent* OuterSplines = GetOuterChecked();
 
 	Points.Reset();
 	Bounds = FBox(ForceInit);
@@ -2621,9 +2629,12 @@ void ULandscapeSplineControlPoint::PostEditUndo()
 	// Update mesh components to make sure the selection state stays in sync
 	UpdateSplineMeshSelectionState();
 
-	ULandscapeSplinesComponent* SplineComponent = GetOuterULandscapeSplinesComponent();
-	SplineComponent->MarkRenderStateDirty();
-	SplineComponent->RequestSplineLayerUpdate();
+	ULandscapeSplinesComponent* SplineComponent = GetOuterSafe();
+	if (SplineComponent)
+	{
+		SplineComponent->MarkRenderStateDirty();
+		SplineComponent->RequestSplineLayerUpdate();
+	}
 }
 
 void ULandscapeSplineControlPoint::PostDuplicate(bool bDuplicateForPIE)
@@ -2633,7 +2644,7 @@ void ULandscapeSplineControlPoint::PostDuplicate(bool bDuplicateForPIE)
 		// if we get duplicated but our local mesh doesn't, then clear our reference to the mesh - it's not ours
 		if (LocalMeshComponent != nullptr)
 		{
-			ULandscapeSplinesComponent* OuterSplines = CastChecked<ULandscapeSplinesComponent>(GetOuter());
+			ULandscapeSplinesComponent* OuterSplines = GetOuterChecked();
 			if (LocalMeshComponent->GetOuter() != OuterSplines->GetOwner())
 			{
 				LocalMeshComponent = nullptr;
@@ -2682,7 +2693,7 @@ void ULandscapeSplineControlPoint::PostEditChangeProperty(FPropertyChangedEvent&
 
 	if (PropertyChangedEvent.ChangeType == EPropertyChangeType::ValueSet)
 	{
-		GetOuterULandscapeSplinesComponent()->RequestSplineLayerUpdate();
+		GetOuterChecked()->RequestSplineLayerUpdate();
 	}
 }
 #endif // WITH_EDITOR
@@ -2913,7 +2924,7 @@ void ULandscapeSplineSegment::PostLoad()
 void ULandscapeSplineSegment::UpdateMeshCollisionProfile(USplineMeshComponent* MeshComponent)
 {
 #if WITH_EDITORONLY_DATA
-	ULandscapeSplinesComponent* OuterSplines = CastChecked<ULandscapeSplinesComponent>(GetOuter());
+	ULandscapeSplinesComponent* OuterSplines = GetOuterChecked();
 	const bool bUsingEditorMesh = OuterSplines->IsUsingEditorMesh(MeshComponent);
 
 	const FName DesiredCollisionProfileName = SplinesAlwaysUseBlockAll ? UCollisionProfile::BlockAll_ProfileName : GetCollisionProfileName();
@@ -2941,7 +2952,11 @@ void ULandscapeSplineSegment::SetSplineSelected(bool bInSelected)
 	}
 
 	bSelected = bInSelected;
-	GetOuterULandscapeSplinesComponent()->MarkRenderStateDirty();
+	ULandscapeSplinesComponent* SplineComponent = GetOuterSafe();
+	if (SplineComponent)
+	{
+		SplineComponent->MarkRenderStateDirty();
+	}
 
 	UpdateSplineMeshSelectionState();
 }
@@ -2995,21 +3010,24 @@ TMap<ULandscapeSplinesComponent*, TArray<USplineMeshComponent*>> ULandscapeSplin
 {
 	TMap<ULandscapeSplinesComponent*, TArray<USplineMeshComponent*>> ForeignMeshComponentsMap;
 
-	ULandscapeSplinesComponent* OuterSplines = GetOuterULandscapeSplinesComponent();
-	TArray<ULandscapeSplinesComponent*> SplineComponents = OuterSplines->GetAllStreamingSplinesComponents();
-
-	for (ULandscapeSplinesComponent* SplineComponent : SplineComponents)
+	ULandscapeSplinesComponent* OuterSplines = GetOuterSafe();
+	if (OuterSplines)
 	{
-		if (SplineComponent != OuterSplines)
+		TArray<ULandscapeSplinesComponent*> SplineComponents = OuterSplines->GetAllStreamingSplinesComponents();
+
+		for (ULandscapeSplinesComponent* SplineComponent : SplineComponents)
 		{
-			auto ForeignMeshComponents = SplineComponent->GetForeignMeshComponents(this);
-			if (ForeignMeshComponents.Num() > 0)
+			if (SplineComponent != OuterSplines)
 			{
-				for (auto* ForeignMeshComponent : ForeignMeshComponents)
+				auto ForeignMeshComponents = SplineComponent->GetForeignMeshComponents(this);
+				if (ForeignMeshComponents.Num() > 0)
 				{
-					ForeignMeshComponent->Modify(false);
+					for (auto* ForeignMeshComponent : ForeignMeshComponents)
+					{
+						ForeignMeshComponent->Modify(false);
+					}
+					ForeignMeshComponentsMap.Add(SplineComponent, MoveTemp(ForeignMeshComponents));
 				}
-				ForeignMeshComponentsMap.Add(SplineComponent, MoveTemp(ForeignMeshComponents));
 			}
 		}
 	}
@@ -3024,14 +3042,15 @@ TArray<USplineMeshComponent*> ULandscapeSplineSegment::GetLocalMeshComponents()
 
 bool ULandscapeSplineSegment::SupportsForeignSplineMesh() const
 {
-	return bPlaceSplineMeshesInStreamingLevels&& GetOuterULandscapeSplinesComponent()->GetSplineOwner()->SupportsForeignSplineMesh();
+	ULandscapeSplinesComponent* SplineComponent = GetOuterSafe();
+	return bPlaceSplineMeshesInStreamingLevels && SplineComponent && SplineComponent->GetSplineOwner()->SupportsForeignSplineMesh();
 }
 
 void ULandscapeSplineSegment::UpdateSplinePoints(bool bUpdateCollision, bool bUpdateMeshLevel)
 {
 	Modify();
 
-	ULandscapeSplinesComponent* OuterSplines = GetOuterULandscapeSplinesComponent();
+	ULandscapeSplinesComponent* OuterSplines = GetOuterChecked();
 
 	SplineInfo.Points.Empty(2);
 	Points.Reset();
@@ -3576,7 +3595,7 @@ void ULandscapeSplineSegment::UpdateSplineMeshSelectionState()
 
 void ULandscapeSplineSegment::UpdateSplineEditorMesh()
 {
-	ULandscapeSplinesComponent* OuterSplines = CastChecked<ULandscapeSplinesComponent>(GetOuter());
+	ULandscapeSplinesComponent* OuterSplines = GetOuterChecked();
 
 	for (auto& LocalMeshComponent : LocalMeshComponents)
 	{
@@ -3603,7 +3622,7 @@ void ULandscapeSplineSegment::DeleteSplinePoints()
 {
 	Modify();
 
-	ULandscapeSplinesComponent* OuterSplines = GetOuterULandscapeSplinesComponent();
+	ULandscapeSplinesComponent* OuterSplines = GetOuterChecked();
 
 	SplineInfo.Reset();
 	Points.Reset();
@@ -3687,9 +3706,12 @@ void ULandscapeSplineSegment::PostEditUndo()
 	// Update mesh components to make sure the selection state stays in sync
 	UpdateSplineMeshSelectionState();
 
-	ULandscapeSplinesComponent* SplineComponent = GetOuterULandscapeSplinesComponent();
-	SplineComponent->MarkRenderStateDirty();
-	SplineComponent->RequestSplineLayerUpdate();
+	ULandscapeSplinesComponent* SplineComponent = GetOuterSafe();
+	if (SplineComponent)
+	{
+		SplineComponent->MarkRenderStateDirty();
+		SplineComponent->RequestSplineLayerUpdate();
+	}
 }
 
 void ULandscapeSplineSegment::PostDuplicate(bool bDuplicateForPIE)
@@ -3699,7 +3721,7 @@ void ULandscapeSplineSegment::PostDuplicate(bool bDuplicateForPIE)
 		// if we get duplicated but our local meshes don't, then clear our reference to the meshes - they're not ours
 		if (LocalMeshComponents.Num() > 0)
 		{
-			ULandscapeSplinesComponent* OuterSplines = CastChecked<ULandscapeSplinesComponent>(GetOuter());
+			ULandscapeSplinesComponent* OuterSplines = GetOuterChecked();
 
 			// we assume all meshes are duplicated or none are, to avoid testing every one
 			if (LocalMeshComponents[0]->GetOuter() != OuterSplines->GetOwner())
@@ -3758,10 +3780,32 @@ void ULandscapeSplineSegment::PostEditChangeProperty(FPropertyChangedEvent& Prop
 
 	if (PropertyChangedEvent.ChangeType == EPropertyChangeType::ValueSet)
 	{
-		GetOuterULandscapeSplinesComponent()->RequestSplineLayerUpdate();
+		GetOuterChecked()->RequestSplineLayerUpdate();
 	}
 }
 
+ULandscapeSplinesComponent* ULandscapeSplineSegment::GetOuterSafe() const
+{
+	return !HasAnyFlags(RF_ClassDefaultObject) ? Cast<ULandscapeSplinesComponent>(GetOuter()) : nullptr;
+}
+
+ULandscapeSplinesComponent* ULandscapeSplineSegment::GetOuterChecked() const
+{
+	check(!HasAnyFlags(RF_ClassDefaultObject));
+	return CastChecked<ULandscapeSplinesComponent>(GetOuter());
+}
+
+ULandscapeSplinesComponent* ULandscapeSplineControlPoint::GetOuterSafe() const
+{
+	return !HasAnyFlags(RF_ClassDefaultObject) ? Cast<ULandscapeSplinesComponent>(GetOuter()) : nullptr;
+}
+
+ULandscapeSplinesComponent* ULandscapeSplineControlPoint::GetOuterChecked() const
+{
+	check(!HasAnyFlags(RF_ClassDefaultObject));
+	return CastChecked<ULandscapeSplinesComponent>(GetOuter());
+}
+
 void ALandscapeProxy::CreateSplineComponent()
 {
 	CreateSplineComponent(FVector(1.0f) / GetRootComponent()->GetRelativeScale3D());
@@ -3935,7 +3979,7 @@ void ULandscapeInfo::MoveSpline(ULandscapeSplineControlPoint* InControlPoint, TS
 {
 	bool bUpdateBounds = false;
 
-	ULandscapeSplinesComponent* FromSplineComponent = InControlPoint->GetOuterULandscapeSplinesComponent();
+	ULandscapeSplinesComponent* FromSplineComponent = InControlPoint->GetOuterChecked();
 	TScriptInterface<ILandscapeSplineInterface> From(FromSplineComponent->GetOwner());
 	if (From == InNewOwner)
 	{
