Fix missing landscape spline meshes after undo of merger of two splines.

diff --git a/Engine/Source/Runtime/Landscape/Classes/LandscapeSplineControlPoint.h b/Engine/Source/Runtime/Landscape/Classes/LandscapeSplineControlPoint.h
--- a/Engine/Source/Runtime/Landscape/Classes/LandscapeSplineControlPoint.h
+++ b/Engine/Source/Runtime/Landscape/Classes/LandscapeSplineControlPoint.h
@@ -288,7 +288,7 @@ public:
 #endif // WITH_EDITOR
 
 	/** Use GetOuterSafe instead of GetOuterULandscapeSplinesComponent if there is any danger that this point has been deleted (via DeleteSplinePoints).
-	  * This will return nullptr in that case, because the outer is no longer a ULandscapeSplineComponent */
+	  * This will return nullptr in that case, because the outer is no longer a ULandscapeSplinesComponent */
 	ULandscapeSplinesComponent* GetOuterSafe() const;
 	/** Use GetOuterChecked when the call is always expected to succeed.  Will assert if it fails. */
 	ULandscapeSplinesComponent* GetOuterChecked() const;
diff --git a/Engine/Source/Runtime/Landscape/Classes/LandscapeSplineSegment.h b/Engine/Source/Runtime/Landscape/Classes/LandscapeSplineSegment.h
--- a/Engine/Source/Runtime/Landscape/Classes/LandscapeSplineSegment.h
+++ b/Engine/Source/Runtime/Landscape/Classes/LandscapeSplineSegment.h
@@ -386,7 +386,7 @@ public:
 	//~ End UObject Interface
 
 	/** Use GetOuterSafe instead of GetOuterULandscapeSplinesComponent if there is any danger that this segment has been deleted (via DeleteSplinePoints).
-	  * This will return nullptr in that case, because the outer is no longer a ULandscapeSplineComponent */
+	  * This will return nullptr in that case, because the outer is no longer a ULandscapeSplinesComponent */
 	ULandscapeSplinesComponent* GetOuterSafe() const;
 	/** Use GetOuterChecked when the call is always expected to succeed.  Will assert if it fails. */
 	ULandscapeSplinesComponent* GetOuterChecked() const;
diff --git a/Engine/Source/Runtime/Landscape/Private/LandscapeSplines.cpp b/Engine/Source/Runtime/Landscape/Private/LandscapeSplines.cpp
--- a/Engine/Source/Runtime/Landscape/Private/LandscapeSplines.cpp
+++ b/Engine/Source/Runtime/Landscape/Private/LandscapeSplines.cpp
@@ -3844,6 +3844,7 @@ void ULandscapeInfo::MoveControlPoint(ULandscapeSplineControlPoint* InControlPoi
 	check(InControlPoint->GetOuterULandscapeSplinesComponent() == FromSplineComponent);
 
 	ToSplineComponent->Modify();
+	InControlPoint->Modify();
 
 	const FTransform FromComponentTransform = FromSplineComponent->GetComponentTransform();	// (From-->World)
 	const FTransform ToComponentTransform = ToSplineComponent->GetComponentTransform();		// (To-->World)
@@ -3915,6 +3916,7 @@ void ULandscapeInfo::MoveSegment(ULandscapeSplineSegment* InSegment, TScriptInte
 	check(InSegment->GetOuterULandscapeSplinesComponent() == FromSplineComponent);
 
 	ToSplineComponent->Modify();
+	InSegment->Modify();
 		
 	// Delete all Mesh Components associated with the Segment. (Will get recreated in UpdateSplinePoints)
 	for (auto& MeshComponent : InSegment->LocalMeshComponents)
