Fix for crash in landscape texture streaming manager when textures are unloaded out from underneath the landscape.

diff --git a/Engine/Source/Runtime/Landscape/Private/LandscapeGrassMapsBuilder.cpp b/Engine/Source/Runtime/Landscape/Private/LandscapeGrassMapsBuilder.cpp
--- a/Engine/Source/Runtime/Landscape/Private/LandscapeGrassMapsBuilder.cpp
+++ b/Engine/Source/Runtime/Landscape/Private/LandscapeGrassMapsBuilder.cpp
@@ -1584,9 +1584,12 @@ void FLandscapeGrassMapsBuilder::PendingToStreaming(FComponentState& State)
 	}
 
 	// Request streaming for our textures
-	for (UTexture* Texture : State.TexturesToStream)
+	for (const TWeakObjectPtr<UTexture>& TexturePtr : State.TexturesToStream)
 	{
-		TextureStreamingManager.RequestTextureFullyStreamedIn(Texture, /* bWaitForStreaming= */ false);
+		if (UTexture* Texture = TexturePtr.Get())
+		{
+			TextureStreamingManager.RequestTextureFullyStreamedIn(Texture, /* bWaitForStreaming= */ false);
+		}
 	}
 
 	check(PendingCount > 0);
@@ -1739,19 +1742,30 @@ void FLandscapeGrassMapsBuilder::PopulateGrassDataFromReadback(FComponentState&
 
 void FLandscapeGrassMapsBuilder::RemoveTextureStreamingRequests(FComponentState& State)
 {
-	for (UTexture* Texture : State.TexturesToStream)
+	for (const TWeakObjectPtr<UTexture>& TexturePtr : State.TexturesToStream)
 	{
-		TextureStreamingManager.UnrequestTextureFullyStreamedIn(Texture);
+		if (UTexture* Texture = TexturePtr.Get())
+		{
+			TextureStreamingManager.UnrequestTextureFullyStreamedIn(Texture);
+		}
 	}
 	State.TexturesToStream.Empty();
 }
 
-bool FLandscapeGrassMapsBuilder::FComponentState::AreTexturesStreamedIn() const
+bool FLandscapeGrassMapsBuilder::FComponentState::AreTexturesStreamedIn()
 {
-	for (UTexture* Texture : TexturesToStream)
+	for (auto It = TexturesToStream.CreateIterator(); It; It++)
 	{
-		if (!FLandscapeTextureStreamingManager::IsTextureFullyStreamedIn(Texture))
-			return false;
+		if (UTexture* Texture = It->Get())
+		{
+			if (!FLandscapeTextureStreamingManager::IsTextureFullyStreamedIn(Texture))
+				return false;
+		}
+		else
+		{
+			// Texture was garbage collected unexpectedly, remove from our list
+			It.RemoveCurrentSwap();
+		}
 	}
 	return true;
 }
diff --git a/Engine/Source/Runtime/Landscape/Private/LandscapeGrassMapsBuilder.h b/Engine/Source/Runtime/Landscape/Private/LandscapeGrassMapsBuilder.h
--- a/Engine/Source/Runtime/Landscape/Private/LandscapeGrassMapsBuilder.h
+++ b/Engine/Source/Runtime/Landscape/Private/LandscapeGrassMapsBuilder.h
@@ -134,7 +134,7 @@ private:
 		#endif // WITH_EDITOR
 
 		// list of textures to stream prior to rendering this component (valid only in TextureSreaming and Rendering stages)
-		TArray<UTexture*> TexturesToStream;
+		TArray<TWeakObjectPtr<UTexture>> TexturesToStream;
 		
 		// the active render (valid only in Rendering stage)
 		TUniquePtr<FLandscapeGrassWeightExporter> ActiveRender;
@@ -144,7 +144,7 @@ private:
 
 		FComponentState(ULandscapeComponent* Component);
 
-		bool AreTexturesStreamedIn() const;
+		bool AreTexturesStreamedIn();
 		bool IsBeyondEvictionRange(const TArray<FVector>& Cameras) const;
 	};
 
