Removed landscape's infamous "intermediate render" code path : partial renders of edit layers is now used everywhere we need to render landscape with only a subset of edit layers enabled so this greatly simplifies the code

diff --git a/Engine/Source/Editor/LandscapeEditor/Private/LandscapeEdModeTools.h b/Engine/Source/Editor/LandscapeEditor/Private/LandscapeEdModeTools.h
--- a/Engine/Source/Editor/LandscapeEditor/Private/LandscapeEdModeTools.h
+++ b/Engine/Source/Editor/LandscapeEditor/Private/LandscapeEdModeTools.h
@@ -543,8 +543,7 @@ public:
 		return (bIsValid && X1 >= CachedX1 && Y1 >= CachedY1 && X2 <= CachedX2 && Y2 <= CachedY2);
 	}
 	
-	using FPrepareRegionForCachingFunction = TFunction<FIntRect(const FIntRect& NewCacheBounds)>;
-	bool GetDataAndCacheWithVisibilityOverride(int32 X1, int32 Y1, int32 X2, int32 Y2, TArray<AccessorType>& OutData, FPrepareRegionForCachingFunction PrepareRegionForCaching, TBitArray<> LayerVisibility, ALandscape* Landscape)
+	bool GetDataAndCacheWithVisibilityOverride(int32 X1, int32 Y1, int32 X2, int32 Y2, TArray<AccessorType>& OutData, TBitArray<> LayerVisibility, ALandscape* Landscape)
 	{
 		// Drop a visual log to indicate the area requested by this data access :
 		VisualizeLandscapeRegion(X1, Y1, X2, Y2, FColor::Blue, TEXT("CacheDataRequest"));
@@ -556,75 +555,58 @@ public:
 			FIntRect NewCacheBounds(bIsValid ? FMath::Min(X1, CachedX1) : X1, bIsValid ? FMath::Min(Y1, CachedY1) : Y1, bIsValid ? FMath::Max(X2, CachedX2) : X2, bIsValid ? FMath::Max(Y2, CachedY2) : Y2);
 
 			// SelectiveRender only supports heightmaps at the moment.
-			if (Landscape->AllowSelectiveRenderEditLayers() && std::is_same<AccessorType, uint16>::value)
+			const bool bIsHeightmap = std::is_same<AccessorType, uint16>::value;
+			check(bIsHeightmap);
+
+			// For now, do the same bounds expansion as OnCacheUpdating.  Align with component boundaries.
+			// Note: FloorToInt32 is "strictly less", not "towards zero".
+			// Note, because of the +1/-1 expansion at a higher level in the callstack, this will pad out an entire component outside the landscape.
+			// CacheDataDirect is currently ignoring this excessive padding and limiting to only 1 row/column extra.
+			// TODO:  simplify when removing IntermediateRender.
+			FIntRect DesiredCacheBounds;
+			DesiredCacheBounds.Min.X = (FMath::FloorToInt((float)NewCacheBounds.Min.X / LandscapeInfo->ComponentSizeQuads)) * LandscapeInfo->ComponentSizeQuads;
+			DesiredCacheBounds.Min.Y = (FMath::FloorToInt((float)NewCacheBounds.Min.Y / LandscapeInfo->ComponentSizeQuads)) * LandscapeInfo->ComponentSizeQuads;
+			DesiredCacheBounds.Max.X = (FMath::CeilToInt((float)NewCacheBounds.Max.X / LandscapeInfo->ComponentSizeQuads)) * LandscapeInfo->ComponentSizeQuads;
+			DesiredCacheBounds.Max.Y = (FMath::CeilToInt((float)NewCacheBounds.Max.Y / LandscapeInfo->ComponentSizeQuads)) * LandscapeInfo->ComponentSizeQuads;
+
+			FIntRect LandscapeExtent;
+			bool bRet = LandscapeInfo->GetLandscapeExtent(LandscapeExtent);
+			ensure(bRet);
+			if (!bRet)
 			{
-				// For now, do the same bounds expansion as OnCacheUpdating.  Align with component boundaries.
-				// Note: FloorToInt32 is "strictly less", not "towards zero".
-				// Note, because of the +1/-1 expansion at a higher level in the callstack, this will pad out an entire component outside the landscape.
-				// CacheDataDirect is currently ignoring this excessive padding and limiting to only 1 row/column extra.
-				// TODO:  simplify when removing IntermediateRender.
-				FIntRect DesiredCacheBounds;
-				DesiredCacheBounds.Min.X = (FMath::FloorToInt((float)NewCacheBounds.Min.X / LandscapeInfo->ComponentSizeQuads)) * LandscapeInfo->ComponentSizeQuads;
-				DesiredCacheBounds.Min.Y = (FMath::FloorToInt((float)NewCacheBounds.Min.Y / LandscapeInfo->ComponentSizeQuads)) * LandscapeInfo->ComponentSizeQuads;
-				DesiredCacheBounds.Max.X = (FMath::CeilToInt((float)NewCacheBounds.Max.X / LandscapeInfo->ComponentSizeQuads)) * LandscapeInfo->ComponentSizeQuads;
-				DesiredCacheBounds.Max.Y = (FMath::CeilToInt((float)NewCacheBounds.Max.Y / LandscapeInfo->ComponentSizeQuads)) * LandscapeInfo->ComponentSizeQuads;
-
-				FIntRect LandscapeExtent;
-				bool bRet = LandscapeInfo->GetLandscapeExtent(LandscapeExtent);
-				ensure(bRet);
-				if (!bRet)
-				{
-					// Failure likely means no proxies loaded
-					return false;
-				}
+				// Failure likely means no proxies loaded
+				return false;
+			}
 
-				// FLandscapeToolStrokeSmooth::Apply is extending bounds by 1 in all directions, outside landscape bounds, and asserts if the result bounds don't include that.
-				// landscape with 2017 verts/samples, DesiredCacheBounds is (-1,2017).  For now, don't change the return rect here.
-				FIntRect EvaluateRect = DesiredCacheBounds;
+			// FLandscapeToolStrokeSmooth::Apply is extending bounds by 1 in all directions, outside landscape bounds, and asserts if the result bounds don't include that.
+			// landscape with 2017 verts/samples, DesiredCacheBounds is (-1,2017).  For now, don't change the return rect here.
+			FIntRect EvaluateRect = DesiredCacheBounds;
 
-				// Limit to the extent of the loaded landscape.  The calling code will ask for data outside the landscape and the old GetData code would generate it to fill in, but that doesn't
-				// appear to have been useful.
-				EvaluateRect.Clip(LandscapeExtent);  
-				EvaluateRect.Max += FInt32Point(1, 1);  // Convert inclusive range to exclusive (half-open).  [a,b)
+			// Limit to the extent of the loaded landscape.  The calling code will ask for data outside the landscape and the old GetData code would generate it to fill in, but that doesn't
+			// appear to have been useful.
+			EvaluateRect.Clip(LandscapeExtent);  
+			EvaluateRect.Max += FInt32Point(1, 1);  // Convert inclusive range to exclusive (half-open).  [a,b)
 
 #if 0
-				UE_LOG(LogLandscapeTools, Warning, TEXT("PartialRender - New ---- (fr%u)  Base [%d, %d]-[%d, %d]  s[%d, %d],   Padded [%d, %d]-[%d, %d] s[%d, %d],   Evaluate [%d, %d]-[%d, %d] s[%d, %d]"), GFrameNumber,
-					NewCacheBounds.Min.X, NewCacheBounds.Min.Y, NewCacheBounds.Max.X, NewCacheBounds.Max.Y, NewCacheBounds.Size().X, NewCacheBounds.Size().Y,
-					DesiredCacheBounds.Min.X, DesiredCacheBounds.Min.Y, DesiredCacheBounds.Max.X, DesiredCacheBounds.Max.Y, DesiredCacheBounds.Size().X, DesiredCacheBounds.Size().Y,
-					EvaluateRect.Min.X, EvaluateRect.Min.Y, EvaluateRect.Max.X, EvaluateRect.Max.Y, EvaluateRect.Size().X, EvaluateRect.Size().Y);
+			UE_LOG(LogLandscapeTools, Warning, TEXT("PartialRender - New ---- (fr%u)  Base [%d, %d]-[%d, %d]  s[%d, %d],   Padded [%d, %d]-[%d, %d] s[%d, %d],   Evaluate [%d, %d]-[%d, %d] s[%d, %d]"), GFrameNumber,
+				NewCacheBounds.Min.X, NewCacheBounds.Min.Y, NewCacheBounds.Max.X, NewCacheBounds.Max.Y, NewCacheBounds.Size().X, NewCacheBounds.Size().Y,
+				DesiredCacheBounds.Min.X, DesiredCacheBounds.Min.Y, DesiredCacheBounds.Max.X, DesiredCacheBounds.Max.Y, DesiredCacheBounds.Size().X, DesiredCacheBounds.Size().Y,
+				EvaluateRect.Min.X, EvaluateRect.Min.Y, EvaluateRect.Max.X, EvaluateRect.Max.Y, EvaluateRect.Size().X, EvaluateRect.Size().Y);
 #endif
 
-				TArray<AccessorType> Results;
-				Results.AddUninitialized(EvaluateRect.Area());
-
-				bRet = AccessorClass::GetSelectiveLayerData(Landscape, EvaluateRect, LayerVisibility, Results);
-				ensure(bRet);  // what to do if this fails?
-
-				// Convert to exclusive rects
-				FIntRect DesiredCacheBoundsExcl = DesiredCacheBounds;
-				DesiredCacheBoundsExcl.Max += FInt32Point(1, 1);
-				FIntRect LandscapeExtentExcl = LandscapeExtent;
-				LandscapeExtentExcl.Max += FInt32Point(1, 1);
+			TArray<AccessorType> Results;
+			Results.AddUninitialized(EvaluateRect.Area());
 
-				CacheDataDirect(EvaluateRect, DesiredCacheBoundsExcl, LandscapeExtentExcl, Results);
-			}
-			else
-			{
-				// Legacy intermediate render path.
-				// TODO:  remove with IntermediateRender
-				// 
-				// The caller might request a cache region that is actually larger than the data they want to sample for this particular read (e.g. to avoid re-caching when doing expensive samples like layer-collapsing ones) : 
-				NewCacheBounds = PrepareRegionForCaching(NewCacheBounds);
-				check((NewCacheBounds.Min.X <= X1) && (NewCacheBounds.Min.Y <= Y1) && (NewCacheBounds.Max.X >= X2) && (NewCacheBounds.Max.Y >= Y2));
+			bRet = AccessorClass::GetSelectiveLayerData(Landscape, EvaluateRect, LayerVisibility, Results);
+			ensure(bRet);  // what to do if this fails?
 
-#if 0
-				//TODO:  remove before check in
-				UE_LOG(LogLandscapeTools, Warning, TEXT("PartialRender - Old ---- (fr%u)  Padded [%d, %d]-[%d, %d] s[%d, %d]"), GFrameNumber,
-					NewCacheBounds.Min.X, NewCacheBounds.Min.Y, NewCacheBounds.Max.X, NewCacheBounds.Max.Y, NewCacheBounds.Size().X, NewCacheBounds.Size().Y);
-#endif
+			// Convert to exclusive rects
+			FIntRect DesiredCacheBoundsExcl = DesiredCacheBounds;
+			DesiredCacheBoundsExcl.Max += FInt32Point(1, 1);
+			FIntRect LandscapeExtentExcl = LandscapeExtent;
+			LandscapeExtentExcl.Max += FInt32Point(1, 1);
 
-				CacheData(NewCacheBounds.Min.X, NewCacheBounds.Min.Y, NewCacheBounds.Max.X, NewCacheBounds.Max.Y);
-			}
+			CacheDataDirect(EvaluateRect, DesiredCacheBoundsExcl, LandscapeExtentExcl, Results);
 		}
 		ensureMsgf(!bIsValid || HasCachedData(X1, Y1, X2, Y2), TEXT("Data failed to cache (X1:%i, Y1:%i, X2:%i, Y2:%i), (CachedX1:%i, CachedY1:%i, CachedX2:%i, CachedY2:%i)"), X1, Y1, X2, Y2, CachedX1, CachedY1, CachedX2, CachedY2);
 		return GetCachedData(X1, Y1, X2, Y2, OutData);
@@ -830,7 +812,6 @@ public:
 		, EditingLayerIndex(MAX_uint8)
 		, bIsInitialized(false)
 		, bCombinedLayerOperation(false)
-		, bVisibilityChanged(false)
 		, bTargetIsHeightmap(InTarget.TargetType == ELandscapeToolTargetType::Heightmap)
 		, CacheUpToEditingLayer(Cache)
 		, CacheBottomLayers(InTarget)
@@ -884,47 +865,15 @@ public:
 				++I;
 			}
 
-			// TODO:  remove with IntermediateRender
-			auto OnCacheUpdating = [&](const FIntRect& NewCacheBounds) -> FIntRect
-			{
-				// This function is triggered when the cache needs expanding. We'll ask for a larger area so that we don't need to re-cache (a GPU-synchronizing operation) every time
-				//  we need an additional row/column. Aligning the cache region on components borders is appropriate since this is what gets rendered. This avoids
-				//  re-caching when we've already sampled a component: 
-				FIntRect DesiredCacheBounds;
-				DesiredCacheBounds.Min.X = (FMath::FloorToInt((float)NewCacheBounds.Min.X / LandscapeInfo->ComponentSizeQuads)) * LandscapeInfo->ComponentSizeQuads;
-				DesiredCacheBounds.Min.Y = (FMath::FloorToInt((float)NewCacheBounds.Min.Y / LandscapeInfo->ComponentSizeQuads)) * LandscapeInfo->ComponentSizeQuads;
-				DesiredCacheBounds.Max.X = (FMath::CeilToInt((float)NewCacheBounds.Max.X / LandscapeInfo->ComponentSizeQuads)) * LandscapeInfo->ComponentSizeQuads;
-				DesiredCacheBounds.Max.Y = (FMath::CeilToInt((float)NewCacheBounds.Max.Y / LandscapeInfo->ComponentSizeQuads)) * LandscapeInfo->ComponentSizeQuads;
-
-				TSet<ULandscapeComponent*> AffectedComponents;
-				LandscapeInfo->GetComponentsInRegion(DesiredCacheBounds.Min.X, DesiredCacheBounds.Min.Y, DesiredCacheBounds.Max.X, DesiredCacheBounds.Max.Y, AffectedComponents);
-
-				SynchronousUpdateComponentVisibilityForHeight(AffectedComponents, NewLayerVisibility);
-
-				if (bVisibilityChanged)
-				{
-					VisibilityChangedComponents.Append(AffectedComponents);
-				}
-
-				return DesiredCacheBounds;
-			};
-
 			// temporarily switch to working on the final runtime data, so we can gather the combined layer data into the caches
 			FGuid PreviousLayerGUID = EditingLayerGuid.GetValue();
 			SetCacheEditingLayer(FGuid());
 
-			CacheUpToEditingLayer.GetDataAndCacheWithVisibilityOverride(X1, Y1, X2, Y2, Data, OnCacheUpdating, NewLayerVisibility, Landscape);
-			// Release Texture Mips that will be Locked by the next SynchronousUpdateComponentVisibilityForHeight
-			// remove Flush with IntermediateRender?  Don't need to release any mapped textures but if there were also pending writes then this has an effect.
-			CacheUpToEditingLayer.DataAccess.Flush();
+			CacheUpToEditingLayer.GetDataAndCacheWithVisibilityOverride(X1, Y1, X2, Y2, Data, NewLayerVisibility, Landscape);
 
 			// Now turn off visibility on the current layer in order to have the data of all bottom layers except the current one
 			NewLayerVisibility[EditingLayerIndex] = false;
-			CacheBottomLayers.GetDataAndCacheWithVisibilityOverride(X1, Y1, X2, Y2, BottomLayersData, OnCacheUpdating, NewLayerVisibility, Landscape);
-			// Do the same here for consistency
-			// remove with IntermediateRender?  CacheBottomLayers never writes back to the landscape, so this one should be safe to remove.  CacheBottomLayers could be
-			// a smaller, simpler class if its only used for GetSelectiveLayerData.
-			CacheBottomLayers.DataAccess.Flush();
+			CacheBottomLayers.GetDataAndCacheWithVisibilityOverride(X1, Y1, X2, Y2, BottomLayersData, NewLayerVisibility, Landscape);
 			
 			SetCacheEditingLayer(PreviousLayerGUID);
 			check(PreviousLayerGUID == CacheUpToEditingLayer.DataAccess.GetEditLayer());
@@ -961,22 +910,10 @@ public:
  				Landscape->GetEditLayerConst(*EditingLayerGuid) ? *(Landscape->GetEditLayerConst(*EditingLayerGuid)->GetName().ToString()) : TEXT("<unknown>"), *EditingLayerGuid->ToString(),
  				Landscape->GetEditLayerConst(CacheAccessorLayerGuid) ? *(Landscape->GetEditLayerConst(CacheAccessorLayerGuid)->GetName().ToString()) : TEXT("<unknown>"), *CacheAccessorLayerGuid.ToString());
 
-			// Restore layers visibility
-			// remove with IntermediateRender.  nothing will change real layer visibility anymore.
-			SetLayersVisibility(BackupLayerVisibility);
-			// Only store data in cache
 			CacheUpToEditingLayer.SetCachedData(X1, Y1, X2, Y2, Data, PaintingRestriction, false);
 			// Effectively write the contribution
 			CacheUpToEditingLayer.DataAccess.SetData(X1, Y1, X2, Y2, DataContribution.GetData(), PaintingRestriction);
 			CacheUpToEditingLayer.DataAccess.Flush();
-			if (bVisibilityChanged)
-			{
-				const bool bUpdateCollision = true;
-				const bool bIntermediateRender = false;
-				SynchronousUpdateHeightmapForComponents(VisibilityChangedComponents, bUpdateCollision, bIntermediateRender);
-				VisibilityChangedComponents.Empty();
-				bVisibilityChanged = false;
-			}
 		}
 		else
 		{
@@ -986,45 +923,6 @@ public:
 	}
 
 private:
-
-	void SynchronousUpdateHeightmapForComponents(const TSet<ULandscapeComponent*>& InComponents, bool bUpdateCollision, bool bIntermediateRender)
-	{
-		for (ULandscapeComponent* Component : InComponents)
-		{
-			const bool bUpdateAll = false; // default value
-			Component->RequestHeightmapUpdate(bUpdateAll, bUpdateCollision);
-		}
-		Landscape->ForceUpdateLayersContent(bIntermediateRender);
-	};
-
-	void SetLayersVisibility(const TBitArray<>& InLayerVisibility)
-	{
-		check(InLayerVisibility.Num() == Landscape->GetEditLayersConst().Num());
-		for (int i = 0; i < InLayerVisibility.Num(); ++i)
-		{
-			if (ULandscapeEditLayerBase* EditLayer = Landscape->GetEditLayer(i))
-			{
-				if (EditLayer->IsVisible() != InLayerVisibility[i])
-				{
-					// Pass bInModify and bInBroadcastDataChange = false here, because we don't want to call Modify on the landscape, we don't actually want to truly change 
-					//  the layer visibility, only temporarily for the purpose of the intermediate render : 
-					EditLayer->SetVisible(InLayerVisibility[i], /*bInModify = */ false, /*bInBroadcastDataChange = */false);
-					bVisibilityChanged = true;
-				}
-			}
-		}
-	};
-
-	// remove with IntermediateRender
-	void SynchronousUpdateComponentVisibilityForHeight(const TSet<ULandscapeComponent*>& InComponents, const TBitArray<>& InLayerVisibility)
-	{
-		SetLayersVisibility(InLayerVisibility);
-        // No need to update collision here as we are only doing a intermediate render to gather heightdata
-        const bool bUpdateCollision = false;
-		const bool bIntermediateRender = true;
-		SynchronousUpdateHeightmapForComponents(InComponents, bUpdateCollision, bIntermediateRender);
-	};
-
 	ULandscapeInfo* LandscapeInfo;
 	ALandscape* Landscape;
 	TOptional<FGuid> EditingLayerGuid;
@@ -1033,9 +931,7 @@ private:
 	TArray<typename ToolTarget::CacheClass::DataType> BottomLayersData;
 	bool bIsInitialized;
 	bool bCombinedLayerOperation;
-	bool bVisibilityChanged;    // remove with IntermediateRender
 	bool bTargetIsHeightmap;
-	TSet<ULandscapeComponent*> VisibilityChangedComponents;  // remove with IntermediateRender.
 
 	typename ToolTarget::CacheClass& CacheUpToEditingLayer;
 	typename ToolTarget::CacheClass CacheBottomLayers;
diff --git a/Engine/Source/Runtime/Landscape/Classes/Landscape.h b/Engine/Source/Runtime/Landscape/Classes/Landscape.h
--- a/Engine/Source/Runtime/Landscape/Classes/Landscape.h
+++ b/Engine/Source/Runtime/Landscape/Classes/Landscape.h
@@ -494,7 +494,6 @@ PRAGMA_ENABLE_DEPRECATION_WARNINGS
 	
 	LANDSCAPE_API bool DeleteLayer(int32 InLayerIndex);
 	LANDSCAPE_API void CollapseLayer(int32 InLayerIndex);
-	void CollapseLayerIntermediateRender(int32 InLayerIndex);
 	UE_DEPRECATED(5.7, "Non-edit layer landscapes are deprecated. Use CollapseEditLayers to replace all existing layers with a new default layer.")
 	LANDSCAPE_API void DeleteLayers();
 	LANDSCAPE_API void CollapseEditLayers();
@@ -534,7 +533,9 @@ PRAGMA_ENABLE_DEPRECATION_WARNINGS
 	LANDSCAPE_API void ConvertNonEditLayerLandscape();
 	UE_DEPRECATED(5.7, "Non-edit layer landscapes are deprecated. Use ConvertNonEditLayerLandscape to convert to non-edit layer landscapes to edit layer based landscapes.")
 	LANDSCAPE_API void ToggleCanHaveLayersContent();
-	LANDSCAPE_API void ForceUpdateLayersContent(bool bIntermediateRender = false);
+	LANDSCAPE_API void ForceUpdateLayersContent();
+	UE_DEPRECATED(5.7, "Intermediate renders are a thing of the past. You should now use SelectiveRenderEditLayersHeightmaps/SelectiveRenderEditLayersWeightmaps to render the landscape edit layers in an \"intermediate\" state")
+	LANDSCAPE_API void ForceUpdateLayersContent(bool bIntermediateRender);
 
 	LANDSCAPE_API bool SelectiveRenderEditLayersHeightmaps(const FLandscapeEditLayerRenderHeightParams& InSelectiveRenderParams);
 	LANDSCAPE_API bool SelectiveRenderEditLayersWeightmaps(const FLandscapeEditLayerRenderWeightParams& InSelectiveRenderParams);
@@ -568,8 +569,6 @@ PRAGMA_ENABLE_DEPRECATION_WARNINGS
 
 	LANDSCAPE_API FName GenerateUniqueLayerName(FName InName = NAME_None) const;
 
-	UE_INTERNAL LANDSCAPE_API bool AllowSelectiveRenderEditLayers() const;
-
 	bool GetIsPendingEditLayerConversion() const;
 	void SetIsPendingEditLayerConversion(bool bInPendingEditLayerConversion);
 
@@ -578,7 +577,7 @@ private:
 	void OnLayerCreatedInternal(ULandscapeEditLayerBase* EditLayer);
 	ULandscapeEditLayerBase* GetEditLayerInternal(int32 InLayerIndex);
 
-	void UpdateLayersContent(bool bInWaitForStreaming = false, bool bInSkipMonitorLandscapeEdModeChanges = false, bool bIntermediateRender = false, bool bFlushRender = false);
+	void UpdateLayersContent(bool bInWaitForStreaming, bool bInSkipMonitorLandscapeEdModeChanges, bool bFlushRender);
 	bool CanUpdateLayersContent() const;
 	void MonitorShaderCompilation();
 	void MonitorLandscapeEdModeChanges();
@@ -592,7 +591,7 @@ private:
 
 	int32 RegenerateLayersHeightmaps(const FUpdateLayersContentContext& InUpdateLayersContentContext);
 	int32 PerformLayersHeightmapsBatchedMerge(const FUpdateLayersContentContext& InUpdateLayersContentContext, const FEditLayersHeightmapMergeParams& InMergeParams);
-	void ResolveLayersHeightmapTexture(const FTextureToComponentHelper& MapHelper, const TSet<UTexture2D*>& HeightmapsToResolve, bool bIntermediateRender, TArray<FLandscapeEditLayerComponentReadbackResult>& InOutComponentReadbackResults);
+	void ResolveLayersHeightmapTexture(const FTextureToComponentHelper& MapHelper, const TSet<UTexture2D*>& HeightmapsToResolve, TArray<FLandscapeEditLayerComponentReadbackResult>& InOutComponentReadbackResults);
 
 	bool PerformSelectiveRenderEditLayersHeightmapsCPU(const FUpdateLayersContentContext& InUpdateLayersContentContext, const FLandscapeEditLayerRenderHeightParams& InSelectiveRenderParams);
 	bool PerformSelectiveRenderEditLayersHeightmapsRT(const FUpdateLayersContentContext& InUpdateLayersContentContext, const FLandscapeEditLayerRenderHeightParams& InSelectiveRenderParams);
@@ -600,9 +599,9 @@ private:
 	int32 RegenerateLayersWeightmaps(FUpdateLayersContentContext& InUpdateLayersContentContext);
 	int32 PerformLayersWeightmapsBatchedMerge(FUpdateLayersContentContext& InUpdateLayersContentContext, const FEditLayersWeightmapMergeParams& InMergeParams);
 	bool PerformSelectiveRenderEditLayersWeightmapsCPU(FUpdateLayersContentContext& InUpdateLayersContentContext, const FLandscapeEditLayerRenderWeightParams& InSelectiveRenderParams);
-	void ResolveLayersWeightmapTexture(const FTextureToComponentHelper& MapHelper, const TSet<UTexture2D*>& WeightmapsToResolve, bool bIntermediateRender, TArray<FLandscapeEditLayerComponentReadbackResult>& InOutComponentReadbackResults);
+	void ResolveLayersWeightmapTexture(const FTextureToComponentHelper& MapHelper, const TSet<UTexture2D*>& WeightmapsToResolve, TArray<FLandscapeEditLayerComponentReadbackResult>& InOutComponentReadbackResults);
 
-	bool ResolveLayersTexture(FTextureToComponentHelper const& MapHelper, FLandscapeEditLayerReadback* InCPUReadBack, UTexture2D* InOutputTexture, bool bIntermediateRender,
+	bool ResolveLayersTexture(FTextureToComponentHelper const& MapHelper, FLandscapeEditLayerReadback* InCPUReadBack, UTexture2D* InOutputTexture,
 		TArray<FLandscapeEditLayerComponentReadbackResult>& InOutComponentReadbackResults, bool bIsWeightmap);
 
 	static bool IsUpdateFlagEnabledForModes(ELandscapeComponentUpdateFlag InFlag, uint32 InUpdateModes);
diff --git a/Engine/Source/Runtime/Landscape/Classes/LandscapeComponent.h b/Engine/Source/Runtime/Landscape/Classes/LandscapeComponent.h
--- a/Engine/Source/Runtime/Landscape/Classes/LandscapeComponent.h
+++ b/Engine/Source/Runtime/Landscape/Classes/LandscapeComponent.h
@@ -833,8 +833,6 @@ public:
 	/** Get the Landscape Actor's editing layer GUID */
 	FGuid GetEditingLayerGUID() const;
 
-	void CopyFinalLayerIntoEditingLayer(FLandscapeEditDataInterface& DataInterface, TSet<UTexture2D*>& ProcessedHeightmaps);
-
 	void SetPendingCollisionDataUpdate(bool bInPendingCollisionDataUpdate) { bPendingCollisionDataUpdate = bInPendingCollisionDataUpdate; }
 	bool GetPendingCollisionDataUpdate() const { return bPendingCollisionDataUpdate; }
 	void SetPendingLayerCollisionDataUpdate(bool bInPendingLayerCollisionDataUpdate) { bPendingLayerCollisionDataUpdate = bInPendingLayerCollisionDataUpdate; }
diff --git a/Engine/Source/Runtime/Landscape/Classes/LandscapeEditLayer.h b/Engine/Source/Runtime/Landscape/Classes/LandscapeEditLayer.h
--- a/Engine/Source/Runtime/Landscape/Classes/LandscapeEditLayer.h
+++ b/Engine/Source/Runtime/Landscape/Classes/LandscapeEditLayer.h
@@ -217,15 +217,6 @@ public:
 	 */
 	LANDSCAPE_API virtual void SetVisible(bool bInVisible, bool bInModify);
 	
-	/**
-	 * Sets the layer's visibility value : reserved for internal use : we only need this version for intermediate renders, which 
-	 * should disappear in a later version
-	 * @param bInVisible : true to set visible, false to set invisible
-	 * @param bInModify : true when Modify() needs to be called
-	 * @param bInBroadcastDataChange : true when BroadcastOnLayerDataChanged() needs to be called
-	 */
-	UE_INTERNAL LANDSCAPE_API virtual void SetVisible(bool bInVisible, bool bInModify, bool bInBroadcastDataChange);
-
 	LANDSCAPE_API virtual bool IsVisible() const;
 
 	/**
diff --git a/Engine/Source/Runtime/Landscape/Classes/LandscapeTextureHash.h b/Engine/Source/Runtime/Landscape/Classes/LandscapeTextureHash.h
--- a/Engine/Source/Runtime/Landscape/Classes/LandscapeTextureHash.h
+++ b/Engine/Source/Runtime/Landscape/Classes/LandscapeTextureHash.h
@@ -45,7 +45,7 @@ class ULandscapeTextureHash : public UAssetUserData
 
 	// cache of recently serialized hash values
 	// this ensure that if the texture is brought back to match exactly a recently serialized state, it will have exactly the same hash that it was serialized with.
-	// This helps us come back to the serialized (presumably cache-friendly) state after no-ops like:  modify + undo, hide layer + unhide layer, intermediate render stomps, etc.
+	// This helps us come back to the serialized (presumably cache-friendly) state after no-ops like:  modify + undo, hide layer + unhide layer, etc.
 	TMap<FGuid, FGuid> RecentlySerializedHashes;
 
 public:
diff --git a/Engine/Source/Runtime/Landscape/Private/Landscape.cpp b/Engine/Source/Runtime/Landscape/Private/Landscape.cpp
--- a/Engine/Source/Runtime/Landscape/Private/Landscape.cpp
+++ b/Engine/Source/Runtime/Landscape/Private/Landscape.cpp
@@ -2561,42 +2561,6 @@ const FLandscapeLayerComponentData* ULandscapeComponent::GetEditingLayer() const
 	return nullptr;
 }
 
-void ULandscapeComponent::CopyFinalLayerIntoEditingLayer(FLandscapeEditDataInterface& DataInterface, TSet<UTexture2D*>& ProcessedHeightmaps)
-{
-	Modify();
-	GetLandscapeProxy()->Modify();
-
-	// Heightmap	
-	UTexture2D* EditingTexture = GetHeightmap(true);
-	if (!ProcessedHeightmaps.Contains(EditingTexture))
-	{
-		DataInterface.CopyTextureFromHeightmap(EditingTexture, this, 0);
-		ProcessedHeightmaps.Add(EditingTexture);
-	}
-
-	// Weightmap
-	const TArray<FWeightmapLayerAllocationInfo>& FinalWeightmapLayerAllocations = GetWeightmapLayerAllocations();
-	TArray<FWeightmapLayerAllocationInfo>& EditingLayerWeightmapLayerAllocations = GetWeightmapLayerAllocations(GetEditingLayerGUID());
-
-	// Add missing Alloc Infos
-	for (const FWeightmapLayerAllocationInfo& FinalAllocInfo : FinalWeightmapLayerAllocations)
-	{
-		int32 Index = EditingLayerWeightmapLayerAllocations.IndexOfByPredicate([&FinalAllocInfo](const FWeightmapLayerAllocationInfo& EditingAllocInfo) { return EditingAllocInfo.LayerInfo == FinalAllocInfo.LayerInfo; });
-		if (Index == INDEX_NONE)
-		{
-			new (EditingLayerWeightmapLayerAllocations) FWeightmapLayerAllocationInfo(FinalAllocInfo.LayerInfo);
-		}
-	}
-
-	ReallocateWeightmaps(&DataInterface, GetEditingLayerGUID(), /*bInSaveToTransactionBuffer = */true, /*bool bInForceReallocate = */false, /*InTargetProxy = */nullptr, /*InRestrictSharingToComponents = */nullptr);
-
-	const TArray<TObjectPtr<UTexture2D>>& EditingWeightmapTextures = GetWeightmapTextures(true);
-	for (const FWeightmapLayerAllocationInfo& AllocInfo : EditingLayerWeightmapLayerAllocations)
-	{
-		DataInterface.CopyTextureFromWeightmap(EditingWeightmapTextures[AllocInfo.WeightmapTextureIndex], AllocInfo.WeightmapTextureChannel, this, AllocInfo.LayerInfo, 0);
-	}
-}
-
 FGuid ULandscapeComponent::GetEditingLayerGUID() const
 {
 	ALandscape* Landscape = GetLandscapeActor();
diff --git a/Engine/Source/Runtime/Landscape/Private/LandscapeEditLayer.cpp b/Engine/Source/Runtime/Landscape/Private/LandscapeEditLayer.cpp
--- a/Engine/Source/Runtime/Landscape/Private/LandscapeEditLayer.cpp
+++ b/Engine/Source/Runtime/Landscape/Private/LandscapeEditLayer.cpp
@@ -151,12 +151,6 @@ FName ULandscapeEditLayerBase::GetName() const
 }
 
 void ULandscapeEditLayerBase::SetVisible(bool bInVisible, bool bInModify)
-{
-	SetVisible(bInVisible, bInModify, /*bInBroadcastDataChange = */true);
-}
-
-// deprecate/remove with IntermediateRender
-void ULandscapeEditLayerBase::SetVisible(bool bInVisible, bool bInModify, bool bInBroadcastDataChange)
 {
 	const bool bHasValueChanged = (bInVisible != bVisible);
 	if (bHasValueChanged)
@@ -169,10 +163,7 @@ void ULandscapeEditLayerBase::SetVisible(bool bInVisible, bool bInModify, bool b
 		bVisible = bInVisible;
 	}
 
-	if (bInBroadcastDataChange)
-	{
-		BroadcastOnLayerDataChanged(GET_MEMBER_NAME_CHECKED(ULandscapeEditLayerBase, bVisible), /*bInUserTriggered = */true, /*bInRequiresLandscapeUpdate =*/ true, bHasValueChanged, /*InChangeType = */EPropertyChangeType::ValueSet);
-	}
+	BroadcastOnLayerDataChanged(GET_MEMBER_NAME_CHECKED(ULandscapeEditLayerBase, bVisible), /*bInUserTriggered = */true, /*bInRequiresLandscapeUpdate =*/ true, bHasValueChanged, /*InChangeType = */EPropertyChangeType::ValueSet);
 }
 
 bool ULandscapeEditLayerBase::IsVisible() const
diff --git a/Engine/Source/Runtime/Landscape/Private/LandscapeEditLayers.cpp b/Engine/Source/Runtime/Landscape/Private/LandscapeEditLayers.cpp
--- a/Engine/Source/Runtime/Landscape/Private/LandscapeEditLayers.cpp
+++ b/Engine/Source/Runtime/Landscape/Private/LandscapeEditLayers.cpp
@@ -1341,12 +1341,6 @@ bool ALandscape::ComputeLandscapeLayerBrushInfo(FTransform& OutLandscapeTransfor
 	return true;
 }
 
-// TODO:  remove with IntermediateRender.
-bool ALandscape::AllowSelectiveRenderEditLayers() const
-{
-	return true;
-}
-
 bool ALandscape::GetIsPendingEditLayerConversion() const
 {
 	return bIsPendingEditLayerConversion;
@@ -5188,7 +5182,6 @@ void ALandscape::UpdateForChangedHeightmaps(const TArrayView<FLandscapeEditLayer
 void ALandscape::ResolveLayersHeightmapTexture(
 	FTextureToComponentHelper const& MapHelper,
 	TSet<UTexture2D*> const& HeightmapsToResolve,
-	bool bIntermediateRender,
 	TArray<FLandscapeEditLayerComponentReadbackResult>& InOutComponentReadbackResults)
 {
 	TRACE_CPUPROFILER_EVENT_SCOPE(LandscapeLayers_ResolveLayersHeightmapTexture);
@@ -5206,7 +5199,7 @@ void ALandscape::ResolveLayersHeightmapTexture(
 		check(LandscapeProxy);
 		if (FLandscapeEditLayerReadback** CPUReadback = LandscapeProxy->HeightmapsCPUReadback.Find(Heightmap))
 		{
-			const bool bChanged = ResolveLayersTexture(MapHelper, *CPUReadback, Heightmap, bIntermediateRender, InOutComponentReadbackResults, /*bIsWeightmap = */false);
+			const bool bChanged = ResolveLayersTexture(MapHelper, *CPUReadback, Heightmap, InOutComponentReadbackResults, /*bIsWeightmap = */false);
 			if (bChanged)
 			{
 				ChangedComponents.Append(MapHelper.HeightmapToComponents[Heightmap]);
@@ -5563,7 +5556,6 @@ bool ALandscape::ResolveLayersTexture(
 	FTextureToComponentHelper const& MapHelper,
 	FLandscapeEditLayerReadback* InCPUReadback,
 	UTexture2D* InOutputTexture,
-	bool bIntermediateRender,
 	TArray<FLandscapeEditLayerComponentReadbackResult>& InOutComponentReadbackResults,
 	bool bIsWeightmap)
 {
@@ -5623,34 +5615,17 @@ bool ALandscape::ResolveLayersTexture(
 				const FColor* NewMipData = NewMipsData[MipIndex].GetData();
 
 				// Do dirty detection on first mip.
-				// Don't do this for intermediate renders.
-				if (MipIndex == 0 && !bIntermediateRender)
+				if (MipIndex == 0)
 				{
 					NewHash = ULandscapeTextureHash::CalculateTextureHash64(NewMipData, MipTexels, TextureType);
-
-					// If the last readback bypassed the hash test because it was an intermediate render, force a lock and copy to restore the "correct" non-intermediate data.  The regular
-					// conditional notification can proceed as normal, based on the stored hash.
-					bool bForceUpdate = false;
-					if (InCPUReadback->GetLastReadbackWasIntermediate())
-					{
-						bForceUpdate = true;
-						InCPUReadback->SetLastReadbackWasIntermediate(false);
-
-						if (ChangedChannelsMask.IsSet())
-						{
-							// The previous data isn't available to compare against. Assume the worst case: all channels changed.
-							*ChangedChannelsMask = 0xf;
-						}
-					}
-
-					if (NewHash != OldHash || bForceUpdate)
+					if (NewHash != OldHash)
 					{
 						// Defer locking the texture for ReadWrite until after we know that the data is changing at all.  Unlocking after a ReadWrite causes an 
 						// expensive rehashing using the slower hash function used in IoHash.  
 						const FColor* OldMipDataReadOnly = (const FColor*) InOutputTexture->Source.LockMipReadOnly(MipIndex);
 						check(OldMipDataReadOnly);
 
-						if (bForceUpdate || ULandscapeTextureHash::DoesTextureDataChangeExceedThreshold(OldMipDataReadOnly, NewMipData, MipTexels, bIsWeightmap ? ELandscapeTextureType::Weightmap : ELandscapeTextureType::Heightmap, OldHash, NewHash, ChangedChannelsMask))
+						if (ULandscapeTextureHash::DoesTextureDataChangeExceedThreshold(OldMipDataReadOnly, NewMipData, MipTexels, bIsWeightmap ? ELandscapeTextureType::Weightmap : ELandscapeTextureType::Heightmap, OldHash, NewHash, ChangedChannelsMask))
 						{
 							// convert the lock to a read/write lock (first have to release the existing read-only lock, then acquire a new read/write lock)
 							InOutputTexture->Source.UnlockMip(MipIndex);
@@ -5659,7 +5634,7 @@ bool ALandscape::ResolveLayersTexture(
 							check(!bLockedForWrite);
 							bLockedForWrite = true;
 
-							// Note that the hash might not change in the bForceUpdate case.  We skip the dirty notice if it hasn't changed.
+							// We skip the dirty notice if it hasn't changed.
 							bChanged |= InCPUReadback->SetHash(NewHash);
 							if (bChanged)
 							{
@@ -5691,19 +5666,11 @@ bool ALandscape::ResolveLayersTexture(
 						}
 					}
 				}
-				else if (bLockedForWrite || bIntermediateRender)
+				else if (bLockedForWrite)
 				{
 					// If we locked mip 0, lock the rest too, even though FTextureSource isn't tracking mip-level locking.
 					MipDataWriteable = (FColor*) InOutputTexture->Source.LockMip(MipIndex);
 					bLockedForWrite = true;
-
-					if (MipIndex == 0 && bIntermediateRender)
-					{
-						// Remember that we bypassed the hash test because this was an intermediate render.  The stored hash will not reflect the real contents of the
-						// destination texture.  An intermediate render is expected to be followed by a readback of a regular render to fix that data.  The code has been
-						// doing this for sime time.  It should no longer be necessary after deprecating GlobalMerge and switching to BatchMerge.
-						InCPUReadback->SetLastReadbackWasIntermediate(true);
-					}
 				}
 
 				if (bChanged)
@@ -5712,12 +5679,10 @@ bool ALandscape::ResolveLayersTexture(
 					check(bLockedForWrite && MipDataWriteable);
 					if (bIsWeightmap)
 					{
-						// NOTE: comparing the New and Old data can be misleading when bOldDataIsFromIntermediateRender
 						OnDirtyWeightmap(MapHelper, InOutputTexture, (FColor*)MipDataWriteable, NewMipData, MipIndex, ChangedChannelsMask.Get(0));
 					}
 					else
 					{
-						// NOTE: comparing the New and Old data can be misleading when bOldDataIsFromIntermediateRender
 						OnDirtyHeightmap(MapHelper, InOutputTexture, (FColor*)MipDataWriteable, NewMipData, MipIndex);
 					}
 				}
@@ -5741,8 +5706,6 @@ bool ALandscape::ResolveLayersTexture(
 				}
 			}
 
-			check(!bIntermediateRender || !bChanged);	// intermediate renders should not be considered changed (we keep the same old hash)
-
 			// update the hash (New hash if it's considered changed, otherwise continue to use the old hash)
 			// this must happen after UnlockMip, so the SourceID is up to date.
 			ULandscapeTextureHash::SetHash64(InOutputTexture, bChanged ? NewHash : OldHash, ELandscapeTextureUsage::FinalData, bIsWeightmap ? ELandscapeTextureType::Weightmap : ELandscapeTextureType::Heightmap);
@@ -6855,7 +6818,6 @@ void ALandscape::UpdateLayersMaterialInstances(const TArray<ULandscapeComponent*
 void ALandscape::ResolveLayersWeightmapTexture(
 	FTextureToComponentHelper const& MapHelper,
 	TSet<UTexture2D*> const& WeightmapsToResolve,
-	bool bIntermediateRender,
 	TArray<FLandscapeEditLayerComponentReadbackResult>& InOutComponentReadbackResults)
 {
 	TRACE_CPUPROFILER_EVENT_SCOPE(LandscapeLayers_ResolveLayersWeightmapTexture);
@@ -6873,7 +6835,7 @@ void ALandscape::ResolveLayersWeightmapTexture(
 		check(LandscapeProxy);
 		if (FLandscapeEditLayerReadback** CPUReadback = LandscapeProxy->WeightmapsCPUReadback.Find(Weightmap))
 		{
-			const bool bChanged = ResolveLayersTexture(MapHelper, *CPUReadback, Weightmap, bIntermediateRender, InOutComponentReadbackResults, /*bIsWeightmap = */true);
+			const bool bChanged = ResolveLayersTexture(MapHelper, *CPUReadback, Weightmap, InOutComponentReadbackResults, /*bIsWeightmap = */true);
 			if (bChanged)
 			{
 				ChangedComponents.Append(MapHelper.WeightmapToComponents[Weightmap]);
@@ -7377,7 +7339,7 @@ bool ALandscape::SelectiveRenderEditLayersWeightmaps(const FLandscapeEditLayerRe
 	}
 }
 
-void ALandscape::UpdateLayersContent(bool bInWaitForStreaming, bool bInSkipMonitorLandscapeEdModeChanges, bool bIntermediateRender, bool bFlushRender)
+void ALandscape::UpdateLayersContent(bool bInWaitForStreaming, bool bInSkipMonitorLandscapeEdModeChanges, bool bFlushRender)
 {
 	TRACE_CPUPROFILER_EVENT_SCOPE(LandscapeLayers_UpdateLayersContent);
 
@@ -7531,26 +7493,17 @@ void ALandscape::UpdateLayersContent(bool bInWaitForStreaming, bool bInSkipMonit
 	FTextureToComponentHelper MapHelper(*LandscapeInfo);
 
 	// Poll and complete any outstanding resolve work
-	// If bIntermediateRender then we want to flush all work here before we do the intermediate render later on
-	// if bFlushRender then we skip this because we will flush later anyway
 	if (bProcessReadbacks)
 	{
-		// These flags might look like they're being mixed up but they're not!
-		const bool bDoIntermediateRender = false; // bIntermediateRender flag is for the work queued up this frame not the delayed resolves
-		// bIntermediateRender needs to flush before the intermediate render happens.  If flush is requested without intermediate render, this function could
-		// quit early because LayerContentUpdateModes==0, so this is still the place to flush.  Flushing twice would only happen in the rare case when flush
-		// is requested, pre-existing readbacks are pending (bProcessReadbacks), and new work is also being initiated via LayerContentUpdateModes.
-		const bool bDoFlushRender = bIntermediateRender || bFlushRender;
-
 		// Flushing once all readback tasks is much faster than asking each to do it so start by doing just this :
-		if (bDoFlushRender)
+		if (bFlushRender)
 		{
 			FLandscapeEditLayerReadback::FlushAllReadbackTasks();
 		}
 
 		TArray<FLandscapeEditLayerComponentReadbackResult> ComponentReadbackResults;
-		ResolveLayersHeightmapTexture(MapHelper, MapHelper.Heightmaps, bDoIntermediateRender, ComponentReadbackResults);
-		ResolveLayersWeightmapTexture(MapHelper, MapHelper.Weightmaps, bDoIntermediateRender, ComponentReadbackResults);
+		ResolveLayersHeightmapTexture(MapHelper, MapHelper.Heightmaps, ComponentReadbackResults);
+		ResolveLayersWeightmapTexture(MapHelper, MapHelper.Weightmaps, ComponentReadbackResults);
 		LayerContentUpdateModes |= UpdateAfterReadbackResolves(ComponentReadbackResults);
 	}
 
@@ -7595,16 +7548,15 @@ void ALandscape::UpdateLayersContent(bool bInWaitForStreaming, bool bInSkipMonit
 	{
 		// Flushing once all readback tasks is much faster than asking each to do it so start by doing just this :
 		FLandscapeEditLayerReadback::FlushAllReadbackTasks();
-		// When flushing, don't bother resolving textures that weren't requested to be updated in the first place. This reduces the workload when doing a flushing intermediate render for a heightmap
-		//  tool (smooth/flatten/... tool), for example, by not resolving weightmaps then. 
+		// When flushing, don't bother resolving textures that weren't requested to be updated in the first place
 		//  We cannot do this in the non-flush case above, because LayerContentUpdateModes might have changed since the readbacks have been requested so we still need to perform the readbacks on all textures
 		if ((LayerContentUpdateModes & ELandscapeLayerUpdateMode::Update_Heightmap_Types) != 0)
 		{
-			ResolveLayersHeightmapTexture(UpdateLayersContentContext.MapHelper, UpdateLayersContentContext.HeightmapsToResolve, bIntermediateRender, UpdateLayersContentContext.AllLandscapeComponentReadbackResults);
+			ResolveLayersHeightmapTexture(UpdateLayersContentContext.MapHelper, UpdateLayersContentContext.HeightmapsToResolve, UpdateLayersContentContext.AllLandscapeComponentReadbackResults);
 		}
 		if ((LayerContentUpdateModes & ELandscapeLayerUpdateMode::Update_Weightmap_Types) != 0)
 		{
-			ResolveLayersWeightmapTexture(UpdateLayersContentContext.MapHelper, UpdateLayersContentContext.WeightmapsToResolve, bIntermediateRender, UpdateLayersContentContext.AllLandscapeComponentReadbackResults);
+			ResolveLayersWeightmapTexture(UpdateLayersContentContext.MapHelper, UpdateLayersContentContext.WeightmapsToResolve, UpdateLayersContentContext.AllLandscapeComponentReadbackResults);
 		}
 	}
 
@@ -7804,17 +7756,19 @@ void ALandscape::FlushLayerContentThisFrame()
 	if (LastFlushedLayerUpdateFrame != CurrentFrame && CanUpdateLayersContent())
 	{
 		LastFlushedLayerUpdateFrame = CurrentFrame;
-		ForceUpdateLayersContent(/*bIntermediateRender = */ false);
+		ForceUpdateLayersContent();
 	}
 }
 
-void ALandscape::ForceUpdateLayersContent(bool bIntermediateRender)
+void ALandscape::ForceUpdateLayersContent()
 {
-	const bool bWaitForStreaming = true;
-	const bool bInSkipMonitorLandscapeEdModeChanges = true;
-	const bool bFlushRender = true;
+	UpdateLayersContent(/*bWaitForStreaming = */true, /*bInSkipMonitorLandscapeEdModeChanges = */true, /*bFlushRender = */true);
+}
 
-	UpdateLayersContent(bWaitForStreaming, bInSkipMonitorLandscapeEdModeChanges, bIntermediateRender, bFlushRender);
+// Deprecated
+void ALandscape::ForceUpdateLayersContent(bool bIntermediateRender)
+{
+	ForceUpdateLayersContent();
 }
 
 void ALandscape::ForceLayersFullUpdate()
@@ -7833,7 +7787,7 @@ void ALandscape::ForceLayersFullUpdate()
 
 	RequestSplineLayerUpdate();
 	RequestLayersContentUpdateForceAll();
-	ForceUpdateLayersContent(/* bIntermediateRender */ false);
+	ForceUpdateLayersContent();
 }
 
 void ALandscape::TickLayers(float DeltaTime)
@@ -7853,7 +7807,7 @@ void ALandscape::TickLayers(float DeltaTime)
 			World->bShouldSimulatePhysics = true;
 		}
 
-		UpdateLayersContent();
+		UpdateLayersContent(/*bInWaitForStreaming = */false, /*bInSkipMonitorLandscapeEdModeChanges = */false, /*bFlushRender = */false);
 	}
 }
 
@@ -8555,13 +8509,6 @@ bool ALandscape::DeleteLayer(int32 InLayerIndex)
 
 void ALandscape::CollapseLayer(int32 InLayerIndex)
 {
-	if (!AllowSelectiveRenderEditLayers())
-	{
-		// TODO:  remove with IntermediateRender.
-		CollapseLayerIntermediateRender(InLayerIndex);
-		return;
-	}
-
 	check(InLayerIndex >= 1);
 
 	FScopedSlowTask SlowTask(static_cast<float>(GetLandscapeInfo()->XYtoComponentMap.Num()), LOCTEXT("Landscape_CollapseLayer_SlowWork", "Collapsing Layer..."));
@@ -8588,8 +8535,7 @@ void ALandscape::CollapseLayer(int32 InLayerIndex)
 	check(LayerSamples > 0);
 
 	// ForceUpdateLayersContent without requesting any updates to flush any pending evaluation.
-	const bool bLocalIntermediateRender = false;
-	ForceUpdateLayersContent(bLocalIntermediateRender);
+	ForceUpdateLayersContent();
 
 	TArray<uint16> HeightData;
 	HeightData.AddZeroed(LayerSamples);
@@ -8661,73 +8607,6 @@ void ALandscape::CollapseLayer(int32 InLayerIndex)
 	RequestLayersContentUpdateForceAll();
 }
 
-void ALandscape::CollapseLayerIntermediateRender(int32 InLayerIndex)
-{
-	// TODO:  remove with IntermediateRender.
-	FScopedSlowTask SlowTask(static_cast<float>(GetLandscapeInfo()->XYtoComponentMap.Num()), LOCTEXT("Landscape_CollapseLayer_SlowWork", "Collapsing Layer..."));
-	SlowTask.MakeDialog();
-	TArray<bool> BackupVisibility;
-	TArray<bool> BackupBrushVisibility;
-	for (int32 i = 0; i < LandscapeEditLayers.Num(); ++i)
-	{
-		BackupVisibility.Add(LandscapeEditLayers[i].EditLayer->IsVisible());
-		LandscapeEditLayers[i].EditLayer->SetVisible(i == InLayerIndex || i == InLayerIndex - 1, /*bInModify = */ true);
-	}
-
-	for (int32 i = 0; i < LandscapeEditLayers[InLayerIndex].Brushes.Num(); ++i)
-	{
-		BackupBrushVisibility.Add(LandscapeEditLayers[InLayerIndex].Brushes[i].GetBrush()->IsVisible());
-		LandscapeEditLayers[InLayerIndex].Brushes[i].GetBrush()->SetIsVisible(false);
-	}
-
-	// Call Request Update on all components...
-	GetLandscapeInfo()->ForAllLandscapeComponents([](ULandscapeComponent* LandscapeComponent)
-	{
-		LandscapeComponent->RequestWeightmapUpdate(false, false);
-		LandscapeComponent->RequestHeightmapUpdate(false, false);
-	});
-
-	const bool bLocalIntermediateRender = true;
-	ForceUpdateLayersContent(bLocalIntermediateRender);
-
-	// Do copy
-	{
-		FLandscapeEditDataInterface DataInterface(GetLandscapeInfo());
-		DataInterface.SetShouldDirtyPackage(true);
-
-		TSet<UTexture2D*> ProcessedHeightmaps;
-		FScopedSetLandscapeEditingLayer ScopeEditingLayer(this, LandscapeEditLayers[InLayerIndex - 1].EditLayer->GetGuid());
-		GetLandscapeInfo()->ForAllLandscapeComponents([&](ULandscapeComponent* LandscapeComponent)
-		{
-			SlowTask.EnterProgressFrame(1.f);
-			LandscapeComponent->CopyFinalLayerIntoEditingLayer(DataInterface, ProcessedHeightmaps);
-		});
-	}
-
-	TArray<ALandscapeBlueprintBrushBase*> BrushesToMove;
-	for (int32 i = 0; i < LandscapeEditLayers[InLayerIndex].Brushes.Num(); ++i)
-	{
-		ALandscapeBlueprintBrushBase* CurrentBrush = LandscapeEditLayers[InLayerIndex].Brushes[i].GetBrush();
-		CurrentBrush->SetIsVisible(BackupBrushVisibility[i]);
-		BrushesToMove.Add(CurrentBrush);
-	}
-
-	for (ALandscapeBlueprintBrushBase* Brush : BrushesToMove)
-	{
-		RemoveBrushFromLayer(InLayerIndex, Brush);
-		AddBrushToLayer(InLayerIndex - 1, Brush);
-	}
-
-	for (int32 i = 0; i < LandscapeEditLayers.Num(); ++i)
-	{
-		LandscapeEditLayers[i].EditLayer->SetVisible(BackupVisibility[i], /*bInModify = */ true);
-	}
-
-	DeleteLayer(InLayerIndex);
-
-	RequestLayersContentUpdateForceAll();
-}
-
 void ALandscape::GetUsedPaintLayers(int32 InLayerIndex, TArray<ULandscapeLayerInfoObject*>& OutUsedLayerInfos) const
 {
 	if (const ULandscapeEditLayerBase* EditLayer = GetEditLayer(InLayerIndex))
@@ -9694,12 +9573,6 @@ void LandscapeDumpSelectiveLayerRender(const TArray<FString>& Args)
 		return;
 	}
 
-	if (!Landscape->AllowSelectiveRenderEditLayers())
-	{
-		UE_LOG(LogLandscape, Warning, TEXT("LandscapeDumpSelectiveLayerRender: can't use Partial renders"));
-		return;
-	}
-
 	TBitArray<> LayerVis;  // leave empty for default layer visibility
 	if (bHasLayerMask)
 	{
diff --git a/Engine/Source/Runtime/Landscape/Private/LandscapeEditReadback.cpp b/Engine/Source/Runtime/Landscape/Private/LandscapeEditReadback.cpp
--- a/Engine/Source/Runtime/Landscape/Private/LandscapeEditReadback.cpp
+++ b/Engine/Source/Runtime/Landscape/Private/LandscapeEditReadback.cpp
@@ -328,8 +328,6 @@ FLandscapeEditLayerReadback::~FLandscapeEditLayerReadback()
 	{
 		GReadbackTaskPool.Free(TaskHandle);
 	}
-
-	ensure(!bLastReadbackWasIntermediate);  // Expecting intermediate render data to always be cleaned up by a regular render.  It shouldn't escape the lifetime of this object.
 }
 
 bool FLandscapeEditLayerReadback::SetHash(uint64 InHash)
@@ -339,16 +337,6 @@ bool FLandscapeEditLayerReadback::SetHash(uint64 InHash)
 	return bChanged;
 }
 
-void FLandscapeEditLayerReadback::SetLastReadbackWasIntermediate(bool bValue)
-{
-	bLastReadbackWasIntermediate = bValue;
-}
-
-bool FLandscapeEditLayerReadback::GetLastReadbackWasIntermediate() const
-{
-	return bLastReadbackWasIntermediate;
-}
-
 void FLandscapeEditLayerReadback::Enqueue(UTexture2D const* InSourceTexture, FReadbackContext&& InReadbackContext)
 {
 	const int32 TaskHandle = GReadbackTaskPool.Allocate(InSourceTexture, MoveTemp(InReadbackContext));
diff --git a/Engine/Source/Runtime/Landscape/Private/LandscapeEditReadback.h b/Engine/Source/Runtime/Landscape/Private/LandscapeEditReadback.h
--- a/Engine/Source/Runtime/Landscape/Private/LandscapeEditReadback.h
+++ b/Engine/Source/Runtime/Landscape/Private/LandscapeEditReadback.h
@@ -24,12 +24,6 @@ public:
 	/** Get the stored hash value. */
 	uint64 GetHash() const { return Hash; }
 
-	// Stores whether the last time the data was read back and written to the texture, it was an intermediate render.  That intermediate
-	// render did not update the stored hash and the data should not be stored permanently.  The next regular readback that occurs should ignore
-	// the stored hash and copy the new data into the texture.
-	void SetLastReadbackWasIntermediate(bool bValue);
-	bool GetLastReadbackWasIntermediate() const;
-
 	using FPerChannelLayerNames = TStaticArray<FName, 4>;
 
 	/** Per component context required for processing read back results. */
@@ -77,5 +71,4 @@ public:
 private:
 	uint64 Hash = 0;
 	TArray<int32> TaskHandles;
-	bool bLastReadbackWasIntermediate = false;
 };
