Changed landscape edit layers merge API from experimental to official

diff --git a/Engine/Source/Runtime/Landscape/Classes/Landscape.h b/Engine/Source/Runtime/Landscape/Classes/Landscape.h
--- a/Engine/Source/Runtime/Landscape/Classes/Landscape.h
+++ b/Engine/Source/Runtime/Landscape/Classes/Landscape.h
@@ -140,7 +140,6 @@ struct FLandscapeLayerBrush
 	void SetOwner(ALandscape* InOwner);
 
 	//~ Begin UE::Landscape::EditLayers::IEditLayerRendererProvider implementation
-	LANDSCAPE_EDIT_LAYERS_BATCHED_MERGE_EXPERIMENTAL
 	LANDSCAPE_API virtual TArray<UE::Landscape::EditLayers::FEditLayerRendererState> GetEditLayerRendererStates(const UE::Landscape::EditLayers::FMergeContext* InMergeContext) override;
 	//~ End UE::Landscape::EditLayers::IEditLayerRendererProvider implementation
 #endif // WITH_EDITOR
@@ -177,7 +176,6 @@ struct FLandscapeLayer
 
 #if WITH_EDITOR
 	//~ Begin UE::Landscape::EditLayers::IEditLayerRendererProvider implementation
-	LANDSCAPE_EDIT_LAYERS_BATCHED_MERGE_EXPERIMENTAL
 	LANDSCAPE_API virtual TArray<UE::Landscape::EditLayers::FEditLayerRendererState> GetEditLayerRendererStates(const UE::Landscape::EditLayers::FMergeContext* InMergeContext) override;
 	//~ End UE::Landscape::EditLayers::IEditLayerRendererProvider implementation
 #endif // WITH_EDITOR
diff --git a/Engine/Source/Runtime/Landscape/Classes/LandscapeEditLayer.h b/Engine/Source/Runtime/Landscape/Classes/LandscapeEditLayer.h
--- a/Engine/Source/Runtime/Landscape/Classes/LandscapeEditLayer.h
+++ b/Engine/Source/Runtime/Landscape/Classes/LandscapeEditLayer.h
@@ -395,18 +395,12 @@ public:
 
 #if WITH_EDITOR
 	//~ Begin ILandscapeEditLayerRenderer implementation
-	LANDSCAPE_EDIT_LAYERS_BATCHED_MERGE_EXPERIMENTAL
 	LANDSCAPE_API virtual void GetRendererStateInfo(const UE::Landscape::EditLayers::FMergeContext* InMergeContext,
 		UE::Landscape::EditLayers::FEditLayerTargetTypeState& OutSupportedTargetTypeState, UE::Landscape::EditLayers::FEditLayerTargetTypeState& OutEnabledTargetTypeState, TArray<UE::Landscape::EditLayers::FTargetLayerGroup>& OutTargetLayerGroups) const override;
-	LANDSCAPE_EDIT_LAYERS_BATCHED_MERGE_EXPERIMENTAL
 	virtual UE::Landscape::EditLayers::ERenderFlags GetRenderFlags(const UE::Landscape::EditLayers::FMergeContext* InMergeContext) const override;
-	LANDSCAPE_EDIT_LAYERS_BATCHED_MERGE_EXPERIMENTAL
 	LANDSCAPE_API virtual TArray<UE::Landscape::EditLayers::FEditLayerRenderItem> GetRenderItems(const UE::Landscape::EditLayers::FMergeContext* InMergeContext) const override;
-	LANDSCAPE_EDIT_LAYERS_BATCHED_MERGE_EXPERIMENTAL
 	LANDSCAPE_API virtual bool RenderLayer(UE::Landscape::EditLayers::FRenderParams& RenderParams, UE::Landscape::FRDGBuilderRecorder& RDGBuilderRecorder) override;
-	LANDSCAPE_EDIT_LAYERS_BATCHED_MERGE_EXPERIMENTAL
 	LANDSCAPE_API virtual void BlendLayer(UE::Landscape::EditLayers::FRenderParams& RenderParams, UE::Landscape::FRDGBuilderRecorder& RDGBuilderRecorder) override;
-	LANDSCAPE_EDIT_LAYERS_BATCHED_MERGE_EXPERIMENTAL
 	LANDSCAPE_API virtual FString GetEditLayerRendererDebugName() const override;
 	//~ End ILandscapeEditLayerRenderer implementation
 #endif // WITH_EDITOR
diff --git a/Engine/Source/Runtime/Landscape/Public/LandscapeBlueprintBrushBase.h b/Engine/Source/Runtime/Landscape/Public/LandscapeBlueprintBrushBase.h
--- a/Engine/Source/Runtime/Landscape/Public/LandscapeBlueprintBrushBase.h
+++ b/Engine/Source/Runtime/Landscape/Public/LandscapeBlueprintBrushBase.h
@@ -108,17 +108,12 @@ public:
 
 #if WITH_EDITOR
 	//~ Begin ILandscapeEditLayerRenderer implementation
-	LANDSCAPE_EDIT_LAYERS_BATCHED_MERGE_EXPERIMENTAL
 	LANDSCAPE_API virtual void GetRendererStateInfo(const UE::Landscape::EditLayers::FMergeContext* InMergeContext,
 		UE::Landscape::EditLayers::FEditLayerTargetTypeState& OutSupportedTargetTypeState, UE::Landscape::EditLayers::FEditLayerTargetTypeState& OutEnabledTargetTypeState, 
 		TArray<UE::Landscape::EditLayers::FTargetLayerGroup>& OutTargetLayerGroups) const override;
-	LANDSCAPE_EDIT_LAYERS_BATCHED_MERGE_EXPERIMENTAL
 	virtual UE::Landscape::EditLayers::ERenderFlags GetRenderFlags(const UE::Landscape::EditLayers::FMergeContext* InMergeContext) const override { return UE::Landscape::EditLayers::ERenderFlags::RenderMode_Immediate; }
-	LANDSCAPE_EDIT_LAYERS_BATCHED_MERGE_EXPERIMENTAL
 	LANDSCAPE_API virtual TArray<UE::Landscape::EditLayers::FEditLayerRenderItem> GetRenderItems(const UE::Landscape::EditLayers::FMergeContext* InMergeContext) const override;
-	LANDSCAPE_EDIT_LAYERS_BATCHED_MERGE_EXPERIMENTAL
 	LANDSCAPE_API virtual bool RenderLayer(UE::Landscape::EditLayers::FRenderParams& RenderParams, UE::Landscape::FRDGBuilderRecorder& RDGBuilderRecorder) override;
-	LANDSCAPE_EDIT_LAYERS_BATCHED_MERGE_EXPERIMENTAL
 	LANDSCAPE_API virtual FString GetEditLayerRendererDebugName() const override;
 	//~ End ILandscapeEditLayerRenderer implementation
 
diff --git a/Engine/Source/Runtime/Landscape/Public/LandscapeEditLayerMergeContext.h b/Engine/Source/Runtime/Landscape/Public/LandscapeEditLayerMergeContext.h
--- a/Engine/Source/Runtime/Landscape/Public/LandscapeEditLayerMergeContext.h
+++ b/Engine/Source/Runtime/Landscape/Public/LandscapeEditLayerMergeContext.h
@@ -42,137 +42,114 @@ public:
 	FMergeContext& operator=(const FMergeContext& Other) = default;
 	FMergeContext& operator=(FMergeContext&& Other) = default;
 
-	LANDSCAPE_EDIT_LAYERS_BATCHED_MERGE_EXPERIMENTAL
 	inline ALandscape* GetLandscape() const { return Landscape; }
 
-	LANDSCAPE_EDIT_LAYERS_BATCHED_MERGE_EXPERIMENTAL
 	inline ULandscapeInfo* GetLandscapeInfo() const { return LandscapeInfo; }
 	
 	/**
 	 * @return all target layer names declared in this context, including ones that are potentially invalid (i.e. no layer info associated)
 	 */
-	LANDSCAPE_EDIT_LAYERS_BATCHED_MERGE_EXPERIMENTAL
 	inline const TArray<FName>& GetAllTargetLayerNames() const { return AllTargetLayerNames; }
 
 	/**
 	 * @return all target layer names that have a valid layer info object. Will return the height layer in the case of a heightmap merge, even though there's no layer info for those : 
 	 */
-	LANDSCAPE_EDIT_LAYERS_BATCHED_MERGE_EXPERIMENTAL
 	TArray<FName> GetValidTargetLayerNames() const;
 
 	/**
 	 * @return all layer info objects declared in this context (can contain null entries for invalid layers or in the case of a heightmap merge)
 	 */
-	LANDSCAPE_EDIT_LAYERS_BATCHED_MERGE_EXPERIMENTAL
 	inline const TArray<ULandscapeLayerInfoObject*>& GetAllWeightmapLayerInfos() const { return AllWeightmapLayerInfos; }
 
 	/**
 	 * @return true if this merge context is used for a heightmap merge, false otherwise (visibility / weightmap)
 	 */
-	LANDSCAPE_EDIT_LAYERS_BATCHED_MERGE_EXPERIMENTAL
 	inline bool IsHeightmapMerge() const { return bIsHeightmapMerge; }
 
 	/**
 	 * @return true if this merge context is used for a heightmap merge, false otherwise (visibility / weightmap)
 	 */
-	LANDSCAPE_EDIT_LAYERS_BATCHED_MERGE_EXPERIMENTAL
 	inline bool ShouldSkipProceduralRenderers() const { return bSkipProceduralRenderers; }
 
 	/**
 	 * @return true if the target layer name is registered to this merge context and has an associated layer info object
 	 */
-	LANDSCAPE_EDIT_LAYERS_BATCHED_MERGE_EXPERIMENTAL
 	LANDSCAPE_API int32 IsValidTargetLayerName(const FName& InName) const;
 
 	/**
 	 * @return true if the target layer name has an associated layer info object. Asserts if the layer name isn't known to the merge context
 	 */
-	LANDSCAPE_EDIT_LAYERS_BATCHED_MERGE_EXPERIMENTAL
 	LANDSCAPE_API int32 IsValidTargetLayerNameChecked(const FName& InName) const;
 
 	/**
 	 * @return true if the target layer index is valid in the list of declared target layers on the merge context. 
 	 *  Important note : it doesn't mean that it's associated with a valid layer info object, use IsValidTargetLayerName for this!
 	 */
-	LANDSCAPE_EDIT_LAYERS_BATCHED_MERGE_EXPERIMENTAL
 	LANDSCAPE_API int32 IsTargetLayerIndexValid(int32 InIndex) const;
 
 	/**
 	 * @return the index of this target layer name in this context, INDEX_NONE if not found
 	 */
-	LANDSCAPE_EDIT_LAYERS_BATCHED_MERGE_EXPERIMENTAL
 	LANDSCAPE_API int32 GetTargetLayerIndexForName(const FName& InName) const;
 
 	/**
 	 * @return the index of this target layer name in this context. Asserts if the layer name isn't known to the merge context
 	 */
-	LANDSCAPE_EDIT_LAYERS_BATCHED_MERGE_EXPERIMENTAL
 	LANDSCAPE_API int32 GetTargetLayerIndexForNameChecked(const FName& InName) const;
 
 	/**
 	 * @return the target layer name corresponding to the index in this context, NAME_None if not found
 	 */
-	LANDSCAPE_EDIT_LAYERS_BATCHED_MERGE_EXPERIMENTAL
 	LANDSCAPE_API FName GetTargetLayerNameForIndex(int32 InIndex) const;
 
 	/**
 	 * @return the target layer name corresponding to the index in this context, NAME_None if not found. Asserts if the layer index isn't known to the merge context
 	 */
-	LANDSCAPE_EDIT_LAYERS_BATCHED_MERGE_EXPERIMENTAL
 	LANDSCAPE_API FName GetTargetLayerNameForIndexChecked(int32 InIndex) const;
 
 	/**
 	 * @return the index of this landscape layer info in this context, INDEX_NONE if not found
 	 */
-	LANDSCAPE_EDIT_LAYERS_BATCHED_MERGE_EXPERIMENTAL
 	LANDSCAPE_API int32 GetTargetLayerIndexForLayerInfo(ULandscapeLayerInfoObject* InLayerInfo) const;
 
 	/**
 	 * @return the index of this landscape layer info in this context. Asserts if the layer info isn't known to the merge context
 	 */
-	LANDSCAPE_EDIT_LAYERS_BATCHED_MERGE_EXPERIMENTAL
 	LANDSCAPE_API int32 GetTargetLayerIndexForLayerInfoChecked(ULandscapeLayerInfoObject* InLayerInfo) const;
 
 	/**
 	 * @return the landscape layer object corresponding to the target layer name in this context, nullptr if not found
 	 */
-	LANDSCAPE_EDIT_LAYERS_BATCHED_MERGE_EXPERIMENTAL
 	LANDSCAPE_API ULandscapeLayerInfoObject* GetTargetLayerInfoForName(const FName& InName) const;
 
 	/**
 	 * @return the landscape layer object corresponding to the target layer name in this context. Asserts if the layer name isn't known to the merge context
 	 */
-	LANDSCAPE_EDIT_LAYERS_BATCHED_MERGE_EXPERIMENTAL
 	LANDSCAPE_API ULandscapeLayerInfoObject* GetTargetLayerInfoForNameChecked(const FName& InName) const;
 
 	/**
 	 * @return the landscape layer object corresponding to the index in this context, nullptr if not found. Asserts if the layer index isn't known to the merge context
 	 */
-	LANDSCAPE_EDIT_LAYERS_BATCHED_MERGE_EXPERIMENTAL
 	LANDSCAPE_API ULandscapeLayerInfoObject* GetTargetLayerInfoForIndex(int32 InIndex) const;
 
 	/**
 	 * @return a TBitArray<> where, for each layer from InTargetLayerNames that is found in this context, the corresponding bit is set
 	 */
-	LANDSCAPE_EDIT_LAYERS_BATCHED_MERGE_EXPERIMENTAL
 	LANDSCAPE_API TBitArray<> ConvertTargetLayerNamesToBitIndices(TConstArrayView<FName> InTargetLayerNames) const;
 
 	/**
 	 * @return a TBitArray<> where, for each layer from InTargetLayerNames that is found in this context, the corresponding bit is set. Asserts if one of the layer names isn't known to the merge context
 	 */
-	LANDSCAPE_EDIT_LAYERS_BATCHED_MERGE_EXPERIMENTAL
 	LANDSCAPE_API TBitArray<> ConvertTargetLayerNamesToBitIndicesChecked(TConstArrayView<FName> InTargetLayerNames) const;
 
 	/**
 	 * @return a list of target layer names for each set bit in InTargetLayerBitIndices. Asserts if InTargetLayerBitIndices isn't the same size as AllTargetLayerNames
 	 */
-	LANDSCAPE_EDIT_LAYERS_BATCHED_MERGE_EXPERIMENTAL
 	LANDSCAPE_API TArray<FName> ConvertTargetLayerBitIndicesToNames(const TBitArray<>& InTargetLayerBitIndices) const;
 
 	/**
 	 * @return a list of landscape layer info objects for each set bit in InTargetLayerBitIndices. Asserts if InTargetLayerBitIndices isn't the same size as AllTargetLayerNames
 	 */
-	LANDSCAPE_EDIT_LAYERS_BATCHED_MERGE_EXPERIMENTAL
 	LANDSCAPE_API TArray<ULandscapeLayerInfoObject*> ConvertTargetLayerBitIndicesToLayerInfos(const TBitArray<>& InTargetLayerBitIndices) const;
 
 	/**
@@ -183,11 +160,9 @@ public:
 	 *     return continueLoop ? true : false;
 	 * });
 	 */
-	LANDSCAPE_EDIT_LAYERS_BATCHED_MERGE_EXPERIMENTAL
 	LANDSCAPE_API void ForEachTargetLayer(const TBitArray<>& InTargetLayerBitIndices, TFunctionRef<bool(int32 /*InTargetLayerIndex*/, const FName& /*InTargetLayerName*/, ULandscapeLayerInfoObject* /*InWeightmapLayerInfo*/)> Fn) const;
 
 	/** Same as ForEachTargetLayer but skips over invalid target layers */
-	LANDSCAPE_EDIT_LAYERS_BATCHED_MERGE_EXPERIMENTAL
 	LANDSCAPE_API void ForEachValidTargetLayer(TFunctionRef<bool(int32 /*InTargetLayerIndex*/, const FName& /*InTargetLayerName*/, ULandscapeLayerInfoObject* /*InWeightmapLayerInfo*/)> Fn) const;
 
 	/**
@@ -211,17 +186,14 @@ public:
 	/**
 	 * @return a bit mask the size of AllTargetLayerNames, where all bits are set to bBitValue
 	 */
-	LANDSCAPE_EDIT_LAYERS_BATCHED_MERGE_EXPERIMENTAL
 	LANDSCAPE_API TBitArray<> BuildTargetLayerBitIndices(bool bInBitValue) const;
 
-	LANDSCAPE_EDIT_LAYERS_BATCHED_MERGE_EXPERIMENTAL
 	const FTargetLayerGroupsPerBlendingMethod& GetTargetLayerGroupsPerBlendingMethod() const { return TargetLayerGroupsPerBlendingMethod; }
 
 	/** 
 	 * Return the target layer groups for the renderers performing the generic blend. This concerns only the blend methods that are applied per renderer
 	 * @param InTargetLayerBitIndices bit mask of the target layers that are concerned by this blend
 	 */
-	LANDSCAPE_EDIT_LAYERS_BATCHED_MERGE_EXPERIMENTAL
 	LANDSCAPE_API TArray<FTargetLayerGroup> BuildGenericBlendTargetLayerGroups(const TBitArray<>& InTargetLayerBitIndices) const;
 
 private:
diff --git a/Engine/Source/Runtime/Landscape/Public/LandscapeEditLayerMergeRenderContext.h b/Engine/Source/Runtime/Landscape/Public/LandscapeEditLayerMergeRenderContext.h
--- a/Engine/Source/Runtime/Landscape/Public/LandscapeEditLayerMergeRenderContext.h
+++ b/Engine/Source/Runtime/Landscape/Public/LandscapeEditLayerMergeRenderContext.h
@@ -316,16 +316,11 @@ public:
 	 * @param InDesiredWriteAccess RHI state of the write RT after the operation (if InDesiredWriteAccess != ERHIAccess::None)
 	 * @param RDGBuilderRecorder recorder to append operations to a single FRDGBuilder, if the the recorder in "recording" mode. In "immediate" mode, a render command will be enqueued immediately
 	 */
-	LANDSCAPE_EDIT_LAYERS_BATCHED_MERGE_EXPERIMENTAL
 	LANDSCAPE_API void CycleBlendRenderTargets(UE::Landscape::FRDGBuilderRecorder& RDGBuilderRecorder);
 
-	LANDSCAPE_EDIT_LAYERS_BATCHED_MERGE_EXPERIMENTAL
 	LANDSCAPE_API ULandscapeScratchRenderTarget* GetBlendRenderTargetWrite() const;
-	LANDSCAPE_EDIT_LAYERS_BATCHED_MERGE_EXPERIMENTAL
 	LANDSCAPE_API ULandscapeScratchRenderTarget* GetBlendRenderTargetRead() const;
-	LANDSCAPE_EDIT_LAYERS_BATCHED_MERGE_EXPERIMENTAL
 	LANDSCAPE_API ULandscapeScratchRenderTarget* GetBlendRenderTargetReadPrevious() const;
-	LANDSCAPE_EDIT_LAYERS_BATCHED_MERGE_EXPERIMENTAL
 	LANDSCAPE_API ULandscapeScratchRenderTarget* GetComponentIdRenderTarget() const;
 	
 	struct FOnRenderBatchTargetGroupDoneParams
@@ -361,26 +356,17 @@ public:
 	};
 	void Render(TFunction<void(const FOnRenderBatchTargetGroupDoneParams& /*InParams*/, UE::Landscape::FRDGBuilderRecorder& /*RDGBuilderRecorder*/)> OnRenderBatchTargetGroupDone);
 
-	LANDSCAPE_EDIT_LAYERS_BATCHED_MERGE_EXPERIMENTAL
 	inline FIntPoint GetMaxNeededResolution() const { return MaxNeededResolution; }
-	LANDSCAPE_EDIT_LAYERS_BATCHED_MERGE_EXPERIMENTAL
 	inline const TArray<FMergeRenderBatch>& GetRenderBatches() const { return RenderBatches; }
-	LANDSCAPE_EDIT_LAYERS_BATCHED_MERGE_EXPERIMENTAL
 	LANDSCAPE_API const FMergeRenderBatch* GetCurrentRenderBatch() const;
-	LANDSCAPE_EDIT_LAYERS_BATCHED_MERGE_EXPERIMENTAL
 	LANDSCAPE_API const int32 GetCurrentRenderBatchIdx() const;
 
-	LANDSCAPE_EDIT_LAYERS_BATCHED_MERGE_EXPERIMENTAL
 	LANDSCAPE_API FTransform ComputeVisualLogTransform(const FTransform& InTransform) const;
-	LANDSCAPE_EDIT_LAYERS_BATCHED_MERGE_EXPERIMENTAL
 	LANDSCAPE_API void IncrementVisualLogOffset();
-	LANDSCAPE_EDIT_LAYERS_BATCHED_MERGE_EXPERIMENTAL
 	LANDSCAPE_API void ResetVisualLogOffset();
 
 #if ENABLE_VISUAL_LOG
-	LANDSCAPE_EDIT_LAYERS_BATCHED_MERGE_EXPERIMENTAL
 	LANDSCAPE_API static int32 GetVisualLogAlpha();
-	LANDSCAPE_EDIT_LAYERS_BATCHED_MERGE_EXPERIMENTAL
 	LANDSCAPE_API bool IsVisualLogEnabled() const;
 #endif // ENABLE_VISUAL_LOG
 
@@ -405,7 +391,6 @@ public:
 	 * @param InBlendParams defines how heightmaps/weightmaps should be blended. In the case of weightmaps, there must be as may entries in WeightmapBlendParams as there are layers in the target layers group
 	 * @param RDGBuilderRecorder recorder to append operations to a single FRDGBuilder, if the the recorder in "recording" mode. In "immediate" mode, a render command will be enqueued immediately
 	 */
-	LANDSCAPE_EDIT_LAYERS_BATCHED_MERGE_EXPERIMENTAL
 	LANDSCAPE_API void GenericBlendLayer(const FBlendParams& InBlendParams, FRenderParams& RenderParams, UE::Landscape::FRDGBuilderRecorder& RDGBuilderRecorder);
 
 	const TBitArray<>& GetFinalTargetLayerBitIndices() const 
@@ -417,7 +402,6 @@ public:
 	 * @return true if there is at least one blackboard item of type T in the context's list 
 	 */
 	template <typename T, typename = std::enable_if_t<std::is_base_of_v<FLandscapeEditLayerMergeRenderBlackboardItem, std::decay_t<T>>>>
-	LANDSCAPE_EDIT_LAYERS_BATCHED_MERGE_EXPERIMENTAL
 	bool HasBlackboardItem() const;
 
 	/** 
@@ -426,41 +410,35 @@ public:
 	 * @return the newly-created blackboard item
 	 */
 	template <typename T, typename... TArgs, typename = std::enable_if_t<std::is_base_of_v<FLandscapeEditLayerMergeRenderBlackboardItem, std::decay_t<T>>>>
-	LANDSCAPE_EDIT_LAYERS_BATCHED_MERGE_EXPERIMENTAL
 	T& AddBlackboardItem(TArgs&&... InArgs);
 
 	/**
 	 * @return the first blackboard item of type T from the context's list or nullptr if there isn't one
 	 */
 	template <typename T, typename = std::enable_if_t<std::is_base_of_v<FLandscapeEditLayerMergeRenderBlackboardItem, std::decay_t<T>>>>
-	LANDSCAPE_EDIT_LAYERS_BATCHED_MERGE_EXPERIMENTAL
 	T* TryGetBlackboardItem();
 
 	/**
 	 * @return the first blackboard item of type T from the context's list. Asserts if there isn't one
 	 */
 	template <typename T, typename = std::enable_if_t<std::is_base_of_v<FLandscapeEditLayerMergeRenderBlackboardItem, std::decay_t<T>>>>
-	LANDSCAPE_EDIT_LAYERS_BATCHED_MERGE_EXPERIMENTAL
 	T& GetBlackboardItem();
 
 	/**
 	 * @return the first blackboard item of type T from the context's list or create a new one if there isn't one
 	 */
 	template <typename T, typename... TArgs, typename = std::enable_if_t<std::is_base_of_v<FLandscapeEditLayerMergeRenderBlackboardItem, std::decay_t<T>>>>
-	LANDSCAPE_EDIT_LAYERS_BATCHED_MERGE_EXPERIMENTAL
 	T& GetOrCreateBlackboardItem(TArgs&&... InArgs);
 
 	/**
 	 * @return all blackboard items of type T from the context's list
 	 */
 	template <typename T, typename = std::enable_if_t<std::is_base_of_v<FLandscapeEditLayerMergeRenderBlackboardItem, std::decay_t<T>>>>
-	LANDSCAPE_EDIT_LAYERS_BATCHED_MERGE_EXPERIMENTAL
 	TArray<T*> GetBlackboardItems();
 
 	/**
 	 * @return all blackboard items currently stored in the context
 	 */
-	LANDSCAPE_EDIT_LAYERS_BATCHED_MERGE_EXPERIMENTAL
 	const TArray<TInstancedStruct<FLandscapeEditLayerMergeRenderBlackboardItem>>& GetAllBlackboardItems() const { return BlackboardItems; }
 
 private:
diff --git a/Engine/Source/Runtime/Landscape/Public/LandscapeEditLayerRenderer.h b/Engine/Source/Runtime/Landscape/Public/LandscapeEditLayerRenderer.h
--- a/Engine/Source/Runtime/Landscape/Public/LandscapeEditLayerRenderer.h
+++ b/Engine/Source/Runtime/Landscape/Public/LandscapeEditLayerRenderer.h
@@ -230,8 +230,7 @@ public:
 	 *
 	 * @return a list of renderer states (i.e. a ILandscapeEditLayerRenderer and its current state) to be processed in that order by the merge operation
 	 */
-	LANDSCAPE_EDIT_LAYERS_BATCHED_MERGE_EXPERIMENTAL
-		virtual TArray<FEditLayerRendererState> GetEditLayerRendererStates(const FMergeContext* InMergeContext) PURE_VIRTUAL(IEditLayerRendererProvider::GetEditLayerRendererStates, return {}; );
+	virtual TArray<FEditLayerRendererState> GetEditLayerRendererStates(const FMergeContext* InMergeContext) PURE_VIRTUAL(IEditLayerRendererProvider::GetEditLayerRendererStates, return {}; );
 };
 
 #endif // WITH_EDITOR
@@ -285,7 +284,6 @@ public:
 	 *  Depending on the other renderer's needs, the final target layer groups might contain more layers than was requested by a given renderer. This only means that more layers will be processed together
 	 *  and if this renderer doesn't act on one of these layers, it will simply do nothing with it in its RenderLayer function
 	 */
-	LANDSCAPE_EDIT_LAYERS_BATCHED_MERGE_EXPERIMENTAL
 	virtual void GetRendererStateInfo(const UE::Landscape::EditLayers::FMergeContext* InMergeContext,
 		UE::Landscape::EditLayers::FEditLayerTargetTypeState& OutSupportedTargetTypeState, UE::Landscape::EditLayers::FEditLayerTargetTypeState& OutEnabledTargetTypeState, 
 		TArray<UE::Landscape::EditLayers::FTargetLayerGroup>& OutTargetLayerGroups) const
@@ -294,7 +292,6 @@ public:
 	/**
 	* @return the a debug name for this renderer
 	*/
-	LANDSCAPE_EDIT_LAYERS_BATCHED_MERGE_EXPERIMENTAL
 	virtual FString GetEditLayerRendererDebugName() const 
 		PURE_VIRTUAL(ILandscapeEditLayerRenderer::GetEditLayerRendererDebugName, return TEXT(""); );
 
@@ -305,7 +302,6 @@ public:
 	 *
 	 * @return list of all render items that affect this renderer
 	*/
-	LANDSCAPE_EDIT_LAYERS_BATCHED_MERGE_EXPERIMENTAL
 	virtual TArray<UE::Landscape::EditLayers::FEditLayerRenderItem> GetRenderItems(const UE::Landscape::EditLayers::FMergeContext* InMergeContext) const
 		PURE_VIRTUAL(ILandscapeEditLayerRenderer::GetRenderItems, return { }; );
 
@@ -337,7 +333,6 @@ public:
 	 * 
 	 * @return true if anything was rendered
 	*/
-	LANDSCAPE_EDIT_LAYERS_BATCHED_MERGE_EXPERIMENTAL
 	virtual bool RenderLayer(UE::Landscape::EditLayers::FRenderParams& RenderParams, UE::Landscape::FRDGBuilderRecorder& RDGBuilderRecorder)
 		PURE_VIRTUAL(ILandscapeEditLayerRenderer::RenderLayer, return false; );
 
@@ -350,7 +345,6 @@ public:
 	 * 
 	 * @param RDGBuilderRecorder (see RenderLayer)
 	*/
-	LANDSCAPE_EDIT_LAYERS_BATCHED_MERGE_EXPERIMENTAL
 	virtual void BlendLayer(UE::Landscape::EditLayers::FRenderParams& RenderParams, UE::Landscape::FRDGBuilderRecorder& RDGBuilderRecorder) {}
 
 	/**
@@ -361,7 +355,6 @@ public:
 	 * 
 	 * @return true if both RenderLayer can be done in the same render layer group
 	*/
-	LANDSCAPE_EDIT_LAYERS_BATCHED_MERGE_EXPERIMENTAL
 	virtual bool CanGroupRenderLayerWith(TScriptInterface<ILandscapeEditLayerRenderer> InOtherRenderer) const { return false; }
 
 	/**
@@ -370,7 +363,6 @@ public:
 	 *
 	 * @param RenderParams contains parameters necessary for this render (affected components, merge context, etc.)
 	 */
-	LANDSCAPE_EDIT_LAYERS_BATCHED_MERGE_EXPERIMENTAL
 	virtual void BeginRenderLayerGroup(UE::Landscape::EditLayers::FRenderParams& RenderParams, UE::Landscape::FRDGBuilderRecorder& RDGBuilderRecorder) {}
 
 	/**
@@ -379,7 +371,6 @@ public:
 	 *
 	 * @param RenderParams contains parameters necessary for this render (affected components, merge context, etc.)
 	 */
-	LANDSCAPE_EDIT_LAYERS_BATCHED_MERGE_EXPERIMENTAL
 	virtual void EndRenderLayerGroup(UE::Landscape::EditLayers::FRenderParams& RenderParams, UE::Landscape::FRDGBuilderRecorder& RDGBuilderRecorder) {}
 
 #endif // WITH_EDITOR
diff --git a/Engine/Source/Runtime/Landscape/Public/LandscapeEditLayerRendererState.h b/Engine/Source/Runtime/Landscape/Public/LandscapeEditLayerRendererState.h
--- a/Engine/Source/Runtime/Landscape/Public/LandscapeEditLayerRendererState.h
+++ b/Engine/Source/Runtime/Landscape/Public/LandscapeEditLayerRendererState.h
@@ -36,85 +36,67 @@ private:
 	FEditLayerRendererState() = default;
 
 public:
-	LANDSCAPE_EDIT_LAYERS_BATCHED_MERGE_EXPERIMENTAL
 	static const FEditLayerRendererState& GetDummyRendererState();
 
-	LANDSCAPE_EDIT_LAYERS_BATCHED_MERGE_EXPERIMENTAL
 	LANDSCAPE_API explicit FEditLayerRendererState(const FMergeContext* InMergeContext);
-	LANDSCAPE_EDIT_LAYERS_BATCHED_MERGE_EXPERIMENTAL
 	LANDSCAPE_API explicit FEditLayerRendererState(const FMergeContext* InMergeContext, TScriptInterface<ILandscapeEditLayerRenderer> InRenderer);
 
 	/** Returns the edit layer renderer which this state relates to */
-	LANDSCAPE_EDIT_LAYERS_BATCHED_MERGE_EXPERIMENTAL
 	TScriptInterface<ILandscapeEditLayerRenderer> GetRenderer() const { return Renderer; }
 
 	/** Returns a mask of all target types / weightmaps supported by this renderer */
-	LANDSCAPE_EDIT_LAYERS_BATCHED_MERGE_EXPERIMENTAL
 	const FEditLayerTargetTypeState& GetSupportedTargetTypeState() const { return SupportedTargetTypeState; }
 
 	/** Returns a mask of all target types / weightmaps enabled on this renderer */
-	LANDSCAPE_EDIT_LAYERS_BATCHED_MERGE_EXPERIMENTAL
 	const FEditLayerTargetTypeState& GetEnabledTargetTypeState() const { return EnabledTargetTypeState; }
 
 	/** Returns a mask of all target types / weightmaps supported and enabled on this renderer */
-	LANDSCAPE_EDIT_LAYERS_BATCHED_MERGE_EXPERIMENTAL
 	FEditLayerTargetTypeState GetActiveTargetTypeState() const { return ActiveTargetTypeState; }
 
 	/** Mutates the EnabledTargetTypeState by adding the target type in parameter to the mask of active target types */
-	LANDSCAPE_EDIT_LAYERS_BATCHED_MERGE_EXPERIMENTAL
 	LANDSCAPE_API void EnableTargetType(ELandscapeToolTargetType InTargetType);
 	
 	/** Mutates the EnabledTargetTypeState by adding the target type mask in parameter to the mask of active target types */
-	LANDSCAPE_EDIT_LAYERS_BATCHED_MERGE_EXPERIMENTAL
 	LANDSCAPE_API void EnableTargetTypeMask(ELandscapeToolTargetTypeFlags InTargetTypeMask);
 
 	/** Mutates the EnabledTargetTypeState by removing the target type in parameter from the mask of active target types */
-	LANDSCAPE_EDIT_LAYERS_BATCHED_MERGE_EXPERIMENTAL
 	LANDSCAPE_API void DisableTargetType(ELandscapeToolTargetType InTargetType);
 
 	/** Mutates the EnabledTargetTypeState by removing the target type mask in parameter from the mask of active target types */
-	LANDSCAPE_EDIT_LAYERS_BATCHED_MERGE_EXPERIMENTAL
 	LANDSCAPE_API void DisableTargetTypeMask(ELandscapeToolTargetTypeFlags InTargetTypeMask);
 
 	/**
 	 * Mutates the EnabledTargetTypeState by adding the weightmap in parameter to the list of enabled weightmaps, if it is known to the merge context
 	 */
-	LANDSCAPE_EDIT_LAYERS_BATCHED_MERGE_EXPERIMENTAL
 	LANDSCAPE_API void EnableWeightmap(ELandscapeToolTargetType InTargetType, const FName& InWeightmapLayerName);
 
 	/**
 	 * Mutates the EnabledTargetTypeState by adding the weightmap in parameter to the list of enabled weightmaps. Asserts if the layer name isn't known to the merge context
 	 */
-	LANDSCAPE_EDIT_LAYERS_BATCHED_MERGE_EXPERIMENTAL
 	LANDSCAPE_API void EnableWeightmapChecked(ELandscapeToolTargetType InTargetType, const FName& InWeightmapLayerName);
 
 	/** 
 	 * Mutates the EnabledTargetTypeState by adding the weightmap in parameter to the list of enabled weightmaps. Asserts if the layer index isn't known to the merge context
 	 */
-	LANDSCAPE_EDIT_LAYERS_BATCHED_MERGE_EXPERIMENTAL
 	LANDSCAPE_API void EnableWeightmap(ELandscapeToolTargetType InTargetType, int32 InWeightmapLayerIndex);
 
 	/**
 	 * Mutates the EnabledTargetTypeState by removing the weightmap in parameter from the list of enabled weightmaps, if it is known to the merge context
 	 */
-	LANDSCAPE_EDIT_LAYERS_BATCHED_MERGE_EXPERIMENTAL
 	LANDSCAPE_API void DisableWeightmap(const FName& InWeightmapLayerName);
 
 	/**
 	 * Mutates the EnabledTargetTypeState by removing the weightmap in parameter from the list of enabled weightmaps. Asserts if the layer index isn't known to the merge context
 	 */
-	LANDSCAPE_EDIT_LAYERS_BATCHED_MERGE_EXPERIMENTAL
 	LANDSCAPE_API void DisableWeightmapChecked(const FName& InWeightmapLayerName);
 
 	/** 
 	 * Mutates the EnabledTargetTypeState by removing the weightmap in parameter from the list of enabled weightmaps. Asserts if the layer index isn't known to the merge context
 	 */
-	LANDSCAPE_EDIT_LAYERS_BATCHED_MERGE_EXPERIMENTAL
 	LANDSCAPE_API void DisableWeightmap(int32 InWeightmapLayerIndex);
 
 	/** Returns the mask of target types that are both supported and enabled by this renderer */
 
-	LANDSCAPE_EDIT_LAYERS_BATCHED_MERGE_EXPERIMENTAL
 	LANDSCAPE_API ELandscapeToolTargetTypeFlags GetActiveTargetTypeMask() const;
 
 	/**
@@ -122,7 +104,6 @@ public:
 	 * @param InTargetType the requested target type (heightmap/weightmap/visibility)
 	 * @param InWeightmapLayerName (optional) is the requested weightmap, only relevant for the ELandscapeToolTargetType::Weightmap case
 	 */
-	LANDSCAPE_EDIT_LAYERS_BATCHED_MERGE_EXPERIMENTAL
 	LANDSCAPE_API bool IsTargetActive(ELandscapeToolTargetType InTargetType, const FName& InWeightmapLayerName = NAME_None) const;
 
 	/**
@@ -130,7 +111,6 @@ public:
 	 * @param InTargetType the requested target type (heightmap/weightmap/visibility)
 	 * @param InWeightmapLayerName (optional) is the requested weightmap, only relevant for the ELandscapeToolTargetType::Weightmap case
 	 */
-	LANDSCAPE_EDIT_LAYERS_BATCHED_MERGE_EXPERIMENTAL
 	LANDSCAPE_API bool IsTargetActiveChecked(ELandscapeToolTargetType InTargetType, const FName& InWeightmapLayerName = NAME_None) const;
 
 	/**
@@ -138,24 +118,20 @@ public:
 	 * @param InTargetType the requested target type (heightmap/weightmap/visibility)
 	 * @param InWeightmapLayerIndex (optional) is the requested weightmap's index in the merge index, only relevant for the ELandscapeToolTargetType::Weightmap case
 	 */
-	LANDSCAPE_EDIT_LAYERS_BATCHED_MERGE_EXPERIMENTAL
 	LANDSCAPE_API bool IsTargetActive(ELandscapeToolTargetType InTargetType, int32 InWeightmapLayerIndex = INDEX_NONE) const;
 
 	/**
 	 * Returns a list of all weightmaps supported and enabled by this renderer (only relevant for ELandscapeToolTargetType::Weightmap (and ELandscapeToolTargetType::Visibility))
 	 */
-	LANDSCAPE_EDIT_LAYERS_BATCHED_MERGE_EXPERIMENTAL
 	LANDSCAPE_API TArray<FName> GetActiveTargetWeightmaps() const;
 
 	/** 
 	 * Returns a list of all weightmaps supported and enabled by this renderer (only relevant for ELandscapeToolTargetType::Weightmap (and ELandscapeToolTargetType::Visibility))
 	 */
-	LANDSCAPE_EDIT_LAYERS_BATCHED_MERGE_EXPERIMENTAL
 	LANDSCAPE_API TBitArray<> GetActiveTargetWeightmapBitIndices() const;
 
 	/** Returns the target layer groups associated with this renderer. A target layer group is a set of target layers (weightmaps) that depend on one another in order to produce the output target layers.
 	 This allows to implement "horizontal blending", where weightmaps can be blended with one another at each step of the landscape edit layers merge algorithm */
-	LANDSCAPE_EDIT_LAYERS_BATCHED_MERGE_EXPERIMENTAL
 	TArray<FTargetLayerGroup> GetTargetLayerGroups() const { return TargetLayerGroups; }
 
 private:
diff --git a/Engine/Source/Runtime/Landscape/Public/LandscapeEditLayerTargetTypeState.h b/Engine/Source/Runtime/Landscape/Public/LandscapeEditLayerTargetTypeState.h
--- a/Engine/Source/Runtime/Landscape/Public/LandscapeEditLayerTargetTypeState.h
+++ b/Engine/Source/Runtime/Landscape/Public/LandscapeEditLayerTargetTypeState.h
@@ -41,7 +41,6 @@ private:
 	FEditLayerTargetTypeState() = default;
 
 public:
-	LANDSCAPE_EDIT_LAYERS_BATCHED_MERGE_EXPERIMENTAL
 	static const FEditLayerTargetTypeState& GetDummyTargetTypeState();
 
 	/**
@@ -49,7 +48,6 @@ public:
 	 * 
 	 * @param InMergeContext context containing all sorts of information related to this merge operation
 	 */
-	LANDSCAPE_EDIT_LAYERS_BATCHED_MERGE_EXPERIMENTAL
 	LANDSCAPE_API FEditLayerTargetTypeState(const FMergeContext* InMergeContext);
 
 	/**
@@ -58,7 +56,6 @@ public:
 	 * @param InMergeContext context containing all sorts of information related to this merge operation
 	 * @param InTargetTypeMask target type mask supported by this target type state (no active weightmaps)
 	 */
-	LANDSCAPE_EDIT_LAYERS_BATCHED_MERGE_EXPERIMENTAL
 	LANDSCAPE_API FEditLayerTargetTypeState(const FMergeContext* InMergeContext, ELandscapeToolTargetTypeFlags InTargetTypeMask);
 
 	/**
@@ -69,7 +66,6 @@ public:
 	 * @param InSupportedWeightmaps names of weightmaps supported by this target type state 
 	 * @param bInChecked, if true, will validate that every weightmap requested is part of the merge context (if false, only the known weightmaps will be marked as supported)
 	 */
-	LANDSCAPE_EDIT_LAYERS_BATCHED_MERGE_EXPERIMENTAL
 	LANDSCAPE_API FEditLayerTargetTypeState(const FMergeContext* InMergeContext, ELandscapeToolTargetTypeFlags InTargetTypeMask, const TArrayView<const FName>& InSupportedWeightmaps, bool bInChecked);
 
 	/**
@@ -80,7 +76,6 @@ public:
 	 * @param InSupportedWeightmaps bit indices of the weightmaps supported by this target type state (corresponds to the FMergeContext's AllTargetLayerNames)
 	 * Asserts if one of the layer indices isn't valid for the merge context (except when InSupportedWeightmapLayerIndices is == TBitArray<>)
 	 */
-	LANDSCAPE_EDIT_LAYERS_BATCHED_MERGE_EXPERIMENTAL
 	LANDSCAPE_API FEditLayerTargetTypeState(const FMergeContext* InMergeContext, ELandscapeToolTargetTypeFlags InTargetTypeMask, const TBitArray<>& InSupportedWeightmapLayerIndices);
 	
 	/**
@@ -88,7 +83,6 @@ public:
 	 * @param InTargetType the requested target type (heightmap/weightmap/visibility)
 	 * @param InWeightmapLayerName (optional) is the requested weightmap, only relevant for the ELandscapeToolTargetType::Weightmap case
 	 */
-	LANDSCAPE_EDIT_LAYERS_BATCHED_MERGE_EXPERIMENTAL
 	LANDSCAPE_API bool IsActive(ELandscapeToolTargetType InTargetType, const FName& InWeightmapLayerName = NAME_None) const;
 
 		/**
@@ -96,7 +90,6 @@ public:
 	 * @param InTargetType the requested target type (heightmap/weightmap/visibility)
 	 * @param InWeightmapLayerName (optional) is the requested weightmap, only relevant for the ELandscapeToolTargetType::Weightmap case. Asserts if the layer name isn't valid for the merge context (except for NAME_None)
 	 */
-	LANDSCAPE_EDIT_LAYERS_BATCHED_MERGE_EXPERIMENTAL
 	LANDSCAPE_API bool IsActiveChecked(ELandscapeToolTargetType InTargetType, const FName& InWeightmapLayerName = NAME_None) const;
 
 	/**
@@ -104,7 +97,6 @@ public:
 	 * @param InTargetType is the requested target type (heightmap/weightmap/visibility)
 	 * @param int32 (optional) is the requested weightmap index (in FMergeContex), only relevant for the ELandscapeToolTargetType::Weightmap case. Asserts if the layer index isn't valid for the merge context (except for INDEX_NONE)
 	 */
-	LANDSCAPE_EDIT_LAYERS_BATCHED_MERGE_EXPERIMENTAL
 	LANDSCAPE_API bool IsActive(ELandscapeToolTargetType InTargetType, int32 InWeightmapLayerIndex = INDEX_NONE) const;
 
 	/** Returns the currently active weightmaps : 
@@ -112,7 +104,6 @@ public:
 	 *  - Additionally, if Visibility is amongst the target types, it will also return the visibility weightmap
 	 *  - If neither Weightmap nor Visibility is amongst the target types, it will return an empty array
 	 */
-	LANDSCAPE_EDIT_LAYERS_BATCHED_MERGE_EXPERIMENTAL
 	LANDSCAPE_API TArray<FName> GetActiveWeightmaps() const;
 
 	/** Returns the currently active weightmaps :
@@ -120,55 +111,42 @@ public:
 	 *  - Additionally, if Visibility is amongst the target types, it will also return the visibility weightmap
 	 *  - If neither Weightmap nor Visibility is amongst the target types, it will return an empty array
 	 */
-	LANDSCAPE_EDIT_LAYERS_BATCHED_MERGE_EXPERIMENTAL
 	LANDSCAPE_API TBitArray<> GetActiveWeightmapBitIndices() const;
 
 	/** Returns the target type mask (i.e. same as ELandscapeToolTargetType, but as bit flags) */
-	LANDSCAPE_EDIT_LAYERS_BATCHED_MERGE_EXPERIMENTAL
 	ELandscapeToolTargetTypeFlags GetTargetTypeMask() const { return TargetTypeMask; }
 
 	/** Sets the target type mask (i.e. same as ELandscapeToolTargetType, but as bit flags) */
-	LANDSCAPE_EDIT_LAYERS_BATCHED_MERGE_EXPERIMENTAL
 	LANDSCAPE_API void SetTargetTypeMask(ELandscapeToolTargetTypeFlags InTargetTypeMask);
 
 	/** Adds the target type in parameter to the mask of active target types */
-	LANDSCAPE_EDIT_LAYERS_BATCHED_MERGE_EXPERIMENTAL
 	LANDSCAPE_API void AddTargetType(ELandscapeToolTargetType InTargetType);
 
 	/** Appends the target type mask in parameter to the mask of active target types */
-	LANDSCAPE_EDIT_LAYERS_BATCHED_MERGE_EXPERIMENTAL
 	LANDSCAPE_API void AddTargetTypeMask(ELandscapeToolTargetTypeFlags InTargetTypeMask);
 
 	/** Removes a single target type from the mask of active target types */
-	LANDSCAPE_EDIT_LAYERS_BATCHED_MERGE_EXPERIMENTAL
 	LANDSCAPE_API void RemoveTargetType(ELandscapeToolTargetType InTargetType);
 
 	/** Removes the target type mask in parameter from the mask of active target types */
-	LANDSCAPE_EDIT_LAYERS_BATCHED_MERGE_EXPERIMENTAL
 	LANDSCAPE_API void RemoveTargetTypeMask(ELandscapeToolTargetTypeFlags InTargetTypeMask);
 
 	/** Adds a weightmap to the list of active weightmaps (make sure ELandscapeToolTargetType::Weightmap is amongst the supported target types) */
-	LANDSCAPE_EDIT_LAYERS_BATCHED_MERGE_EXPERIMENTAL
 	LANDSCAPE_API void AddWeightmap(const FName& InWeightmapLayerName);
 
 	/** Adds a weightmap to the list of active weightmaps (make sure ELandscapeToolTargetType::Weightmap is amongst the supported target types). Asserts if the layer name isn't known to the merge context.  */
-	LANDSCAPE_EDIT_LAYERS_BATCHED_MERGE_EXPERIMENTAL
 	LANDSCAPE_API void AddWeightmapChecked(const FName& InWeightmapLayerName);
 
 	/** Adds a weightmap to the list of active weightmaps (make sure ELandscapeToolTargetType::Weightmap is amongst the supported target types). Asserts if the layer index isn't valid for the merge context */
-	LANDSCAPE_EDIT_LAYERS_BATCHED_MERGE_EXPERIMENTAL
 	LANDSCAPE_API void AddWeightmap(int32 InWeightmapLayerIndex);
 
 	/** Removes a weightmap from the list of active weightmaps */
-	LANDSCAPE_EDIT_LAYERS_BATCHED_MERGE_EXPERIMENTAL
 	LANDSCAPE_API void RemoveWeightmap(const FName& InWeightmapLayerName);
 
 	/** Removes a weightmap from the list of active weightmaps. Asserts if the layer name isn't known to the merge context. */
-	LANDSCAPE_EDIT_LAYERS_BATCHED_MERGE_EXPERIMENTAL
 	LANDSCAPE_API void RemoveWeightmapChecked(const FName& InWeightmapLayerName);
 
 	/** Removes a weightmap from the list of active weightmaps. Asserts if the layer index isn't valid for the merge context. */
-	LANDSCAPE_EDIT_LAYERS_BATCHED_MERGE_EXPERIMENTAL
 	LANDSCAPE_API void RemoveWeightmap(int32 InWeightmapLayerIndex);
 
 	/**
diff --git a/Engine/Source/Runtime/Landscape/Public/LandscapeEditResourcesSubsystem.h b/Engine/Source/Runtime/Landscape/Public/LandscapeEditResourcesSubsystem.h
--- a/Engine/Source/Runtime/Landscape/Public/LandscapeEditResourcesSubsystem.h
+++ b/Engine/Source/Runtime/Landscape/Public/LandscapeEditResourcesSubsystem.h
@@ -25,7 +25,7 @@ namespace UE::Landscape
 {
 class FRDGBuilderRecorder;
 
-struct LANDSCAPE_EDIT_LAYERS_BATCHED_MERGE_EXPERIMENTAL FScratchRenderTargetParams
+struct FScratchRenderTargetParams
 {
 	FScratchRenderTargetParams() = default;
 	explicit FScratchRenderTargetParams(const FString& InDebugName, bool bInExactDimensions, bool bInUseUAV, bool bInTargetArraySlicesIndependently, const FIntPoint& InResolution, int32 InNumSlices, ETextureRenderTargetFormat InFormat, const FLinearColor& InClearColor, ERHIAccess InInitialState)
@@ -60,7 +60,7 @@ struct LANDSCAPE_EDIT_LAYERS_BATCHED_MERGE_EXPERIMENTAL FScratchRenderTargetPara
 	ERHIAccess InitialState = ERHIAccess::None;
 };
 
-struct LANDSCAPE_EDIT_LAYERS_BATCHED_MERGE_EXPERIMENTAL FScratchRenderTargetScope : public FNoncopyable
+struct FScratchRenderTargetScope : public FNoncopyable
 {
 	FScratchRenderTargetScope() = delete;
 	FScratchRenderTargetScope(const FScratchRenderTargetParams& InParams);
@@ -90,41 +90,28 @@ public:
 	ULandscapeScratchRenderTarget();
 
 	/** @return a debug name while this render target is in use. This helps track things down, since the underlying resource can be recycled */
-	LANDSCAPE_EDIT_LAYERS_BATCHED_MERGE_EXPERIMENTAL
 	LANDSCAPE_API const FString& GetDebugName() const;
 	/** @return Getter for the internal UTextureRenderTarget */
 	UTextureRenderTarget* GetRenderTarget() const { return RenderTarget; }
 	/** @return Getter for the internal UTextureRenderTarget2D (only when CurrentRenderTargetParams.NumSlices == 0, asserts if CurrentRenderTargetParams.NumSlices > 0) */
-	LANDSCAPE_EDIT_LAYERS_BATCHED_MERGE_EXPERIMENTAL
 	LANDSCAPE_API UTextureRenderTarget2D* GetRenderTarget2D() const;
 	/** @return the internal UTextureRenderTarget2D (only when CurrentRenderTargetParams.NumSlices == 0, returns nullptr if CurrentRenderTargetParams.NumSlices > 0) */
-	LANDSCAPE_EDIT_LAYERS_BATCHED_MERGE_EXPERIMENTAL
 	LANDSCAPE_API UTextureRenderTarget2D* TryGetRenderTarget2D() const;
 	/**  @return the internal UTextureRenderTarget2DArray (only when CurrentRenderTargetParams.NumSlices > 0, asserts if CurrentRenderTargetParams.NumSlices == 0)  */
-	LANDSCAPE_EDIT_LAYERS_BATCHED_MERGE_EXPERIMENTAL
 	LANDSCAPE_API UTextureRenderTarget2DArray* GetRenderTarget2DArray() const;
 	/** @return the internal UTextureRenderTarget2DArray (only when CurrentRenderTargetParams.NumSlices > 0, returns nullptr if CurrentRenderTargetParams.NumSlices == 0) */
-	LANDSCAPE_EDIT_LAYERS_BATCHED_MERGE_EXPERIMENTAL
 	LANDSCAPE_API UTextureRenderTarget2DArray* TryGetRenderTarget2DArray() const;
 	/** @return the internal render target's true resolution (can be different than the effective resolution if CurrentRenderTargetParams.bExactDimensions is false) */
-	LANDSCAPE_EDIT_LAYERS_BATCHED_MERGE_EXPERIMENTAL
 	LANDSCAPE_API FIntPoint GetResolution() const;
 	/**@return the internal render target's effective resolution (can be different than the actual resolution if CurrentRenderTargetParams.bExactDimensions is false) while this render target is in use */
-	LANDSCAPE_EDIT_LAYERS_BATCHED_MERGE_EXPERIMENTAL
 	LANDSCAPE_API FIntPoint GetEffectiveResolution() const;
 	/** @return the internal render target's number of slices (can be different than the actual number of slices if CurrentRenderTargetParams.bExactDimensions is false) */
-	LANDSCAPE_EDIT_LAYERS_BATCHED_MERGE_EXPERIMENTAL
 	LANDSCAPE_API int32 GetNumSlices() const;
 	/** @return the internal render target's effective number of slices (can be different than the actual number of slices if CurrentRenderTargetParams.bExactDimensions is false) while this render target is in use */
-	LANDSCAPE_EDIT_LAYERS_BATCHED_MERGE_EXPERIMENTAL
 	LANDSCAPE_API int32 GetEffectiveNumSlices() const;
-	LANDSCAPE_EDIT_LAYERS_BATCHED_MERGE_EXPERIMENTAL
 	bool IsTexture2D() const { return (TryGetRenderTarget2D() != nullptr); }
-	LANDSCAPE_EDIT_LAYERS_BATCHED_MERGE_EXPERIMENTAL
 	bool IsTexture2DArray() const { return (TryGetRenderTarget2DArray() != nullptr); }
-	LANDSCAPE_EDIT_LAYERS_BATCHED_MERGE_EXPERIMENTAL
 	LANDSCAPE_API ETextureRenderTargetFormat GetFormat() const;
-	LANDSCAPE_EDIT_LAYERS_BATCHED_MERGE_EXPERIMENTAL
 	LANDSCAPE_API FLinearColor GetClearColor() const;
 	ERHIAccess GetCurrentState() const { return CurrentState; }
 	const UE::Landscape::FScratchRenderTargetParams& GetCurrentRenderTargetParams() { return CurrentRenderTargetParams; }
@@ -198,7 +185,6 @@ public:
 	 *  @param RDGBuilderRecorder when the recorder is in "recording" mode, the copy's render command enqueuing will be delayed to the render thread, after a FRDGBuilder has been created
 	 *   but the scratch render target's current state will still be mutated to reflect the state change on the game thread. In "immediate" mode, a render command will be enqueued immediately
 	 */
-	LANDSCAPE_EDIT_LAYERS_BATCHED_MERGE_EXPERIMENTAL
 	LANDSCAPE_API void CopyFrom(const FCopyFromTextureParams& InCopyParams, UE::Landscape::FRDGBuilderRecorder& RDGBuilderRecorder);
 
 	/**
@@ -207,7 +193,6 @@ public:
 	 * @param RDGBuilderRecorder when the recorder is in "recording" mode, the copy's render command enqueuing will be delayed to the render thread, after a FRDGBuilder has been created
 	*  but the scratch render targets' current state will still be mutated to reflect the state changes on the game thread. In "immediate" mode, a render command will be enqueued immediately
 	 */
-	LANDSCAPE_EDIT_LAYERS_BATCHED_MERGE_EXPERIMENTAL
 	LANDSCAPE_API void CopyFrom(const FCopyFromScratchRenderTargetParams& InCopyParams, UE::Landscape::FRDGBuilderRecorder& RDGBuilderRecorder);
 
 	/** 
@@ -216,7 +201,6 @@ public:
 	*  @param RDGBuilderRecorder  when the recorder is in "recording" mode, the transition's render command enqueuing will not be executed (we assume the FRDGBuilder will take care of it)
 	 *   but the scratch render target's current state will still be mutated to reflect the state change on the game thread. In "immediate" mode, a render command will be enqueued immediately
 	 */
-	LANDSCAPE_EDIT_LAYERS_BATCHED_MERGE_EXPERIMENTAL
 	LANDSCAPE_API void TransitionTo(ERHIAccess InStateAfter, UE::Landscape::FRDGBuilderRecorder& RDGBuilderRecorder);
 
 	/** 
@@ -224,7 +208,6 @@ public:
 	 *  @param RDGBuilderRecorder  when the recorder is in "recording" mode, the clear's render command enqueuing will be delayed to the render thread, after a FRDGBuilder has been created
 	 *   but the scratch render target's current state will still be mutated to reflect the state change on the game thread. In "immediate" mode, a render command will be enqueued immediately
 	 */
-	LANDSCAPE_EDIT_LAYERS_BATCHED_MERGE_EXPERIMENTAL
 	LANDSCAPE_API void Clear(UE::Landscape::FRDGBuilderRecorder& RDGBuilderRecorder);
 
 	/** @return true if this scratch texture is compatible (and can therefore be used) with the requested render target params */
diff --git a/Engine/Source/Runtime/Landscape/Public/LandscapeEditTypes.h b/Engine/Source/Runtime/Landscape/Public/LandscapeEditTypes.h
--- a/Engine/Source/Runtime/Landscape/Public/LandscapeEditTypes.h
+++ b/Engine/Source/Runtime/Landscape/Public/LandscapeEditTypes.h
@@ -7,12 +7,6 @@
 
 #include "LandscapeEditTypes.generated.h"
 
-// APIs marked with LANDSCAPE_EDIT_LAYERS_BATCHED_MERGE_EXPERIMENTAL are experimental and therefore purposefully marked as UE_EXPERIMENTAL : 
-//  These APIs are part of a large refactor of the landscape edit layers system (batched merge) that will eventually replace the previous global 
-//  and local merge code paths but is still optional and therefore considered experimental for the time being. Since the API will likely evolve 
-//  until then, use at your own risk.
-#define LANDSCAPE_EDIT_LAYERS_BATCHED_MERGE_EXPERIMENTAL UE_EXPERIMENTAL(5.5, "Landscape edit layers batched merge is still being actively developed, this API may change.")
-
 UENUM()
 enum class ELandscapeToolTargetType : uint8
 {
