Removed landscape "paint-time weight balancing" for the legacy weight blending layers : this was misleading as it was trying to "invent" weight values per-landscape component, which could lead to weird painting behavior and is not necessary anymore since edit layers are now non-optional, and weight blending is applied as the final edit layers merge step, leading to normalized weights all the time

diff --git a/Engine/Source/Editor/LandscapeEditor/Private/LandscapeEdMode.cpp b/Engine/Source/Editor/LandscapeEditor/Private/LandscapeEdMode.cpp
--- a/Engine/Source/Editor/LandscapeEditor/Private/LandscapeEdMode.cpp
+++ b/Engine/Source/Editor/LandscapeEditor/Private/LandscapeEdMode.cpp
@@ -3405,7 +3405,7 @@ void FEdModeLandscape::ImportWeightData(ULandscapeInfo* LandscapeInfo, const FGu
 
 			FScopedTransaction Transaction(LOCTEXT("Undo_ImportWeightmap", "Importing Landscape Layer"));
 
-			FAlphamapAccessor<false, false> AlphamapAccessor(LandscapeInfo, LayerInfo);
+			TAlphamapAccessor</*bInUseInterp = */false> AlphamapAccessor(LandscapeInfo, LayerInfo);
 			AlphamapAccessor.SetData(MinX, MinY, MaxX, MaxY, Data.GetData(), PaintRestriction);
 		});
 	}
@@ -4341,7 +4341,7 @@ bool LandscapeEditorUtils::SetWeightmapData(ALandscapeProxy* Landscape, ULandsca
 
 	if (Data.Num() == (1 + ComponentsRect.Width())*(1 + ComponentsRect.Height()))
 	{
-		FAlphamapAccessor<false, false> AlphamapAccessor(Landscape->GetLandscapeInfo(), LayerObject);
+		TAlphamapAccessor</*bInUseInterp = */false> AlphamapAccessor(Landscape->GetLandscapeInfo(), LayerObject);
 		AlphamapAccessor.SetData(ComponentsRect.Min.X, ComponentsRect.Min.Y, ComponentsRect.Max.X, ComponentsRect.Max.Y, Data.GetData(), ELandscapeLayerPaintingRestriction::None);
 		return true;
 	}
diff --git a/Engine/Source/Editor/LandscapeEditor/Private/LandscapeEdModeMirrorTool.cpp b/Engine/Source/Editor/LandscapeEditor/Private/LandscapeEdModeMirrorTool.cpp
--- a/Engine/Source/Editor/LandscapeEditor/Private/LandscapeEdModeMirrorTool.cpp
+++ b/Engine/Source/Editor/LandscapeEditor/Private/LandscapeEdModeMirrorTool.cpp
@@ -634,7 +634,7 @@ public:
 				TempMaxY = SourceMaxY;
 				LandscapeEdit.GetWeightData(LayerInfo, TempMinX, TempMinY, TempMaxX, TempMaxY, &SourceWeightData[0], SourceSizeX);
 				ApplyMirrorInternal<uint8>(SourceWeightData, DestWeightData, SourceSizeX, SourceSizeY, DestSizeX, DestSizeY, MirrorPos, BlendWidth);
-				LandscapeEdit.SetAlphaData(LayerInfo, DestMinX, DestMinY, DestMaxX, DestMaxY, &DestWeightData[0], DestSizeX, ELandscapeLayerPaintingRestriction::None, false, false);
+				LandscapeEdit.SetAlphaData(LayerInfo, DestMinX, DestMinY, DestMaxX, DestMaxY, &DestWeightData[0], DestSizeX, ELandscapeLayerPaintingRestriction::None);
 				//LayerInfo->IsReferencedFromLoadedData = true;
 			}
 		}
diff --git a/Engine/Source/Editor/LandscapeEditor/Private/LandscapeEdModeTools.h b/Engine/Source/Editor/LandscapeEditor/Private/LandscapeEdModeTools.h
--- a/Engine/Source/Editor/LandscapeEditor/Private/LandscapeEdModeTools.h
+++ b/Engine/Source/Editor/LandscapeEditor/Private/LandscapeEdModeTools.h
@@ -767,35 +767,35 @@ struct FLandscapeHeightCache : public TLandscapeEditCache<FHeightmapAccessorTool
 	}
 };
 
-template<bool bInUseInterp, bool bInUseTotalNormalize>
-struct FAlphamapAccessorTool : public FAlphamapAccessor<bInUseInterp, bInUseTotalNormalize>
+template<bool bInUseInterp>
+struct TAlphamapAccessorTool : public TAlphamapAccessor<bInUseInterp>
 {
-	FAlphamapAccessorTool(ULandscapeInfo* InLandscapeInfo, ULandscapeLayerInfoObject* InLayerInfo)
-	:	FAlphamapAccessor<bInUseInterp, bInUseTotalNormalize>(InLandscapeInfo, InLayerInfo)
+	TAlphamapAccessorTool(ULandscapeInfo* InLandscapeInfo, ULandscapeLayerInfoObject* InLayerInfo)
+		: TAlphamapAccessor<bInUseInterp>(InLandscapeInfo, InLayerInfo)
 	{}
 
-	FAlphamapAccessorTool(const FLandscapeToolTarget& InTarget)
-	:	FAlphamapAccessor<bInUseInterp, bInUseTotalNormalize>(InTarget.LandscapeInfo.Get(), InTarget.LayerInfo.Get())
+	TAlphamapAccessorTool(const FLandscapeToolTarget& InTarget)
+		: TAlphamapAccessor<bInUseInterp>(InTarget.LandscapeInfo.Get(), InTarget.LayerInfo.Get())
 	{
 	}
 };
 
-struct FLandscapeAlphaCache : public TLandscapeEditCache<FAlphamapAccessorTool<true, false>, uint8>
+struct FLandscapeAlphaCache : public TLandscapeEditCache<TAlphamapAccessorTool</*bInUseInterp = */true>, uint8>
 {
 	static uint8 ClampValue(int32 Value) { return static_cast<uint8>(FMath::Clamp(Value, 0, 255)); }
 
 	FLandscapeAlphaCache(const FLandscapeToolTarget& InTarget)
-		: TLandscapeEditCache<FAlphamapAccessorTool<true, false>, uint8>(InTarget)
+		: TLandscapeEditCache<TAlphamapAccessorTool</*bInUseInterp = */true>, uint8>(InTarget)
 	{
 	}
 };
 
-struct FLandscapeVisCache : public TLandscapeEditCache<FAlphamapAccessorTool<false, false>, uint8>
+struct FLandscapeVisCache : public TLandscapeEditCache<TAlphamapAccessorTool</*bInUseInterp = */false>, uint8>
 {
 	static uint8 ClampValue(int32 Value) { return static_cast<uint8>(FMath::Clamp(Value, 0, 255)); }
 
 	FLandscapeVisCache(const FLandscapeToolTarget& InTarget)
-		: TLandscapeEditCache<FAlphamapAccessorTool<false, false>, uint8>(InTarget)
+		: TLandscapeEditCache<TAlphamapAccessorTool</*bInUseInterp = */false>, uint8>(InTarget)
 	{
 	}
 };
@@ -982,7 +982,7 @@ struct FFullWeightmapAccessor
 		TSet<ULandscapeComponent*> Components;
 		if (LandscapeEdit.GetComponentsInRegion(X1, Y1, X2, Y2, &Components))
 		{
-			LandscapeEdit.SetAlphaData(DirtyLayerInfos, X1, Y1, X2, Y2, Data, 0, PaintingRestriction);
+			LandscapeEdit.SetAlphaData(DirtyLayerInfos, X1, Y1, X2, Y2, Data, /*Stride = */0, PaintingRestriction);
 		}
 		DirtyLayerInfos.Empty();
 	}
diff --git a/Engine/Source/Editor/LandscapeEditor/Private/LandscapeEditorObject.cpp b/Engine/Source/Editor/LandscapeEditor/Private/LandscapeEditorObject.cpp
--- a/Engine/Source/Editor/LandscapeEditor/Private/LandscapeEditorObject.cpp
+++ b/Engine/Source/Editor/LandscapeEditor/Private/LandscapeEditorObject.cpp
@@ -827,18 +827,6 @@ ELandscapeImportResult ULandscapeEditorObject::CreateNewLayersInfo(TArray<FLands
 		OutNewLayerInfos.Add(MoveTemp(ImportLayer));
 	}
 
-	// Fill the first weight-blended layer to 100%
-	// TODO [jonathan.bard] : remove all this once we get rid of non-edit layers, weight balancing at the data level is not reliable and weight blending with edit layers active 
-	//  is providing the same service in a non-destructive way
-	if (FLandscapeImportLayerInfo* FirstBlendedLayer = OutNewLayerInfos.FindByPredicate([](const FLandscapeImportLayerInfo& ImportLayer) { return ImportLayer.LayerInfo && (ImportLayer.LayerInfo->GetBlendMethod() == ELandscapeTargetLayerBlendMethod::FinalWeightBlending); }))
-	{
-		const int32 DataSize = SizeX * SizeY;
-		FirstBlendedLayer->LayerData.AddUninitialized(DataSize);
-
-		uint8* ByteData = FirstBlendedLayer->LayerData.GetData();
-		FMemory::Memset(ByteData, 255, DataSize);
-	}
-
 	return ELandscapeImportResult::Success;
 }
 
diff --git a/Engine/Source/Runtime/Landscape/Private/LandscapeConfigHelper.cpp b/Engine/Source/Runtime/Landscape/Private/LandscapeConfigHelper.cpp
--- a/Engine/Source/Runtime/Landscape/Private/LandscapeConfigHelper.cpp
+++ b/Engine/Source/Runtime/Landscape/Private/LandscapeConfigHelper.cpp
@@ -390,7 +390,7 @@ void FLandscapeConfigHelper::CopyRegionToComponent(ULandscapeInfo* InLandscapeIn
 		{
 			CopyData(Layer.LayerData, NewWeightData, InRegion, NewRegion, bInResample);
 			FLandscapeEditDataInterface LandscapeEdit(NewLandscapeInfo);
-			LandscapeEdit.SetAlphaData(Layer.LayerInfo, NewRegion.Min.X, NewRegion.Min.Y, NewRegion.Max.X, NewRegion.Max.Y, NewWeightData.GetData(), 0);
+			LandscapeEdit.SetAlphaData(Layer.LayerInfo, NewRegion.Min.X, NewRegion.Min.Y, NewRegion.Max.X, NewRegion.Max.Y, NewWeightData.GetData(), /*Stride = */0);
 		}
 	}
 }
diff --git a/Engine/Source/Runtime/Landscape/Private/LandscapeEdit.cpp b/Engine/Source/Runtime/Landscape/Private/LandscapeEdit.cpp
--- a/Engine/Source/Runtime/Landscape/Private/LandscapeEdit.cpp
+++ b/Engine/Source/Runtime/Landscape/Private/LandscapeEdit.cpp
@@ -569,9 +569,7 @@ FString ULandscapeComponent::GetLayerAllocationKey(const TArray<FWeightmapLayerA
 	{
 		if (Allocations[LayerIdx].GetLayerName() != NAME_None) // ignore layers that have no name
 		{
-			// TODO [jonathan.bard] : remove from layer allocation key : this doesn't make any difference in the shader anymore since mobile landscape mesh was removed (see bAtLeastOneWeightBasedBlend in UE4.27)
-			const bool bNoWeightBlend = Allocations[LayerIdx].LayerInfo && (Allocations[LayerIdx].LayerInfo->GetBlendMethod() != ELandscapeTargetLayerBlendMethod::FinalWeightBlending);
-			LayerStrings.Add(FString::Printf(TEXT("_%s_%s%d"), *Allocations[LayerIdx].GetLayerName().ToString(), bNoWeightBlend ? TEXT("n") : TEXT("w"), Allocations[LayerIdx].WeightmapTextureIndex));
+			LayerStrings.Add(FString::Printf(TEXT("_%s%d"), *Allocations[LayerIdx].GetLayerName().ToString(), Allocations[LayerIdx].WeightmapTextureIndex));
 		}
 	}
 	// Sort them alphabetically so we can share across components even if the order is different
@@ -3348,8 +3346,6 @@ void ALandscapeProxy::Import(const FGuid& InGuid, int32 InMinX, int32 InMinY, in
 			WeightValues.Empty(EditingAlphaLayerData.Num());
 			WeightValues.AddZeroed(EditingAlphaLayerData.Num());
 			ComponentWeightmapLayerAllocations.Empty(EditingAlphaLayerData.Num());
-			// TODO [jonathan.bard] : remove all this once we get rid of non-edit layers, weight balancing at the data level is not reliable and weight blending with edit layers active 
-			//  is providing the same service in a non-destructive way
 			TArray<bool, TInlineAllocator<16>> IsNoBlendArray;
 			IsNoBlendArray.Empty(EditingAlphaLayerData.Num());
 			IsNoBlendArray.AddZeroed(EditingAlphaLayerData.Num());
@@ -3359,12 +3355,13 @@ void ALandscapeProxy::Import(const FGuid& InGuid, int32 InMinX, int32 InMinY, in
 				// Lookup the original layer name
 				WeightValues[WeightLayerIndex] = EditingAlphaLayerData[WeightLayerIndex].AlphaValues;
 				new(ComponentWeightmapLayerAllocations) FWeightmapLayerAllocationInfo(ImportLayerInfos[EditingAlphaLayerData[WeightLayerIndex].LayerIndex].LayerInfo);
-				IsNoBlendArray[WeightLayerIndex] = (ImportLayerInfos[EditingAlphaLayerData[WeightLayerIndex].LayerIndex].LayerInfo->GetBlendMethod() != ELandscapeTargetLayerBlendMethod::FinalWeightBlending);
+				IsNoBlendArray[WeightLayerIndex] = (ImportLayerInfos[EditingAlphaLayerData[WeightLayerIndex].LayerIndex].LayerInfo->GetBlendMethod() == ELandscapeTargetLayerBlendMethod::None);
 			}
 
 			// Discard the temporary alpha data
 			EditingAlphaLayerData.Empty();
-
+			
+			// Handle the special import method where each alpha map specifies the reminder from previous layers : apply it only to weight-blended layers
 			if (InImportMaterialLayerType == ELandscapeImportAlphamapType::Layered)
 			{
 				// For each layer...
@@ -3402,51 +3399,6 @@ void ALandscapeProxy::Import(const FGuid& InGuid, int32 InMinX, int32 InMinY, in
 					}
 				}
 			}
-
-			// Weight normalization for total should be 255...
-			if (WeightValues.Num())
-			{
-				for (int32 Idx = 0; Idx < WeightValues[0].Num(); Idx++)
-				{
-					int32 TotalWeight = 0;
-					int32 MaxLayerIdx = -1;
-					int32 MaxWeight = INT_MIN;
-
-					for (int32 WeightLayerIndex = 0; WeightLayerIndex < WeightValues.Num(); WeightLayerIndex++)
-					{
-						if (!IsNoBlendArray[WeightLayerIndex])
-						{
-							int32 Weight = WeightValues[WeightLayerIndex][Idx];
-							TotalWeight += Weight;
-							if (MaxWeight < Weight)
-							{
-								MaxWeight = Weight;
-								MaxLayerIdx = WeightLayerIndex;
-							}
-						}
-					}
-
-					if (TotalWeight > 0 && TotalWeight != 255)
-					{
-						// normalization...
-						float Factor = 255.0f / TotalWeight;
-						TotalWeight = 0;
-						for (int32 WeightLayerIndex = 0; WeightLayerIndex < WeightValues.Num(); WeightLayerIndex++)
-						{
-							if (!IsNoBlendArray[WeightLayerIndex])
-							{
-								WeightValues[WeightLayerIndex][Idx] = (uint8)(Factor * WeightValues[WeightLayerIndex][Idx]);
-								TotalWeight += WeightValues[WeightLayerIndex][Idx];
-							}
-						}
-
-						if (255 - TotalWeight && MaxLayerIdx >= 0)
-						{
-							WeightValues[MaxLayerIdx][Idx] += static_cast<uint8>(255 - TotalWeight);
-						}
-					}
-				}
-			}
 		}
 	}
 
@@ -3888,7 +3840,7 @@ void ALandscapeProxy::Import(const FGuid& InGuid, int32 InMinX, int32 InMinY, in
 			{
 				if (MaterialLayerInfo.LayerInfo != nullptr && MaterialLayerInfo.LayerData.Num() > 0)
 				{
-					LandscapeEdit.SetAlphaData(MaterialLayerInfo.LayerInfo, InMinX, InMinY, InMaxX, InMaxY, MaterialLayerInfo.LayerData.GetData(), 0, ELandscapeLayerPaintingRestriction::None, true, false);
+					LandscapeEdit.SetAlphaData(MaterialLayerInfo.LayerInfo, InMinX, InMinY, InMaxX, InMaxY, MaterialLayerInfo.LayerData.GetData(), /*Stride = */0, ELandscapeLayerPaintingRestriction::None);
 
 					// Add to TargetLayers if needed.
 					if (!HasTargetLayer(MaterialLayerInfo.LayerInfo->LayerName))
@@ -7727,11 +7679,6 @@ void ULandscapeComponent::GenerateMobileWeightmapLayerAllocations()
 		if (LhsLayerInfo == ALandscapeProxy::VisibilityLayer && RhsLayerInfo != ALandscapeProxy::VisibilityLayer) return true;
 		if (RhsLayerInfo == ALandscapeProxy::VisibilityLayer && LhsLayerInfo != ALandscapeProxy::VisibilityLayer) return false;
 
-		// TODO [jonathan.bard] : this doesn't make any difference in the shader anymore since mobile landscape mesh was removed (see bAtLeastOneWeightBasedBlend in UE4.27)
-		// Sort non-weight blended layers to the front so if we have exactly 3 layers, the 3rd is definitely weight-based.
-		if ((LhsLayerInfo->GetBlendMethod() != ELandscapeTargetLayerBlendMethod::FinalWeightBlending) && (RhsLayerInfo->GetBlendMethod() == ELandscapeTargetLayerBlendMethod::FinalWeightBlending)) return true;
-		if ((RhsLayerInfo->GetBlendMethod() != ELandscapeTargetLayerBlendMethod::FinalWeightBlending) && (LhsLayerInfo->GetBlendMethod() == ELandscapeTargetLayerBlendMethod::FinalWeightBlending)) return false;
-
 		return false; // equal, preserve order
 	}));
 }
@@ -8423,7 +8370,7 @@ bool ALandscapeProxy::LandscapeImportWeightmapFromRenderTarget(UTextureRenderTar
 
 				FScopedTransaction Transaction(LOCTEXT("Undo_ImportWeightmap", "Importing Landscape Layer"));
 
-				FAlphamapAccessor<false, false> AlphamapAccessor(LandscapeInfo, CurWeightmapInfo.LayerInfoObj);\
+				TAlphamapAccessor</*bInUseInterp = */false> AlphamapAccessor(LandscapeInfo, CurWeightmapInfo.LayerInfoObj);
 				AlphamapAccessor.SetEditLayer(EditLayerGuid);
 				AlphamapAccessor.SetData(MinX, MinY, MaxX, MaxY, LayerData.GetData(), ELandscapeLayerPaintingRestriction::None);
 
diff --git a/Engine/Source/Runtime/Landscape/Private/LandscapeEditInterface.cpp b/Engine/Source/Runtime/Landscape/Private/LandscapeEditInterface.cpp
--- a/Engine/Source/Runtime/Landscape/Private/LandscapeEditInterface.cpp
+++ b/Engine/Source/Runtime/Landscape/Private/LandscapeEditInterface.cpp
@@ -1500,113 +1500,6 @@ void ULandscapeComponent::DeleteLayerInternal(ULandscapeLayerInfoObject* LayerIn
 
 	Component->DeleteLayerAllocation(InEditLayerGuid, DeleteAllocIdx, LandscapeEdit.GetShouldDirtyPackage());
 
-	// See if the deleted layer is a NoWeightBlend layer - if not, we don't have to worry about normalization
-	// If the layer doesn't exist, assume it is a blended layer (so renormalization will run whether needed or not)
-	// TODO [jonathan.bard] : remove all this once we get rid of non-edit layers, weight balancing at the data level is not reliable and weight blending with edit layers active 
-	//  is providing the same service in a non-destructive way
-	if ((LayerInfo == nullptr) || (LayerInfo->GetBlendMethod() == ELandscapeTargetLayerBlendMethod::FinalWeightBlending))
-	{
-		// Lock data for all the weightmaps
-		TArray<FLandscapeTextureDataInfo*> TexDataInfos;
-		for (int32 WeightmapIdx = 0; WeightmapIdx < ComponentWeightmapTextures.Num(); WeightmapIdx++)
-		{
-			TexDataInfos.Add(LandscapeEdit.GetTextureDataInfo(ComponentWeightmapTextures[WeightmapIdx]));
-		}
-
-		TArray<bool> LayerNoWeightBlends;	// Array of NoWeightBlend flags
-		TArray<uint8*> LayerDataPtrs;		// Pointers to all layers' data 
-
-		// Get the data for each layer
-		for (int32 LayerIdx = 0; LayerIdx < ComponentWeightmapLayerAllocations.Num(); LayerIdx++)
-		{
-			const FWeightmapLayerAllocationInfo& Allocation = ComponentWeightmapLayerAllocations[LayerIdx];
-			LayerDataPtrs.Add((uint8*)TexDataInfos[Allocation.WeightmapTextureIndex]->GetMipData(0) + ChannelOffsets[Allocation.WeightmapTextureChannel]);
-
-			// Find the layer info and record if it is a "no weight blend" layer.
-			const bool bPerformLegacyWeightBalancing = (Allocation.LayerInfo == nullptr) || (Allocation.LayerInfo->GetBlendMethod() == ELandscapeTargetLayerBlendMethod::FinalWeightBlending);
-			LayerNoWeightBlends.Add(!bPerformLegacyWeightBalancing);
-		}
-
-		// Find the texture data corresponding to this vertex
-		const int32 SizeU = (SubsectionSizeQuads + 1) * NumSubsections;
-		const int32 SizeV = (SubsectionSizeQuads + 1) * NumSubsections;
-		const int32 WeightmapOffsetX = static_cast<int32>(Component->WeightmapScaleBias.Z * SizeU);
-		const int32 WeightmapOffsetY = static_cast<int32>(Component->WeightmapScaleBias.W * SizeV);
-
-		for (int32 SubIndexY = 0; SubIndexY < NumSubsections; SubIndexY++)
-		{
-			for (int32 SubIndexX = 0; SubIndexX < NumSubsections; SubIndexX++)
-			{
-				for (int32 SubY = 0; SubY <= SubsectionSizeQuads; SubY++)
-				{
-					for (int32 SubX = 0; SubX <= SubsectionSizeQuads; SubX++)
-					{
-						const int32 TexX = WeightmapOffsetX + (SubsectionSizeQuads + 1) * SubIndexX + SubX;
-						const int32 TexY = WeightmapOffsetY + (SubsectionSizeQuads + 1) * SubIndexY + SubY;
-						const int32 TexDataIndex = 4 * (TexX + TexY * SizeU);
-
-						// Calculate the sum of other layer weights
-						int32 OtherLayerWeightSum = 0;
-						for (int32 LayerIdx = 0; LayerIdx < ComponentWeightmapLayerAllocations.Num(); LayerIdx++)
-						{
-							if (LayerIdx != DeleteAllocIdx && LayerNoWeightBlends[LayerIdx] == false)
-							{
-								OtherLayerWeightSum += LayerDataPtrs[LayerIdx][TexDataIndex];
-							}
-						}
-
-						if (OtherLayerWeightSum == 0)
-						{
-							// Set the first other weight-blend layer we can find to 255 to avoid a black hole
-							// This isn't ideal but it's the best option
-							// There's nothing we can easily do if this was the only weight-blend layer on this component
-							for (int32 LayerIdx = 0; LayerIdx < ComponentWeightmapLayerAllocations.Num(); LayerIdx++)
-							{
-								if (LayerIdx != DeleteAllocIdx && LayerNoWeightBlends[LayerIdx] == false)
-								{
-									uint8& Weight = LayerDataPtrs[LayerIdx][TexDataIndex];
-									Weight = 255;
-									break;
-								}
-							}
-						}
-						else
-						{
-							// Adjust other layer weights
-							for (int32 LayerIdx = 0; LayerIdx < ComponentWeightmapLayerAllocations.Num(); LayerIdx++)
-							{
-								if (LayerIdx != DeleteAllocIdx && LayerNoWeightBlends[LayerIdx] == false)
-								{
-									uint8& Weight = LayerDataPtrs[LayerIdx][TexDataIndex];
-									Weight = static_cast<uint8>(FMath::Clamp<int32>(FMath::RoundToInt32(255.0f * (float)Weight / (float)OtherLayerWeightSum), 0, 255));
-								}
-							}
-						}
-					}
-				}
-			}
-		}
-
-		// Update all the textures and mips
-		for (int32 Idx = 0; Idx < ComponentWeightmapTextures.Num(); Idx++)
-		{
-			UTexture2D* WeightmapTexture = ComponentWeightmapTextures[Idx];
-			FLandscapeTextureDataInfo* WeightmapDataInfo = TexDataInfos[Idx];
-
-			const int32 NumMips = WeightmapTexture->Source.GetNumMips();
-			TArray<FColor*> WeightmapTextureMipData;
-			WeightmapTextureMipData.AddUninitialized(NumMips);
-			for (int32 MipIdx = 0; MipIdx < NumMips; MipIdx++)
-			{
-				WeightmapTextureMipData[MipIdx] = (FColor*)WeightmapDataInfo->GetMipData(MipIdx);
-			}
-
-			ULandscapeComponent::UpdateWeightmapMips(Component->NumSubsections, Component->SubsectionSizeQuads, WeightmapTexture, WeightmapTextureMipData, 0, 0, MAX_int32, MAX_int32, WeightmapDataInfo);
-
-			WeightmapDataInfo->AddMipUpdateRegion(0, 0, 0, WeightmapTexture->Source.GetSizeX() - 1, WeightmapTexture->Source.GetSizeY() - 1);
-		}
-	}
-
 	Component->RequestWeightmapUpdate();
 }
 
@@ -1643,17 +1536,15 @@ void ULandscapeComponent::FillLayer(ULandscapeLayerInfoObject* LayerInfo, FLands
 
 	FGuid EditLayerGuid = LandscapeEdit.GetEditLayer();
 
+	// TODO [jonathan.bard] : This should always be false now, as we should always use this function in the context of edit layers. Check that it's the case
+	//  and look for all places where we use LandscapeEdit.GetEditLayer() (and check that the result !IsValid(), these should all go. We might even be able to 
+	//  get rid of FScopedSetLandscapeEditingLayer at long last!)
 	const bool bIsFinalWeightmap = !EditLayerGuid.IsValid();
 
 	ALandscapeProxy* Proxy = Component->GetLandscapeProxy();
 	Component->Modify(LandscapeEdit.GetShouldDirtyPackage());
 	Proxy->Modify(LandscapeEdit.GetShouldDirtyPackage());
 
-	// TODO [jonathan.bard] : remove all this once we get rid of non-edit layers, weight balancing at the data level is not reliable and weight blending with edit layers active 
-	//  is providing the same service in a non-destructive way
-	const bool bPerformLegacyWeightBalancing = (LayerInfo->GetBlendMethod() == ELandscapeTargetLayerBlendMethod::FinalWeightBlending);
-	bool bClearOtherWeightBlendLayers = bPerformLegacyWeightBalancing;
-
 	TArray<FWeightmapLayerAllocationInfo>& ComponentWeightmapLayerAllocations = Component->GetWeightmapLayerAllocations(EditLayerGuid);
 	TArray<TObjectPtr<UTexture2D>>& ComponentWeightmapTextures = Component->GetWeightmapTextures(EditLayerGuid);
 	TArray<TObjectPtr<ULandscapeWeightmapUsage>>& ComponentWeightmapTexturesUsage = Component->GetWeightmapTexturesUsage(EditLayerGuid);
@@ -1662,26 +1553,6 @@ void ULandscapeComponent::FillLayer(ULandscapeLayerInfoObject* LayerInfo, FLands
 	int32 FillLayerIdx = ComponentWeightmapLayerAllocations.IndexOfByPredicate(
 		[LayerInfo](const FWeightmapLayerAllocationInfo& Allocation) { return Allocation.LayerInfo == LayerInfo; });
 
-	// if the layer isn't used on this component yet but is a weight-blend layer, then simply steal the allocation of another weight-blend layer!
-	if (FillLayerIdx == INDEX_NONE && bPerformLegacyWeightBalancing)
-	{
-		FillLayerIdx = ComponentWeightmapLayerAllocations.IndexOfByPredicate(
-			[](const FWeightmapLayerAllocationInfo& Allocation) 
-			{ 
-				return (Allocation.LayerInfo == nullptr) || (Allocation.LayerInfo->GetBlendMethod() == ELandscapeTargetLayerBlendMethod::FinalWeightBlending);
-			});
-
-		if (FillLayerIdx != INDEX_NONE)
-		{
-			ComponentWeightmapLayerAllocations[FillLayerIdx].LayerInfo = LayerInfo;
-		}
-		else
-		{
-			// no other weight-blend layers exist
-			bClearOtherWeightBlendLayers = false;
-		}
-	}
-
 	// if the layer is still not found then we are forced to make a new allocation
 	if (FillLayerIdx == INDEX_NONE)
 	{
@@ -1714,61 +1585,6 @@ void ULandscapeComponent::FillLayer(ULandscapeLayerInfoObject* LayerInfo, FLands
 		}
 	}
 
-	// clear other layers
-	if (bClearOtherWeightBlendLayers)
-	{
-		for (int32 LayerIdx = 0; LayerIdx < ComponentWeightmapLayerAllocations.Num(); ++LayerIdx)
-		{
-			if (LayerIdx == FillLayerIdx)
-			{
-				continue;
-			}
-			FWeightmapLayerAllocationInfo& Allocation = ComponentWeightmapLayerAllocations[LayerIdx];
-			if ((Allocation.LayerInfo == nullptr) || (Allocation.LayerInfo->GetBlendMethod() != ELandscapeTargetLayerBlendMethod::FinalWeightBlending))
-			{
-				continue;
-			}
-
-			ULandscapeWeightmapUsage* Usage = ComponentWeightmapTexturesUsage[Allocation.WeightmapTextureIndex];
-			if (Usage) // can be null if WeightmapUsageMap hasn't been built yet
-			{
-				Usage->ChannelUsage[Allocation.WeightmapTextureChannel] = nullptr;
-			}
-
-			Allocation.Free();
-		}
-
-		ComponentWeightmapLayerAllocations.RemoveAll(
-			[](const FWeightmapLayerAllocationInfo& Allocation) { return !Allocation.IsAllocated(); });
-
-		// remove any textures we're no longer using
-		for (int32 TextureIdx = 0; TextureIdx < ComponentWeightmapTextures.Num(); ++TextureIdx)
-		{
-			if (!ComponentWeightmapLayerAllocations.ContainsByPredicate(
-				[TextureIdx](const FWeightmapLayerAllocationInfo& Allocation) { return Allocation.WeightmapTextureIndex == TextureIdx; }))
-			{
-				ComponentWeightmapTextures[TextureIdx]->Modify(LandscapeEdit.GetShouldDirtyPackage());
-
-				for (int32 LayerIdx = 0; LayerIdx < ComponentWeightmapLayerAllocations.Num(); ++LayerIdx)
-				{
-					FWeightmapLayerAllocationInfo& Allocation = ComponentWeightmapLayerAllocations[LayerIdx];
-					if (Allocation.WeightmapTextureIndex > TextureIdx)
-					{
-						--Allocation.WeightmapTextureIndex;
-					}
-				}
-
-				ComponentWeightmapTextures.RemoveAt(TextureIdx);
-				if (ComponentWeightmapTexturesUsage.IsValidIndex(TextureIdx))
-				{
-					ComponentWeightmapTexturesUsage.RemoveAt(TextureIdx);
-				}
-
-				TextureIdx--;
-			}
-		}
-	}
-
 	// todo - normalize texture usage: it's possible to end up with two textures each using one channel at this point
 	// e.g. if you start with 4 blended layers (in one texture) and a non-blended layer (in a 2nd), and fill one weight-blended layer (deleting the other three)
 	// this can also happen with normal painting I believe
@@ -1930,14 +1746,6 @@ void ULandscapeComponent::ReplaceLayerInternal(ULandscapeLayerInfoObject* FromLa
 		}
 	}
 
-	// See if the deleted layer is a NoWeightBlend layer - if not, we don't have to worry about normalization
-	const bool bFromLayerIsNoWeightBlend = (FromLayerInfo && (FromLayerInfo->GetBlendMethod() != ELandscapeTargetLayerBlendMethod::FinalWeightBlending));
-	const bool bToLayerIsNoWeightBlend = (ToLayerInfo && (ToLayerInfo->GetBlendMethod() != ELandscapeTargetLayerBlendMethod::FinalWeightBlending));
-
-	// TODO [jonathan.bard] : remove all this once we get rid of non-edit layers, weight balancing at the data level is not reliable and weight blending with edit layers active 
-	//  is providing the same service in a non-destructive way
-	const bool bRequireNormalization = (bFromLayerIsNoWeightBlend != bToLayerIsNoWeightBlend);
-
 	if (bMerging)
 	{
 		FWeightmapLayerAllocationInfo& ToLayerAllocation = ComponentWeightmapLayerAllocations[ToLayerIdx];
@@ -1997,11 +1805,6 @@ void ULandscapeComponent::ReplaceLayerInternal(ULandscapeLayerInfoObject* FromLa
 		}
 	}
 
-	if (bRequireNormalization)
-	{
-		// TODO
-	}
-
 	// if merging into an existing layer, remove the layer and potentially the texture
 	if (bMerging)
 	{
@@ -2302,170 +2105,13 @@ inline bool FLandscapeEditDataInterface::IsLayerAllowed(const ULandscapeLayerInf
 	return true;
 }
 
-inline TMap<const ULandscapeLayerInfoObject*, uint32> FLandscapeEditDataInterface::CountWeightBlendedLayerInfluence(const int32 ComponentIndexX, const int32 ComponentIndexY, TOptional<TArrayView<const uint8* const>> InOptionalLayerDataPtrs)
-{
-	// the counts should easily fit in a uint32, a 255x255 x2x2 Component with weights of all 255 only totals 26 bits
-	checkSlow(FMath::CeilLogTwo(ComponentSizeQuads + 1) * 2 + 8 /*ceillog2(255)*/ <= 32);
-	TMap<const ULandscapeLayerInfoObject*, uint32> LayerInfluenceMap;
-
-	ULandscapeComponent* Component = LandscapeInfo->XYtoComponentMap.FindChecked(FIntPoint(ComponentIndexX,ComponentIndexY));
-	const TArray<FWeightmapLayerAllocationInfo>& ComponentWeightmapLayerAllocations = Component->GetWeightmapLayerAllocations(GetEditLayer());
-	const TArray<UTexture2D*>& ComponentWeightmapTextures = Component->GetWeightmapTextures(GetEditLayer());
-
-	// used if InOptionalLayerDataPtrs is null
-	TArray<FLandscapeTextureDataInfo*, TInlineAllocator<2>> InternalTexDataInfos;
-	TArray<const uint8*, TInlineAllocator<8>> InternalLayerDataPtrs;
-	TArrayView<const uint8* const> LayerDataPtrs;
-	if (InOptionalLayerDataPtrs)
-	{
-		check(InOptionalLayerDataPtrs->Num() == ComponentWeightmapLayerAllocations.Num());
-		LayerDataPtrs = InOptionalLayerDataPtrs.GetValue();
-	}
-	else
-	{
-		InternalTexDataInfos.AddUninitialized(ComponentWeightmapTextures.Num());
-		for (int32 WeightmapIdx = 0; WeightmapIdx < ComponentWeightmapTextures.Num(); ++WeightmapIdx)
-		{
-			InternalTexDataInfos[WeightmapIdx] = GetTextureDataInfo(ComponentWeightmapTextures[WeightmapIdx]);
-		}
-
-		InternalLayerDataPtrs.AddUninitialized(ComponentWeightmapLayerAllocations.Num());
-		for (int32 LayerIdx = 0; LayerIdx < ComponentWeightmapLayerAllocations.Num(); LayerIdx++)
-		{
-			const FWeightmapLayerAllocationInfo& Allocation = ComponentWeightmapLayerAllocations[LayerIdx];
-			InternalLayerDataPtrs[LayerIdx] = (uint8*)InternalTexDataInfos[Allocation.WeightmapTextureIndex]->GetMipData(0) + ChannelOffsets[Allocation.WeightmapTextureChannel];
-		}
-
-		LayerDataPtrs = InternalLayerDataPtrs;
-	}
-
-	const int32 ScanlineSize = (SubsectionSizeQuads + 1) * ComponentNumSubsections * 4;
-
-	for (int32 LayerIdx = 0; LayerIdx < ComponentWeightmapLayerAllocations.Num(); LayerIdx++)
-	{
-		const FWeightmapLayerAllocationInfo& Allocation = ComponentWeightmapLayerAllocations[LayerIdx];
-		if ((Allocation.LayerInfo == nullptr) || (Allocation.LayerInfo->GetBlendMethod() != ELandscapeTargetLayerBlendMethod::FinalWeightBlending))
-		{
-			continue;
-		}
-		auto& Count = LayerInfluenceMap.Add(Allocation.LayerInfo, 0);
-
-		for (int32 SubIndexY = 0; SubIndexY < ComponentNumSubsections; ++SubIndexY)
-		{
-			const int32 YStart = SubIndexY * (SubsectionSizeQuads + 1);
-			const int32 YEnd = YStart + (SubsectionSizeQuads + 1);
-			for (int32 Y = YStart; Y < YEnd; ++Y)
-			{
-				for (int32 SubIndexX = 0; SubIndexX < ComponentNumSubsections; ++SubIndexX)
-				{
-					const int32 XStart = SubIndexX * (SubsectionSizeQuads + 1);
-					const int32 XEnd = XStart + (SubsectionSizeQuads + 1);
-					for (int32 X = XStart; X < XEnd; ++X)
-					{
-						const int32 TexDataIndex = Y * ScanlineSize + X * 4;
-						const uint8 Weight = LayerDataPtrs[LayerIdx][TexDataIndex];
-						Count += Weight;
-					}
-				}
-			}
-		}
-	}
-
-	LayerInfluenceMap.ValueSort(TGreater<uint32>());
-	return LayerInfluenceMap;
-}
-
-const ULandscapeLayerInfoObject* FLandscapeEditDataInterface::ChooseReplacementLayer(const ULandscapeLayerInfoObject* const LayerInfo, const int32 ComponentIndexX, const int32 SubIndexX, const int32 SubX, const int32 ComponentIndexY, const int32 SubIndexY, const int32 SubY, TMap<FIntPoint, TMap<const ULandscapeLayerInfoObject*, uint32>>& LayerInfluenceCache, TArrayView<const uint8* const> LayerDataPtrs)
-{
-	const TMap<const ULandscapeLayerInfoObject*, uint32>* LayerInfluenceMapCacheEntry = LayerInfluenceCache.Find(FIntPoint(ComponentIndexX, ComponentIndexY));
-	if (!LayerInfluenceMapCacheEntry)
-	{
-		LayerInfluenceMapCacheEntry = &LayerInfluenceCache.Add(FIntPoint(ComponentIndexX, ComponentIndexY), CountWeightBlendedLayerInfluence(ComponentIndexX, ComponentIndexY, LayerDataPtrs));
-	}
-
-	if (!(SubIndexX == 0 && SubX == 0) &&
-		!(SubIndexX == ComponentNumSubsections - 1 && SubX == SubsectionSizeQuads) &&
-		!(SubIndexY == 0 && SubY == 0) &&
-		!(SubIndexY == ComponentNumSubsections - 1 && SubY == SubsectionSizeQuads))
-	{
-		for (const auto& LayerInfluenceMapPair : *LayerInfluenceMapCacheEntry)
-		{
-			if (LayerInfluenceMapPair.Key != LayerInfo)
-			{
-				return LayerInfluenceMapPair.Key;
-			}
-		}
-		return nullptr;
-	}
-
-	TMap<const ULandscapeLayerInfoObject*, uint32, TInlineSetAllocator<8>> LayerInfluenceMap = *LayerInfluenceMapCacheEntry;
-
-	const int32 ComponentXStart = (SubIndexX == 0 && SubX == 0) ? ComponentIndexX - 1 : ComponentIndexX;
-	const int32 ComponentXEnd = (SubIndexX == ComponentNumSubsections - 1 && SubX == SubsectionSizeQuads) ? ComponentIndexX + 1 : ComponentIndexX;
-	const int32 ComponentYStart = (SubIndexY == 0 && SubY == 0) ? ComponentIndexY - 1 : ComponentIndexY;
-	const int32 ComponentYEnd = (SubIndexY == ComponentNumSubsections - 1 && SubY == SubsectionSizeQuads) ? ComponentIndexY + 1 : ComponentIndexY;
-	for (int32 Y = ComponentYStart; Y <= ComponentYEnd; ++Y)
-	{
-		for (int32 X = ComponentXStart; X <= ComponentXEnd; ++X)
-		{
-			if (X == ComponentIndexX && Y == ComponentIndexY)
-			{
-				// skip the current component, it is already included above
-				continue;
-			}
-
-			ULandscapeComponent* Component = LandscapeInfo->XYtoComponentMap.FindRef(FIntPoint(X, Y));
-			if (!Component)
-			{
-				// skip missing components
-				continue;
-			}
-
-			const TMap<const ULandscapeLayerInfoObject*, uint32>* OtherLayerInfluenceMapCacheEntry = LayerInfluenceCache.Find(FIntPoint(X, Y));
-			if (!OtherLayerInfluenceMapCacheEntry)
-			{
-				OtherLayerInfluenceMapCacheEntry = &LayerInfluenceCache.Add(FIntPoint(X, Y), CountWeightBlendedLayerInfluence(X, Y, {}));
-			}
-
-			for (auto LayerInfluenceMapIt = LayerInfluenceMap.CreateIterator(); LayerInfluenceMapIt; ++LayerInfluenceMapIt)
-			{
-				const uint32* Value = OtherLayerInfluenceMapCacheEntry->Find(LayerInfluenceMapIt->Key);
-				if (Value)
-				{
-					LayerInfluenceMapIt->Value += *Value;
-				}
-				else
-				{
-					// only allow layers that exist in *all* the touched components
-					LayerInfluenceMapIt.RemoveCurrent();
-				}
-			}
-		}
-	}
-
-	LayerInfluenceMap.ValueSort(TGreater<uint32>());
-	for (const auto& LayerInfluenceMapPair : LayerInfluenceMap)
-	{
-		if (LayerInfluenceMapPair.Key != LayerInfo)
-		{
-			return LayerInfluenceMapPair.Key;
-		}
-	}
-	return nullptr;
-}
-
-void FLandscapeEditDataInterface::SetAlphaData(ULandscapeLayerInfoObject* const LayerInfo, const int32 X1, const int32 Y1, const int32 X2, const int32 Y2, const uint8* Data, int32 Stride, ELandscapeLayerPaintingRestriction PaintingRestriction /*= None*/, bool bWeightAdjust /*= true*/, bool bTotalWeightAdjust /*= false*/)
+void FLandscapeEditDataInterface::SetAlphaData(ULandscapeLayerInfoObject* const LayerInfo, const int32 X1, const int32 Y1, const int32 X2, const int32 Y2, const uint8* Data, int32 Stride, ELandscapeLayerPaintingRestriction PaintingRestriction /*= None*/)
 {
 	TRACE_CPUPROFILER_EVENT_SCOPE(LandscapeEditDataInterface_SetAlphaData);
 
 	check(LayerInfo != nullptr);
 	check(Data != nullptr);
 
-	if ((LayerInfo == nullptr) || (LayerInfo->GetBlendMethod() != ELandscapeTargetLayerBlendMethod::FinalWeightBlending))
-	{
-		bWeightAdjust = false;
-	}
-
 	if (Stride == 0)
 	{
 		Stride = (1+X2-X1);
@@ -2480,14 +2126,11 @@ void FLandscapeEditDataInterface::SetAlphaData(ULandscapeLayerInfoObject* const
 
 	TArray<FLandscapeTextureDataInfo*, TInlineAllocator<2>> TexDataInfos;
 	TArray<uint8*, TInlineAllocator<8>> LayerDataPtrs;      // Pointers to all layers' data
-	TArray<bool, TInlineAllocator<8>> LayerNoWeightBlends;  // NoWeightBlend flags
 	TArray<bool, TInlineAllocator<8>> LayerEditDataAllZero; // Whether the data we are editing for this layer is all zero
 	TArray<FColor*> CollisionWeightmapMipData;
 	TArray<FColor*> SimpleCollisionWeightmapMipData;
 	TArray<FColor*> WeightmapTextureMipData;
 
-	TMap<FIntPoint, TMap<const ULandscapeLayerInfoObject*, uint32>> LayerInfluenceCache;
-
 	for (int32 ComponentIndexY = ComponentIndexY1; ComponentIndexY <= ComponentIndexY2; ComponentIndexY++)
 	{
 		for (int32 ComponentIndexX = ComponentIndexX1; ComponentIndexX <= ComponentIndexX2; ComponentIndexX++)
@@ -2602,8 +2245,6 @@ void FLandscapeEditDataInterface::SetAlphaData(ULandscapeLayerInfoObject* const
 
 			LayerDataPtrs.Reset();        // Pointers to all layers' data
 			LayerDataPtrs.AddUninitialized(ComponentWeightmapLayerAllocations.Num());
-			LayerNoWeightBlends.Reset();  // NoWeightBlend flags
-			LayerNoWeightBlends.AddUninitialized(ComponentWeightmapLayerAllocations.Num());
 			LayerEditDataAllZero.Reset(); // Whether the data we are editing for this layer is all zero 
 			LayerEditDataAllZero.AddUninitialized(ComponentWeightmapLayerAllocations.Num());
 
@@ -2614,13 +2255,11 @@ void FLandscapeEditDataInterface::SetAlphaData(ULandscapeLayerInfoObject* const
 				if (Allocation.LayerInfo != nullptr) // only take into account valid layer
 				{
 					LayerDataPtrs[LayerIdx] = (uint8*)TexDataInfos[Allocation.WeightmapTextureIndex]->GetMipData(0) + ChannelOffsets[Allocation.WeightmapTextureChannel];
-					LayerNoWeightBlends[LayerIdx] = (Allocation.LayerInfo->GetBlendMethod() != ELandscapeTargetLayerBlendMethod::FinalWeightBlending);
 					LayerEditDataAllZero[LayerIdx] = true;
 				}
 				else
 				{
 					LayerDataPtrs[LayerIdx] = nullptr;
-					LayerNoWeightBlends[LayerIdx] = true;
 					LayerEditDataAllZero[LayerIdx] = false;
 				}
 			}
@@ -2683,164 +2322,12 @@ void FLandscapeEditDataInterface::SetAlphaData(ULandscapeLayerInfoObject* const
 								}
 							}
 
-							// Adjust all layer weights
-							// (bWeightAdjust implies that this is a weight-blended layer)
-							int32 OtherLayerWeightSum = 0;
-							if (bWeightAdjust)
-							{
-								// Normalize all layers including the painted one
-								// TODO [jonathan.bard] : this isn't used. TODO: Remove
-								if (bTotalWeightAdjust)
-								{
-									int32 MaxLayerIdx = -1;
-									int32 MaxWeight = INT_MIN;
-
-									// Adjust other layers' weights accordingly
-									for (int32 LayerIdx = 0; LayerIdx < ComponentWeightmapLayerAllocations.Num(); LayerIdx++)
-									{
-										if (LayerDataPtrs[LayerIdx] == nullptr)
-										{
-											continue;
-										}
-
-										uint8& ExistingWeight = LayerDataPtrs[LayerIdx][TexDataIndex];
-
-										if (LayerIdx == UpdateLayerIdx)
-										{
-											ExistingWeight = NewWeight;
-										}
-										// Exclude bNoWeightBlend layers
-										if (LayerNoWeightBlends[LayerIdx] == false)
-										{
-											OtherLayerWeightSum += ExistingWeight;
-											if (MaxWeight < ExistingWeight)
-											{
-												MaxWeight = ExistingWeight;
-												MaxLayerIdx = LayerIdx;
-											}
-										}
-									}
-
-									if (OtherLayerWeightSum > 0 && OtherLayerWeightSum != 255)
-									{
-										const float Factor = 255.0f / OtherLayerWeightSum;
-										OtherLayerWeightSum = 0;
-
-										// Normalize
-										for (int32 LayerIdx = 0; LayerIdx < ComponentWeightmapLayerAllocations.Num(); LayerIdx++)
-										{
-											if (LayerDataPtrs[LayerIdx] == nullptr)
-											{
-												continue;
-											}
-
-											uint8& ExistingWeight = LayerDataPtrs[LayerIdx][TexDataIndex];
-
-											if (LayerNoWeightBlends[LayerIdx] == false)
-											{
-												// normalization...
-												ExistingWeight = (uint8)(Factor * ExistingWeight);
-												OtherLayerWeightSum += ExistingWeight;
-
-												if (ExistingWeight != 0)
-												{
-													LayerEditDataAllZero[LayerIdx] = false;
-												}
-											}
-										}
-
-										if ((255 - OtherLayerWeightSum) && MaxLayerIdx >= 0)
-										{
-											// No need to check for nullptr here because MaxLayerIdx can only be set to a valid layer
-											LayerDataPtrs[MaxLayerIdx][TexDataIndex] += static_cast<uint8>(255 - OtherLayerWeightSum);
-										}
-									}
-								}
-								else
-								{
-									// Adjust other layers' weights accordingly
-									for (int32 LayerIdx = 0; LayerIdx < ComponentWeightmapLayerAllocations.Num(); LayerIdx++)
-									{
-										// Exclude bNoWeightBlend layers 
-										if (LayerIdx != UpdateLayerIdx && LayerNoWeightBlends[LayerIdx] == false)
-										{
-											// No need to check for nullptr here because invalid layers have LayerNoWeightBlends set to true.
-											OtherLayerWeightSum += LayerDataPtrs.IsValidIndex(LayerIdx) ? LayerDataPtrs[LayerIdx][TexDataIndex] : 0;
-										}
-									}
-
-									if (OtherLayerWeightSum == 0 && NewWeight < 255)
-									{
-										if (NewWeight < CurrentWeight)
-										{
-											// When reducing the layer weight from 255, we need to choose another layer to fill to avoid a black hole
-											const ULandscapeLayerInfoObject* ReplacementLayer = ChooseReplacementLayer(LayerInfo, ComponentIndexX, SubIndexX, SubX, ComponentIndexY, SubIndexY, SubY, LayerInfluenceCache, LayerDataPtrs);
-											if (ReplacementLayer)
-											{
-												const int32 ReplacementLayerIndex = ComponentWeightmapLayerAllocations.IndexOfByPredicate([&](const FWeightmapLayerAllocationInfo& AllocationInfo) { return AllocationInfo.LayerInfo == ReplacementLayer; });
-
-												// No need to check for nullptr here because ChooseReplacementLayer can't return an invalid layer.
-												LayerDataPtrs[ReplacementLayerIndex][TexDataIndex] = 255 - NewWeight;
-												LayerEditDataAllZero[ReplacementLayerIndex] = false;
-											}
-											else
-											{
-												// if we didn't find a suitable replacement we just have to leave it at 255, unfortunately
-												NewWeight = 255;
-											}
-										}
-										else if (NewWeight > CurrentWeight)
-										{
-											// if weight is increasing on a black spot then go straight to 255
-											NewWeight = 255;
-										}
-
-										// No need to check for nullptr here because UpdateLayerIdx is always valid
-										LayerDataPtrs[UpdateLayerIdx][TexDataIndex] = NewWeight;
-									}
-									else
-									{
-										for (int32 LayerIdx = 0; LayerIdx < ComponentWeightmapLayerAllocations.Num(); LayerIdx++)
-										{
-											if (LayerDataPtrs[LayerIdx] == nullptr)
-											{
-												continue;
-											}
-
-											uint8& Weight = LayerDataPtrs[LayerIdx][TexDataIndex];
-
-											if (LayerIdx == UpdateLayerIdx)
-											{
-												Weight = NewWeight;
-											}
-											else
-											{
-												// Exclude bNoWeightBlend layers
-												if (LayerNoWeightBlends[LayerIdx] == false)
-												{
-													Weight = static_cast<uint8>(FMath::Clamp(
-														FMath::RoundToInt((float)(255 - NewWeight) * (float)Weight / (float)OtherLayerWeightSum), 0, 255));
-												}
-											}
-
-											if (Weight != 0)
-											{
-												LayerEditDataAllZero[LayerIdx] = false;
-											}
-										}
-									}
-								}
-							}
-							else
+							// No need to check for nullptr here because UpdateLayerIdx is always valid
+							uint8& Weight = LayerDataPtrs[UpdateLayerIdx][TexDataIndex];
+							Weight = NewWeight;
+							if (Weight != 0)
 							{
-								// No need to check for nullptr here because UpdateLayerIdx is always valid
-								// Weight value set without adjusting other layers' weights
-								uint8& Weight = LayerDataPtrs[UpdateLayerIdx][TexDataIndex];
-								Weight = NewWeight;
-								if (Weight != 0)
-								{
-									LayerEditDataAllZero[UpdateLayerIdx] = false;
-								}
+								LayerEditDataAllZero[UpdateLayerIdx] = false;
 							}
 						}
 					}
@@ -2889,6 +2376,12 @@ void FLandscapeEditDataInterface::SetAlphaData(ULandscapeLayerInfoObject* const
 	}
 }
 
+// Deprecated
+void FLandscapeEditDataInterface::SetAlphaData(ULandscapeLayerInfoObject* const LayerInfo, const int32 X1, const int32 Y1, const int32 X2, const int32 Y2, const uint8* Data, int32 Stride, ELandscapeLayerPaintingRestriction PaintingRestriction /*= None*/, bool bWeightAdjust /*= true*/, bool bTotalWeightAdjust /*= false*/)
+{
+	SetAlphaData(LayerInfo, X1, Y1, X2, Y2, Data, Stride, PaintingRestriction);
+}
+
 void FLandscapeEditDataInterface::SetAlphaData(const TSet<ULandscapeLayerInfoObject*>& DirtyLayerInfos, const int32 X1, const int32 Y1, const int32 X2, const int32 Y2, const uint8* Data, int32 Stride, ELandscapeLayerPaintingRestriction PaintingRestriction /*= None*/)
 {
 	if (DirtyLayerInfos.Num() == 0)
@@ -3059,8 +2552,6 @@ void FLandscapeEditDataInterface::SetAlphaData(const TSet<ULandscapeLayerInfoObj
 
 							const int32 TexDataIndex = 4 * (TexX + TexY * TexSize);
 
-							int32 OtherLayerWeightSum = 0;
-
 							// Apply weights to all layers simultaneously.
 							for (int32 LayerIdx = 0; LayerIdx < ComponentWeightmapLayerAllocations.Num(); LayerIdx++)
 							{
diff --git a/Engine/Source/Runtime/Landscape/Private/LandscapeSplineRaster.cpp b/Engine/Source/Runtime/Landscape/Private/LandscapeSplineRaster.cpp
--- a/Engine/Source/Runtime/Landscape/Private/LandscapeSplineRaster.cpp
+++ b/Engine/Source/Runtime/Landscape/Private/LandscapeSplineRaster.cpp
@@ -448,8 +448,7 @@ void RasterizeControlPointAlpha(int32& MinX, int32& MinY, int32& MaxX, int32& Ma
 		Rasterizer.DrawTriangle(FalloffRight0, Left1, FalloffLeft1, FalloffRight0Pos, Left1Pos, FalloffLeft1Pos, false);
 	}
 
-	LandscapeEdit.SetAlphaData(LayerInfo, MinX, MinY, MaxX, MaxY, Data.GetData(), 0, ELandscapeLayerPaintingRestriction::None, 
-		/*bWeightAdjust = */(LayerInfo->GetBlendMethod() == ELandscapeTargetLayerBlendMethod::FinalWeightBlending), /*bTotalWeightAdjust = */false);
+	LandscapeEdit.SetAlphaData(LayerInfo, MinX, MinY, MaxX, MaxY, Data.GetData(), /*Stride = */0, ELandscapeLayerPaintingRestriction::None);
 
 	LandscapeEdit.GetComponentsInRegion(MinX, MinY, MaxX, MaxY, &ModifiedComponents);
 }
@@ -572,8 +571,7 @@ void RasterizeSegmentAlpha(int32& MinX, int32& MinY, int32& MaxX, int32& MaxY, F
 		Rasterizer.DrawTriangle(FalloffRight0, Right1, FalloffRight1, FalloffRight0Pos, Right1Pos, FalloffRight1Pos, false);
 	}
 
-	LandscapeEdit.SetAlphaData(LayerInfo, MinX, MinY, MaxX, MaxY, Data.GetData(), 0, ELandscapeLayerPaintingRestriction::None, 
-		/*bWeightAdjust = */(LayerInfo->GetBlendMethod() == ELandscapeTargetLayerBlendMethod::FinalWeightBlending), /*bTotalWeightAdjust = */false);
+	LandscapeEdit.SetAlphaData(LayerInfo, MinX, MinY, MaxX, MaxY, Data.GetData(), /*Stride = */0, ELandscapeLayerPaintingRestriction::None);
 
 	LandscapeEdit.GetComponentsInRegion(MinX, MinY, MaxX, MaxY, &ModifiedComponents);
 }
diff --git a/Engine/Source/Runtime/Landscape/Public/LandscapeEdit.h b/Engine/Source/Runtime/Landscape/Public/LandscapeEdit.h
--- a/Engine/Source/Runtime/Landscape/Public/LandscapeEdit.h
+++ b/Engine/Source/Runtime/Landscape/Public/LandscapeEdit.h
@@ -230,8 +230,9 @@ struct FLandscapeEditDataInterface : public FLandscapeTextureDataInterface
 	LANDSCAPE_API void GetWeightDataFast(ULandscapeLayerInfoObject* LayerInfo, const int32 X1, const int32 Y1, const int32 X2, const int32 Y2, TMap<FIntPoint, uint8>& SparseData);
 	LANDSCAPE_API void GetWeightDataFast(ULandscapeLayerInfoObject* LayerInfo, const int32 X1, const int32 Y1, const int32 X2, const int32 Y2, TMap<FIntPoint, TArray<uint8>>& SparseData);
 	// Updates weightmap for LayerInfo, optionally adjusting all other weightmaps.
-	// TODO [jonathan.bard] : Remove bWeightAdjust (and bTotalWeightAdjust, which is unused) as soon as we remove support for non-edit layers
-	LANDSCAPE_API void SetAlphaData(ULandscapeLayerInfoObject* LayerInfo, const int32 X1, const int32 Y1, const int32 X2, const int32 Y2, const uint8* Data, int32 Stride, ELandscapeLayerPaintingRestriction PaintingRestriction = ELandscapeLayerPaintingRestriction::None, bool bWeightAdjust = true, bool bTotalWeightAdjust = false);
+	LANDSCAPE_API void SetAlphaData(ULandscapeLayerInfoObject* LayerInfo, const int32 X1, const int32 Y1, const int32 X2, const int32 Y2, const uint8* Data, int32 Stride, ELandscapeLayerPaintingRestriction PaintingRestriction = ELandscapeLayerPaintingRestriction::None);
+	UE_DEPRECATED(5.7, "Support for bWeightAdjust and bTotalWeightAdjust have been removed, this function shall be used anymore")
+	LANDSCAPE_API void SetAlphaData(ULandscapeLayerInfoObject* LayerInfo, const int32 X1, const int32 Y1, const int32 X2, const int32 Y2, const uint8* Data, int32 Stride, ELandscapeLayerPaintingRestriction PaintingRestriction, bool bWeightAdjust, bool bTotalWeightAdjust);
 	// Updates weightmaps for all layers. Data points to packed data for all layers in the landscape info
 	LANDSCAPE_API void SetAlphaData(const TSet<ULandscapeLayerInfoObject*>& DirtyLayerInfos, const int32 X1, const int32 Y1, const int32 X2, const int32 Y2, const uint8* Data, int32 Stride, ELandscapeLayerPaintingRestriction PaintingRestriction = ELandscapeLayerPaintingRestriction::None);
 	// Delete a layer and re-normalize other layers
@@ -335,14 +336,6 @@ private:
 	inline bool IsLayerAllowed(const ULandscapeLayerInfoObject* LayerInfo,
 	                          int32 ComponentIndexX, int32 SubIndexX, int32 SubX,
 	                          int32 ComponentIndexY, int32 SubIndexY, int32 SubY);
-
-	// counts the total influence of each weight-blended layer on this component
-	// TODO [jonathan.bard] : Remove when we get rid of non-edit layers (this is only used by ChooseReplacementLayer, which is only used by SetAlphaData, when bWeightAdjust == true)
-	inline TMap<const ULandscapeLayerInfoObject*, uint32> CountWeightBlendedLayerInfluence(int32 ComponentIndexX, int32 ComponentIndexY, TOptional<TArrayView<const uint8* const>> LayerDataPtrs);
-
-	// chooses a replacement layer to use when erasing from 100% influence on a texel
-	// TODO [jonathan.bard] : Remove when we get rid of non-edit layers (this is only used by SetAlphaData, when bWeightAdjust == true)
-	const ULandscapeLayerInfoObject* ChooseReplacementLayer(const ULandscapeLayerInfoObject* LayerInfo, int32 ComponentIndexX, int32 SubIndexX, int32 SubX, int32 ComponentIndexY, int32 SubIndexY, int32 SubY, TMap<FIntPoint, TMap<const ULandscapeLayerInfoObject*, uint32>>& LayerInfluenceCache, TArrayView<const uint8* const> LayerDataPtrs);
 };
 
 struct FLandscapeDoNotDirtyScope
@@ -492,32 +485,17 @@ private:
 };
 
 //
-// FAlphamapAccessor
+// TAlphamapAccessor
 //
-template<bool bInUseInterp, bool bInUseTotalNormalize>
-struct FAlphamapAccessor
+template<bool bInUseInterp>
+struct TAlphamapAccessor
 {
 	enum { bUseInterp = bInUseInterp };
-	enum { bUseTotalNormalize = bInUseTotalNormalize };
-	FAlphamapAccessor(ULandscapeInfo* InLandscapeInfo, ULandscapeLayerInfoObject* InLayerInfo)
+	TAlphamapAccessor(ULandscapeInfo* InLandscapeInfo, ULandscapeLayerInfoObject* InLayerInfo)
 		: LandscapeInfo(InLandscapeInfo)
 		, LandscapeEdit(InLandscapeInfo)
 		, LayerInfo(InLayerInfo)
-		, bPerformLegacyWeightBalancing(true)
-	{
-		// should be no Layer change during FAlphamapAccessor lifetime...
-		if (InLandscapeInfo && InLayerInfo)
-		{
-			if (LayerInfo == ALandscapeProxy::VisibilityLayer)
-			{
-				bPerformLegacyWeightBalancing = false;
-			}
-			else
-			{
-				bPerformLegacyWeightBalancing = (LayerInfo->GetBlendMethod() == ELandscapeTargetLayerBlendMethod::FinalWeightBlending);
-			}
-		}
-	}
+	{}
 
 	void GetData(int32& X1, int32& Y1, int32& X2, int32& Y2, TMap<FIntPoint, uint8>& Data)
 	{
@@ -556,7 +534,7 @@ struct FAlphamapAccessor
 				LandscapeComponent->RequestWeightmapUpdate();
 			}
 			
-			LandscapeEdit.SetAlphaData(LayerInfo, X1, Y1, X2, Y2, Data, 0, PaintingRestriction, /*bWeightAdjust = */bPerformLegacyWeightBalancing, /*bTotalWeightAdjust = */bUseTotalNormalize);
+			LandscapeEdit.SetAlphaData(LayerInfo, X1, Y1, X2, Y2, Data, /*Stride = */0, PaintingRestriction);
 			//LayerInfo->IsReferencedFromLoadedData = true;
 		}
 	}
@@ -584,5 +562,10 @@ private:
 	bool bPerformLegacyWeightBalancing;
 };
 
+template<bool bInUseInterp, bool bInUseTotalNormalize>
+struct UE_DEPRECATED(5.7, "Use TAlphamapAccessor instead") FAlphamapAccessor : public TAlphamapAccessor<bInUseInterp>
+{
+};
+
 #endif
 
