Fixed warning about unknown shader for landscape thumbnails

diff --git a/Engine/Source/Runtime/Landscape/Private/LandscapeRender.cpp b/Engine/Source/Runtime/Landscape/Private/LandscapeRender.cpp
--- a/Engine/Source/Runtime/Landscape/Private/LandscapeRender.cpp
+++ b/Engine/Source/Runtime/Landscape/Private/LandscapeRender.cpp
@@ -4146,6 +4146,14 @@ public:
 			FName(TEXT("TBasePassPSFSelfShadowedVolumetricLightmapPolicySkylight")),
 			FName(TEXT("TBasePassVSFSelfShadowedVolumetricLightmapPolicy")),
 
+			// No need for velocity in thumbnails
+			FName(TEXT("TVelocityVS<EVelocityPassMode::Velocity_Standard>")),
+			FName(TEXT("TVelocityPS<EVelocityPassMode::Velocity_Standard>")),
+			FName(TEXT("TVelocityVS<EVelocityPassMode::Velocity_ClippedDepth>")),
+			FName(TEXT("TVelocityPS<EVelocityPassMode::Velocity_ClippedDepth>")),
+			FName(TEXT("TVelocityVS<EVelocityPassMode::Velocity_StereoMotionVectors>")),
+			FName(TEXT("TVelocityPS<EVelocityPassMode::Velocity_StereoMotionVectors>")),
+
 #if RHI_RAYTRACING
 			// No ray tracing on thumbnails
 			FName(TEXT("FRayTracingDynamicGeometryConverterCS")),
