Landscape Editor - Added ability to copy obsolete edit layer data to a proxy through new MapCheck
* Refactors non-edit layer auto-conversion so proxy logic does not rely on version number. It's possible that a Commandlet can load and save a Proxy without running the auto-conversion (In Commandlet, Proxy may be registered in a world without the ALansdcape). If a proxy version is updated, the version number increase causes the auto-conversion to fail next load
* Refactors copy data to edit layer - Added parameters to specify source/dest edit layer. Added PostLoad copies of final/obsolete edit layer data to be referenced by MapCheck system

diff --git a/Engine/Source/Runtime/Landscape/Classes/Landscape.h b/Engine/Source/Runtime/Landscape/Classes/Landscape.h
--- a/Engine/Source/Runtime/Landscape/Classes/Landscape.h
+++ b/Engine/Source/Runtime/Landscape/Classes/Landscape.h
@@ -413,7 +413,17 @@ public:
 	LANDSCAPE_API void CreateDefaultLayer();
 
 	LANDSCAPE_API void CopyOldDataToDefaultLayer();
-	LANDSCAPE_API void CopyOldDataToDefaultLayer(ALandscapeProxy* Proxy);
+	LANDSCAPE_API void CopyOldDataToDefaultLayer(ALandscapeProxy* InProxy);
+
+	/**
+	 * Copies edit layer data from each of the InProxy's landscape components to the given edit layer. Overwrites existing data in the destination edit layer
+	 * @param InProxy is the Proxy to copy Landscape Component data from
+	 * @param InSourceEditLayerGuid is the source edit layer Guid. When valid, copies data from the given edit layer. If invalid, copies each components final merged data
+	 * @param InDestEditLayerGuid is the destination edit layer Guid. When valid, must represent a ULandscapeEditLayerPersistent on the ALandscape. If invalid, the default edit layer is used
+	 * @param bInUseObsoleteLayerData When true, uses obsolete or duplicated final layer data stored during initial registration. Landscape components only duplicate final data when they have >= 1 obsolete layer on load
+	 */
+	LANDSCAPE_API void CopyDataToEditLayer(ALandscapeProxy* InProxy, const FGuid& InSourceEditLayerGuid = FGuid(), const FGuid& InDestEditLayerGuid = FGuid(), bool bInUseObsoleteLayerData = false);
+
 	LANDSCAPE_API void AddLayersToProxy(ALandscapeProxy* InProxy);
 	LANDSCAPE_API FIntPoint ComputeComponentCounts() const;
 	LANDSCAPE_API bool IsLayerNameUnique(const FName& InName) const;
@@ -567,9 +577,6 @@ PRAGMA_ENABLE_DEPRECATION_WARNINGS
 
 	LANDSCAPE_API FName GenerateUniqueLayerName(FName InName = NAME_None) const;
 
-	bool GetIsPendingEditLayerConversion() const;
-	void SetIsPendingEditLayerConversion(bool bInPendingEditLayerConversion);
-
 private:
 	FLandscapeLayer* GetLayerInternal(int32 InLayerIndex);
 	void OnLayerCreatedInternal(ULandscapeEditLayerBase* EditLayer);
@@ -779,9 +786,6 @@ private:
 	// Counter to detect re-entrance
 	uint32 LayerUpdateCount = 0;
 #endif
-
-	/** If true, edit layers were not supported on load. All proxies must copy data to the default edit layer during registration */
-	bool bIsPendingEditLayerConversion = false;
 };
 
 #if WITH_EDITOR
diff --git a/Engine/Source/Runtime/Landscape/Classes/LandscapeComponent.h b/Engine/Source/Runtime/Landscape/Classes/LandscapeComponent.h
--- a/Engine/Source/Runtime/Landscape/Classes/LandscapeComponent.h
+++ b/Engine/Source/Runtime/Landscape/Classes/LandscapeComponent.h
@@ -520,6 +520,11 @@ private:
 	UPROPERTY()
 	TMap<FGuid, FLandscapeLayerComponentData> LayersData;
 
+	/** Obsolete edit layers (Edit layers not tracked by ALandscape). During registration, moves obsolete layer data out of active LayersData into
+	* this transient storage. When obsolete layers exists, creates an additional duplicate of the final merged data on initial registration mapped to FGuid() */
+	UPROPERTY(Transient)
+	TMap<FGuid, FLandscapeLayerComponentData> ObsoleteEditLayerData;
+
 	// Final layer data
 	UPROPERTY(Transient)
 	TArray<TObjectPtr<ULandscapeWeightmapUsage>> WeightmapTexturesUsage;
@@ -821,6 +826,15 @@ public:
 	LANDSCAPE_API bool HasLayersData() const;
 	LANDSCAPE_API const FLandscapeLayerComponentData* GetLayerData(const FGuid& InLayerGuid) const;
 	LANDSCAPE_API FLandscapeLayerComponentData* GetLayerData(const FGuid& InLayerGuid);
+
+	void CopyComponentDataToEditLayer(const FGuid& InSourceEditLayerGuid, const FGuid& InDestEditLayerGuid, bool bInUseObsoleteLayerData, TSet<UTexture2D*>& OutProcessedHeightmaps, TSet<UTexture2D*>& OutProcessedWeightmaps);
+
+	// Obsolete edit layer helpers
+	const TMap<FGuid, FLandscapeLayerComponentData>& GetObsoleteEditLayersData() const;
+	const FLandscapeLayerComponentData* GetObsoleteEditLayerData(const FGuid& InLayerGuid) const; 
+	void RemoveObsoleteEditLayerData(const FGuid& InLayerGuid);
+	void InitializeObsoleteEditLayerData(const TSet<FGuid>& InParentEditLayerGuids);
+
 	LANDSCAPE_API void AddLayerData(const FGuid& InLayerGuid, const FLandscapeLayerComponentData& InData);
 	LANDSCAPE_API void AddDefaultLayerData(const FGuid& InLayerGuid, const TArray<ULandscapeComponent*>& InComponentsUsingHeightmap, TMap<UTexture2D*, UTexture2D*>& InOutCreatedHeightmapTextures);
 	LANDSCAPE_API void RemoveLayerData(const FGuid& InLayerGuid);
diff --git a/Engine/Source/Runtime/Landscape/Classes/LandscapeInfo.h b/Engine/Source/Runtime/Landscape/Classes/LandscapeInfo.h
--- a/Engine/Source/Runtime/Landscape/Classes/LandscapeInfo.h
+++ b/Engine/Source/Runtime/Landscape/Classes/LandscapeInfo.h
@@ -195,6 +195,8 @@ public:
 	// @todo document 
 	// all below.
 	LANDSCAPE_API bool SupportsLandscapeEditing() const;
+	// False if landscape editing not supported, in PIE world, or in a Commandlet
+	LANDSCAPE_API bool IsLandscapeEditableWorld() const;
 
 	LANDSCAPE_API bool AreAllComponentsRegistered() const;
 	LANDSCAPE_API void GetComponentsInRegion(int32 X1, int32 Y1, int32 X2, int32 Y2, TSet<ULandscapeComponent*>& OutComponents, bool bOverlap = true) const;
diff --git a/Engine/Source/Runtime/Landscape/Classes/LandscapeProxy.h b/Engine/Source/Runtime/Landscape/Classes/LandscapeProxy.h
--- a/Engine/Source/Runtime/Landscape/Classes/LandscapeProxy.h
+++ b/Engine/Source/Runtime/Landscape/Classes/LandscapeProxy.h
@@ -85,14 +85,6 @@ DECLARE_MULTICAST_DELEGATE_TwoParams(FOnLandscapeProxyFixupSharedDataDelegate, A
 
 // ----------------------------------------------------------------------------------
 
-enum class ELandscapeEditLayerContentState : uint8
-{
-	Undefined = 0,
-	PostLoad_HadEditLayers,
-	PostLoad_NoEditLayers,
-	AutoConvertedToEditLayers
-};
-
 USTRUCT()
 struct FLandscapeTargetLayerSettings
 {
@@ -877,6 +869,8 @@ PRAGMA_ENABLE_DEPRECATION_WARNINGS
 	/** CurrentVersion is bumped whenever a landscape component has an undo/redo operation applied. This lets us detect when the weightmap fixup needs to be run. */
 	uint32 CurrentVersion = 1;
 	uint32 WeightmapFixupVersion = 0;
+	/** Tracks the number of times this Proxy has been registered since load */
+	uint32 RegistrationCount = 0;
 
 	/** True when this Proxy is registered with the LandscapeInfo */
 	bool bIsRegisteredWithLandscapeInfo : 1 = false;
@@ -961,9 +955,6 @@ PRAGMA_ENABLE_DEPRECATION_WARNINGS
 	UE_DEPRECATED(5.7, "Non-edit layer landscapes are deprecated, all landscapes use the edit layer system now.")
 	TOptional<bool> HadLayersContentAtPostLoadTime;
 
-	/** Tracks the has edit layer state of a proxy through load and registration */
-	ELandscapeEditLayerContentState HadLayersContentAtPostLoadState = ELandscapeEditLayerContentState::Undefined;
-
 	UE_DEPRECATED(5.6, "This has been replaced by the per-platform CVar landscape.HeightmapCompressionMode")
 	UPROPERTY()
 	bool bUseCompressedHeightmapStorage_DEPRECATED = false;
@@ -1366,6 +1357,10 @@ PRAGMA_ENABLE_DEPRECATION_WARNINGS
 	// Converts shared properties overridden before explicit override, asking the user how to proceed.
 	LANDSCAPE_API void UpgradeSharedProperties(ALandscape* InParentLandscape);
 
+	// Displays the obsolete edit layers within the LandscapeComponents in a actionable MapCheck log
+	// Asks users to copy obsolete layer data or delete for good.
+ 	void DisplayObsoleteEditLayerMapCheck(ALandscape* InParentLandscape);
+
 	// Assign only mismatching data and mark proxy package dirty
 	LANDSCAPE_API void FixupSharedData(ALandscape* Landscape, bool bMapCheck = false);
 
@@ -1633,11 +1628,18 @@ protected:
 	*/
 	LANDSCAPE_API void DeleteLayer(const FGuid& InLayerGuid);
 
-	/** Remove Layers not found in InExistingLayers
+	/** Removes Obsolete Edit layer data from all landscape components
+	* @param bShowMapCheckWarning When true, opens MapCheck to display a log for each edit layer removed
 	* @return True if some layers were removed.
 	*/
+	LANDSCAPE_API bool RemoveObsoleteEditLayers(bool bShowMapCheckWarning);
+
+	UE_DEPRECATED(5.7, "Use RemoveObsoleteEditLayers based on ALandscape or use DeleteLayer to directly remove layers")
 	LANDSCAPE_API bool RemoveObsoleteLayers(const TSet<FGuid>& InExistingLayers);
 
+	// Returns a map of this Proxy's LandscapeComponents obsolete edit layers Guids to Debug Names
+	TMap<FGuid, FName> GetObsoleteEditLayersInComponents();
+
 	/** Initialize Layer with empty content if it hasn't been initialized yet. */
 	LANDSCAPE_API void InitializeLayerWithEmptyContent(const FGuid& InLayerGuid);
 
@@ -1700,6 +1702,11 @@ private:
 public:
 	static const TArray<ALandscapeProxy*>& GetLandscapeProxies() { return LandscapeProxies; }
 
+	inline bool HasEditLayersDataOnLoad() const
+	{
+		return bHasEditLayersOnLoad;
+	}
+
 private:
 	/** Maintain list of Proxies for faster iteration */
 	static TArray<ALandscapeProxy*> LandscapeProxies;
@@ -1710,6 +1717,9 @@ private:
 	UPROPERTY()
 	FIntPoint SectionBase = FIntPoint::ZeroValue;
 
+	/** True if this proxy's first component has edit layer data on load. When false, proxy has no components or is part of a legacy non-edit layer landscape and needs to be auto-converted */
+	bool bHasEditLayersOnLoad = false;
+
 #if WITH_EDITOR
 	static constexpr TCHAR const* LandscapeInheritedTag = TEXT("LandscapeInherited");
 	static constexpr TCHAR const* LandscapeOverridableTag = TEXT("LandscapeOverridable");
diff --git a/Engine/Source/Runtime/Landscape/Private/Landscape.cpp b/Engine/Source/Runtime/Landscape/Private/Landscape.cpp
--- a/Engine/Source/Runtime/Landscape/Private/Landscape.cpp
+++ b/Engine/Source/Runtime/Landscape/Private/Landscape.cpp
@@ -122,6 +122,7 @@ Landscape.cpp: Terrain rendering
 #include "Engine/Texture2D.h"
 #include "AssetCompilingManager.h"
 #include "FileHelpers.h"
+#include "ScopedTransaction.h"
 #endif
 
 #include UE_INLINE_GENERATED_CPP_BY_NAME(Landscape)
@@ -2582,6 +2583,185 @@ FLandscapeLayerComponentData* ULandscapeComponent::GetLayerData(const FGuid& InL
 	return LayersData.Find(InLayerGuid);
 }
 
+void ULandscapeComponent::CopyComponentDataToEditLayer(const FGuid& InSourceEditLayerGuid, const FGuid& InDestEditLayerGuid, bool bInUseObsoleteLayerData, TSet<UTexture2D*>& OutProcessedHeightmaps, TSet<UTexture2D*>& OutProcessedWeightmaps)
+{
+	const ALandscapeProxy* LandscapeProxy = GetLandscapeProxy();
+	FLandscapeLayerComponentData* DestLayerComponentData = GetLayerData(InDestEditLayerGuid);
+
+	if (!InDestEditLayerGuid.IsValid() || DestLayerComponentData == nullptr || !DestLayerComponentData->IsInitialized() || !LandscapeProxy)
+	{
+		UE_LOG(LogLandscape, Error, TEXT("CopyDataToEditLayer: Failed to copy data to component: %s. Destination edit layer is invalid or Proxy is not yet registered with ALandscape"), *GetName());
+		return;
+	}
+
+	// Determine the source data: uses the given edit layer data when valid, or the final merged data
+	const FLandscapeLayerComponentData* DataToCopy = nullptr;
+	FLandscapeLayerComponentData CurrentFinalLayerData;
+
+	if (InSourceEditLayerGuid.IsValid())
+	{
+		DataToCopy = bInUseObsoleteLayerData ? GetObsoleteEditLayerData(InSourceEditLayerGuid) : GetLayerData(InSourceEditLayerGuid);
+	}
+	else
+	{
+		CurrentFinalLayerData.HeightmapData.Texture = GetHeightmap();
+		CurrentFinalLayerData.WeightmapData.LayerAllocations = GetWeightmapLayerAllocations();
+		CurrentFinalLayerData.WeightmapData.Textures = GetWeightmapTextures();
+		CurrentFinalLayerData.WeightmapData.TextureUsages = GetWeightmapTexturesUsage();
+
+		DataToCopy = bInUseObsoleteLayerData ? GetObsoleteEditLayerData(FGuid()) : &CurrentFinalLayerData;
+	}
+
+	if (DataToCopy == nullptr || !DataToCopy->IsInitialized())
+	{
+		UE_LOG(LogLandscape, Error, TEXT("CopyDataToEditLayer: Failed to copy data to destination layer for %s : Component: %s. Source layer GUID is invalid or data is not initialized properly"),
+			*LandscapeProxy->GetName(), *GetName());
+		return;
+	}
+
+	Modify();
+
+	// Heightmap: 
+	UTexture2D* SourceHeightmap = DataToCopy->HeightmapData.Texture;
+	if (!OutProcessedHeightmaps.Contains(SourceHeightmap))
+	{
+		OutProcessedHeightmaps.Add(SourceHeightmap);
+
+		UTexture* DestinationHeightmap = DestLayerComponentData->HeightmapData.Texture;
+		check(DestinationHeightmap != nullptr);
+
+		// Only copy Mip0 as other mips will get regenerated
+		TArray64<uint8> ExistingMip0Data;
+		SourceHeightmap->Source.GetMipData(ExistingMip0Data, 0);
+
+		// Mark transactional for Undo operations 
+		DestinationHeightmap->SetFlags(RF_Transactional);
+
+		// Calling modify here makes sure that async texture compilation finishes so we can Lock the mip
+		DestinationHeightmap->Modify();
+		FColor* Mip0Data = (FColor*)DestinationHeightmap->Source.LockMip(0);
+		FMemory::Memcpy(Mip0Data, ExistingMip0Data.GetData(), ExistingMip0Data.Num());
+		DestinationHeightmap->Source.UnlockMip(0);
+
+		DestinationHeightmap->UpdateResource();
+		DestinationHeightmap->ClearFlags(RF_Transactional);
+	}
+
+	// Weightmaps:
+	const TArray<UTexture2D*>& SourceWeightmapTextures = DataToCopy->WeightmapData.Textures;
+	const TArray<FWeightmapLayerAllocationInfo>& SourceLayerAllocations = DataToCopy->WeightmapData.LayerAllocations;
+
+	DestLayerComponentData->WeightmapData.Textures.Empty(SourceWeightmapTextures.Num());
+	DestLayerComponentData->WeightmapData.LayerAllocations.Empty(SourceLayerAllocations.Num());
+	DestLayerComponentData->WeightmapData.TextureUsages.Empty();
+
+	DestLayerComponentData->WeightmapData.Textures.AddDefaulted(SourceWeightmapTextures.Num());
+
+	for (int32 TextureIndex = 0; TextureIndex < SourceWeightmapTextures.Num(); ++TextureIndex)
+	{
+		UTexture2D* ComponentWeightmap = SourceWeightmapTextures[TextureIndex];
+		if (OutProcessedWeightmaps.Contains(ComponentWeightmap))
+		{
+			DestLayerComponentData->WeightmapData.Textures[TextureIndex] = ComponentWeightmap;
+		}
+		else
+		{
+			// No need for mip chain on edit layers : 
+			check(LandscapeProxy);
+			UTexture2D* DestLayerWeightmapTexture = LandscapeProxy->CreateLandscapeTexture(ComponentWeightmap->Source.GetSizeX(), ComponentWeightmap->Source.GetSizeY(), TEXTUREGROUP_Terrain_Weightmap,
+				ComponentWeightmap->Source.GetFormat(), /*OptionalOverrideOuter = */ nullptr, /*bCompress = */false, /*bMipChain = */false);
+
+			// Only copy Mip0 as other mips will get regenerated
+			TArray64<uint8> ExistingMip0Data;
+			ComponentWeightmap->Source.GetMipData(ExistingMip0Data, 0);
+
+			// Mark transactional for Undo operations
+			DestLayerWeightmapTexture->SetFlags(RF_Transactional);
+
+			// Calling modify here makes sure that async texture compilation finishes so we can Lock the mip
+			DestLayerWeightmapTexture->Modify();
+			FColor* Mip0Data = (FColor*)DestLayerWeightmapTexture->Source.LockMip(0);
+			FMemory::Memcpy(Mip0Data, ExistingMip0Data.GetData(), ExistingMip0Data.Num());
+			DestLayerWeightmapTexture->Source.UnlockMip(0);
+
+			DestLayerComponentData->WeightmapData.Textures[TextureIndex] = DestLayerWeightmapTexture;
+
+			// Copy weightmap layer allocations 
+			for (const FWeightmapLayerAllocationInfo& Allocation : SourceLayerAllocations)
+			{
+				if (Allocation.WeightmapTextureIndex == TextureIndex)
+				{
+					DestLayerComponentData->WeightmapData.LayerAllocations.Add(Allocation);
+				}
+			}
+
+			DestLayerWeightmapTexture->UpdateResource();
+			DestLayerWeightmapTexture->ClearFlags(RF_Transactional);
+
+			OutProcessedWeightmaps.Add(DestLayerWeightmapTexture);
+		}
+	}
+
+	// Recalculate weightmap usage based on Textures/Allocations instead of relying on copies from source data
+	InitializeLayersWeightmapUsage(InDestEditLayerGuid);
+}
+
+const TMap<FGuid, FLandscapeLayerComponentData>& ULandscapeComponent::GetObsoleteEditLayersData() const
+{
+	return ObsoleteEditLayerData;
+}
+
+const FLandscapeLayerComponentData* ULandscapeComponent::GetObsoleteEditLayerData(const FGuid& InLayerGuid) const
+{
+	return ObsoleteEditLayerData.Find(InLayerGuid);
+}
+
+void ULandscapeComponent::RemoveObsoleteEditLayerData(const FGuid& InLayerGuid)
+{
+	Modify();
+	ObsoleteEditLayerData.Remove(InLayerGuid);
+}
+
+void ULandscapeComponent::InitializeObsoleteEditLayerData(const TSet<FGuid>& InParentEditLayerGuids)
+{
+	// Data should only be set during first registration
+	check(!GetLandscapeProxy() || GetLandscapeProxy()->RegistrationCount == 0);
+
+	// Move all obsolete layer data from the active layers data to transient storage
+	ForEachLayer([&InParentEditLayerGuids, this](const FGuid& LayerGuid, FLandscapeLayerComponentData& LayerData)
+	{
+		if (LayerGuid.IsValid() && !InParentEditLayerGuids.Contains(LayerGuid))
+		{
+			ObsoleteEditLayerData.Add(LayerGuid, *GetLayerData(LayerGuid));
+		}
+	});
+
+	if (!ObsoleteEditLayerData.IsEmpty())
+	{
+		Modify();
+
+		// Obsolete data has been added to transient storage
+		// Remove layer data from active layers list so it is not included in merge
+		for (TPair<FGuid, FLandscapeLayerComponentData>& ObsoleteLayer : ObsoleteEditLayerData)
+		{
+			RemoveLayerData(ObsoleteLayer.Key);
+		}
+
+		// Create duplicate textures of the final data mapped to FGuid()
+		FLandscapeLayerComponentData CurrentFinalLayerData;
+		CurrentFinalLayerData.HeightmapData.Texture = HeightmapTexture;
+		CurrentFinalLayerData.WeightmapData.LayerAllocations = WeightmapLayerAllocations;
+		CurrentFinalLayerData.WeightmapData.Textures = WeightmapTextures;
+		CurrentFinalLayerData.WeightmapData.TextureUsages = WeightmapTexturesUsage;
+
+		FLandscapeLayerComponentData FinalLayerDataCopy("Obsolete Final Data OnLoad");
+		UE::Landscape::Private::CreateLandscapeComponentLayerDataDuplicate(CurrentFinalLayerData, FinalLayerDataCopy);
+		ObsoleteEditLayerData.Add(FGuid(), FinalLayerDataCopy);
+
+		GetLandscapeActor()->RequestLayersContentUpdateForceAll();
+	}
+}
+
 void ULandscapeComponent::ForEachLayer(TFunctionRef<void(const FGuid&, struct FLandscapeLayerComponentData&)> Fn)
 {
 	for (auto& Pair : LayersData)
@@ -4313,22 +4493,8 @@ void ALandscapeProxy::PostLoad()
 		}
 	}
 
-	if (GetLinkerCustomVersion(FFortniteMainBranchObjectVersion::GUID) < FFortniteMainBranchObjectVersion::MigrateLandscapeNonEditLayerToEditLayer)
-	{
-PRAGMA_DISABLE_DEPRECATION_WARNINGS
-		// Fix up bHasLayersContent_DEPRECATED if needed : there was a point where there was a missing call (when adding new components) that was leading to  it not being 
-		//  properly updated even though there was some layer content. We need to fix it now, so that we don't attempt to migrate data from non-edit layers to 
-		//  edit layers later on, since that would stomp the edit layer data : 
-		if (!bHasLayersContent_DEPRECATED && !LandscapeComponents.IsEmpty() && (LandscapeComponents[0] != nullptr) && LandscapeComponents[0]->HasLayersData())
-		{
-			bHasLayersContent_DEPRECATED = true;
-		}
-
-		// save the load time state of layers content before doing anything else (specifically FixupSharedData can stomp it)
-		check(HadLayersContentAtPostLoadState == ELandscapeEditLayerContentState::Undefined);
-		HadLayersContentAtPostLoadState = bHasLayersContent_DEPRECATED ? ELandscapeEditLayerContentState::PostLoad_HadEditLayers : ELandscapeEditLayerContentState::PostLoad_NoEditLayers;
-PRAGMA_ENABLE_DEPRECATION_WARNINGS
-	}
+	// Store edit layer state before fixup edit layer data is run. If edit layers do not exist, proxy is a legacy non-edit layer landscape
+	bHasEditLayersOnLoad = !LandscapeComponents.IsEmpty() && (LandscapeComponents[0] != nullptr) && LandscapeComponents[0]->HasLayersData();
 #endif // WITH_EDITOR
 
 	Super::PostLoad();
@@ -5096,6 +5262,69 @@ void ALandscapeProxy::UpgradeSharedProperties(ALandscape* InParentLandscape)
 	}
 }
 
+void ALandscapeProxy::DisplayObsoleteEditLayerMapCheck(ALandscape* InParentLandscape)
+{
+	ULandscapeInfo* LandscapeInfo = GetLandscapeInfo();
+	ULandscapeSubsystem* LandscapeSubsystem = GetWorld() ? GetWorld()->GetSubsystem<ULandscapeSubsystem>() : nullptr;
+	checkf(LandscapeInfo != nullptr && LandscapeSubsystem != nullptr, TEXT("DisplayObsoleteEditLayerMapCheck can only be called after the proxies are registered to ULandscapeInfo"));
+
+	TMap<FGuid, FName> ObsoleteEditLayers = GetObsoleteEditLayersInComponents();
+
+	if (!ObsoleteEditLayers.IsEmpty())
+	{
+		TWeakObjectPtr<ALandscapeProxy> LandscapeProxy = this;
+		TWeakObjectPtr<ALandscape> ParentLandscape = InParentLandscape;
+
+		auto CanUpdateObsoleteLayersLamda = [LandscapeProxy]()
+		{
+			if (LandscapeProxy.IsValid() && !LandscapeProxy->GetObsoleteEditLayersInComponents().IsEmpty())
+			{
+				ULandscapeInfo* LandscapeInfo = LandscapeProxy->GetLandscapeInfo();
+				return LandscapeInfo && LandscapeInfo->IsLandscapeEditableWorld();
+			}
+
+			return false;
+		};
+
+		// TODO [jared.ritchie] in the future we can show an entry for each obsolete edit layer (instead of only allowing copy of final data) 
+		// TODO [jared.ritchie] in the future, add options for Copy to Layer X (instead of default layer only)
+		FFormatNamedArguments Arguments;
+		Arguments.Add(TEXT("Proxy"), FText::FromString(LandscapeProxy->GetActorNameOrLabel()));
+		Arguments.Add(TEXT("Landscape"), FText::FromString(ParentLandscape->GetActorNameOrLabel()));
+		TSharedRef<FTokenizedMessage> Message = FMessageLog("MapCheck").Warning()
+			->AddToken(FUObjectToken::Create(LandscapeProxy.Get(), FText::FromString(LandscapeProxy->GetActorNameOrLabel())))
+			->AddToken(FTextToken::Create(LOCTEXT("MapCheck_Message_LandscapeProxy_CheckForObsoleteEditLayers", "has data on edit layers that do not exist on it's parent landscape actor. Saving will delete this data for good unless copied to a valid edit layer.")))
+			->AddToken(FActionToken::Create(LOCTEXT("MapCheck_CheckForObsoleteEditLayers_Copy_Default", "Copy Data to Default Edit Layer"), LOCTEXT("MapCheck_CheckForObsoleteEditLayers_Copy_Default_Desc", "Copy the final merged data from all obsolete layers to the default edit layer. Marks proxy dirty"),
+				FOnActionTokenExecuted::CreateLambda([LandscapeProxy, ParentLandscape]()
+				{
+					if (LandscapeProxy.IsValid() && ParentLandscape.IsValid())
+					{
+						FScopedTransaction Transaction(LOCTEXT("Landscape_CopyObsoleteLayerToDefault", "Copy Obsolete Edit Layer Data to Default Layer"));
+						// First Copy the final merged data from all obsolete layers to the default edit layer
+						ParentLandscape->CopyDataToEditLayer(LandscapeProxy.Get(), FGuid(), FGuid(), /*bInUseObsoleteLayerData =*/ true); 
+						LandscapeProxy->RemoveObsoleteEditLayers(/*bShowMapCheckWarning =*/ false);
+					}
+				}),
+				FCanExecuteActionToken::CreateLambda(CanUpdateObsoleteLayersLamda)))
+			->AddToken(FTextToken::Create(LOCTEXT("MapCheck_Message_LandscapeProxy_CheckForObsoleteEditLayers_Or", " or ")))
+			->AddToken(FActionToken::Create(LOCTEXT("MapCheck_CheckForObsoleteEditLayers_Delete", "Delete Obsolete Edit Layer Data"), LOCTEXT("MapCheck_CheckForObsoleteEditLayers_Delete_Desc", "Deletes all of this proxy's obsolete edit layer data for good. Marks proxy dirty"),
+				FOnActionTokenExecuted::CreateLambda([LandscapeProxy]()
+				{
+					if (LandscapeProxy.IsValid())
+					{
+						FScopedTransaction Transaction(LOCTEXT("Landscape_CopyObsoleteDeleteLayers", "Delete Obsolete Layer Data"));
+						LandscapeProxy->RemoveObsoleteEditLayers(/*bShowMapCheckWarning =*/ true);
+					}
+				}),
+				FCanExecuteActionToken::CreateLambda(CanUpdateObsoleteLayersLamda)));
+
+		Message->AddToken(FMapErrorToken::Create(FMapErrors::LandscapeComponentPostLoad_Warning));
+
+		// Show MapCheck window
+		FMessageLog("MapCheck").Open(EMessageSeverity::Warning);
+	}
+}
+
 void ALandscapeProxy::FixupSharedData(ALandscape* Landscape, const bool bMapCheck)
 {
 	if ((Landscape == nullptr) || (Landscape == this))
@@ -5123,10 +5352,6 @@ void ALandscapeProxy::FixupSharedData(ALandscape* Landscape, const bool bMapChec
 		bUpdated |= !SynchronizedProperties.IsEmpty();
 	}
 
-	TSet<FGuid> LayerGuids;
-	Algo::Transform(Landscape->GetEditLayersConst(), LayerGuids, [](const ULandscapeEditLayerBase* EditLayer) { return EditLayer->GetGuid(); });
-	bUpdated |= RemoveObsoleteLayers(LayerGuids);
-
 	for (const ULandscapeEditLayerBase* EditLayer : Landscape->GetEditLayersConst())
 	{
 		bUpdated |= AddLayer(EditLayer->GetGuid());
@@ -5924,10 +6149,35 @@ void ULandscapeInfo::RegisterLandscapeActorWithProxyInternal(ALandscapeProxy* Pr
 			Proxy->FixupWeightmaps();
 		}
 		
-		// All proxy data fixup is complete - CopyOldDataToDefaultLayer will only migrate data once this flag is set
-		Proxy->bIsLandscapeActorRegisteredWithLandscapeInfo = true;
-
+		// Check for legacy non-edit layer landscapes and obsolete edit layer data on first registration only
+		if (Proxy->RegistrationCount == 0)
 		{
+			TSet<FGuid> ParentEditLayerGuids;
+			Algo::Transform(Landscape->GetEditLayersConst(), ParentEditLayerGuids, [](const ULandscapeEditLayerBase* EditLayer) { return EditLayer->GetGuid(); });
+
+			// If there are obsolete layers on initial registration, move them out of the active LayersData and store duplicate of the current final layer data
+			UWorld* World = Landscape->GetWorld();
+			
+			for (ULandscapeComponent* Component : Proxy->LandscapeComponents)
+			{
+				if (Component != nullptr)
+				{
+					if (IsLandscapeEditableWorld())
+					{
+						Component->InitializeObsoleteEditLayerData(ParentEditLayerGuids);
+					}
+				}
+			}
+			
+			if (IsLandscapeEditableWorld())
+			{
+				Proxy->DisplayObsoleteEditLayerMapCheck(Landscape);
+			}
+			
+			// Non-edit layer auto conversion can only happen once on initial load
+			// All proxy data fixup is complete - CopyOldDataToDefaultLayer will only migrate data once this flag is set
+			Proxy->bIsLandscapeActorRegisteredWithLandscapeInfo = true;
+
 			/*
 			Automatically enable edit layers for non-edit layer landscapes here (since FixupSharedData can stomp updates in PostLoad)
 
@@ -5940,28 +6190,13 @@ void ULandscapeInfo::RegisterLandscapeActorWithProxyInternal(ALandscapeProxy* Pr
 				* Conditions - the entire landscape had edit layers disabled, or edit layers were enabled and serialized only on the ALandscape in a WP setup
 			2. Proxy has layers content -> If ALandscape does not support layers content on load, copy data to the default edit layer
 				* Conditions - the entire landscape had edit layers enabled, or edit layers are enabled and serialized on just the ALandscapeProxy in a WP setup
-				* During an auto conversion the proxy edit layer data is flattened and moved to the default edit layer
+				* If the proxy has edit layer content but the ALandscape doesn't, the data exists on an obsolete layer. MapCheck will display allowing user to decide if data is copied
 			*/
 
 			const bool bIsLandscapeActor = Proxy == LandscapeActor;
+			const bool bProxyRequiresNonEditLayerConversion = bIsLandscapeActor ? LandscapeActor->GetEditLayersConst().IsEmpty() : !Proxy->HasEditLayersDataOnLoad();
 
-			// Some ALandscape actors in a WP setup may be serialized with bHasLayersContent == false but have existing edit layer data
-			if (bIsLandscapeActor && !LandscapeActor->GetEditLayersConst().IsEmpty())
-			{
-				Proxy->HadLayersContentAtPostLoadState = ELandscapeEditLayerContentState::AutoConvertedToEditLayers;
-			}
-
-			const bool bProxyHasNoLayersContentPostLoad = Proxy->HadLayersContentAtPostLoadState == ELandscapeEditLayerContentState::PostLoad_NoEditLayers;
-			const bool bProxyHasBeenAutoConverted = Proxy->HadLayersContentAtPostLoadState == ELandscapeEditLayerContentState::AutoConvertedToEditLayers;
-			
-			// If the LandscapeActor did not support edit layers we need to run a full auto conversion
-			if (bIsLandscapeActor)
-			{
-				LandscapeActor->SetIsPendingEditLayerConversion(bProxyHasNoLayersContentPostLoad && !bProxyHasBeenAutoConverted);
-			}
-
-			// Proxies are registered many times (on Undo), but conversion can only happen once on initial load
-			if (!LandscapeActor->IsTemplate() && !bProxyHasBeenAutoConverted && (LandscapeActor->GetIsPendingEditLayerConversion() || bProxyHasNoLayersContentPostLoad))
+			if (!LandscapeActor->IsTemplate() && bProxyRequiresNonEditLayerConversion)
 			{
 				if (bIsLandscapeActor)
 				{
@@ -5979,6 +6214,9 @@ void ULandscapeInfo::RegisterLandscapeActorWithProxyInternal(ALandscapeProxy* Pr
 					// As proxies are registered, copy the data to edit layer based landscape. Proxies with edit layer data will be flattened to the default layer
 					// Proxy data should not be copied before this point in registration to ensure data is fixed up
 					LandscapeActor->CopyOldDataToDefaultLayer(Proxy);
+
+					// After copying, ensure the InProxy's components are valid and have an edit layer in their LayersData map (Default edit layer)
+					check(!Proxy->LandscapeComponents.IsEmpty() && (Proxy->LandscapeComponents[0] != nullptr) && Proxy->LandscapeComponents[0]->HasLayersData());
 				}
 
 				Proxy->RemoveXYOffsets();
@@ -5986,12 +6224,12 @@ void ULandscapeInfo::RegisterLandscapeActorWithProxyInternal(ALandscapeProxy* Pr
 
 				// Mark the proxy dirty to display the 'Requires Save' warning to the user
 				MarkObjectDirty(/*InObject = */ Proxy, /*bInForceResave =*/ true, /*InLandscapeOverride =*/ Landscape);
-				Proxy->HadLayersContentAtPostLoadState = ELandscapeEditLayerContentState::AutoConvertedToEditLayers;
 			}
 		}
 	}
 
 	Proxy->bIsLandscapeActorRegisteredWithLandscapeInfo = true;
+	Proxy->RegistrationCount++;
 #endif // WITH_EDITOR
 }
 
diff --git a/Engine/Source/Runtime/Landscape/Private/LandscapeEdit.cpp b/Engine/Source/Runtime/Landscape/Private/LandscapeEdit.cpp
--- a/Engine/Source/Runtime/Landscape/Private/LandscapeEdit.cpp
+++ b/Engine/Source/Runtime/Landscape/Private/LandscapeEdit.cpp
@@ -2928,6 +2928,16 @@ bool ULandscapeInfo::SupportsLandscapeEditing() const
 	return bSupportsEditing;
 }
 
+bool ULandscapeInfo::IsLandscapeEditableWorld() const
+{
+	if (GIsEditor)
+	{
+		UWorld* World = LandscapeActor->GetWorld();
+		return World && SupportsLandscapeEditing() && !IsRunningCommandlet() && !World->IsPlayInEditor() && GEditor->PlayWorld == nullptr;
+	}
+	return false;
+}
+
 bool ULandscapeInfo::AreAllComponentsRegistered() const
 {
 	bool bAllRegistered = true;
diff --git a/Engine/Source/Runtime/Landscape/Private/LandscapeEditLayers.cpp b/Engine/Source/Runtime/Landscape/Private/LandscapeEditLayers.cpp
--- a/Engine/Source/Runtime/Landscape/Private/LandscapeEditLayers.cpp
+++ b/Engine/Source/Runtime/Landscape/Private/LandscapeEditLayers.cpp
@@ -1341,16 +1341,6 @@ bool ALandscape::ComputeLandscapeLayerBrushInfo(FTransform& OutLandscapeTransfor
 	return true;
 }
 
-bool ALandscape::GetIsPendingEditLayerConversion() const
-{
-	return bIsPendingEditLayerConversion;
-}
-
-void ALandscape::SetIsPendingEditLayerConversion(bool bInIsPendingEditLayerConversion)
-{
-	bIsPendingEditLayerConversion = bInIsPendingEditLayerConversion;
-}
-
 void ALandscape::ConvertNonEditLayerLandscape()
 {
 	// Detect any attempt to call this in the middle of UpdateLayersContent.  If called from blueprint, log an error and return early instead of asserting.
@@ -1433,6 +1423,12 @@ void ALandscape::CopyOldDataToDefaultLayer()
 }
 
 void ALandscape::CopyOldDataToDefaultLayer(ALandscapeProxy* InProxy)
+{
+	// "Old" data represents the last generated final merged weightmap/heightmap
+	CopyDataToEditLayer(InProxy);
+}
+
+void ALandscape::CopyDataToEditLayer(ALandscapeProxy* InProxy, const FGuid& InSourceEditLayerGuid, const FGuid& InDestEditLayerGuid, bool bInUseObsoleteLayerData)
 {
 	ULandscapeInfo* Info = GetLandscapeInfo();
 
@@ -1449,139 +1445,22 @@ void ALandscape::CopyOldDataToDefaultLayer(ALandscapeProxy* InProxy)
 
 	InProxy->Modify();
 
-	const ULandscapeEditLayerBase* DefaultEditLayer = GetEditLayerConst(0);
-	check(DefaultEditLayer != nullptr);
-
-	struct FWeightmapTextureData
+	const ULandscapeEditLayerBase* DestEditLayer = InDestEditLayerGuid.IsValid() ? GetEditLayerConst(InDestEditLayerGuid) : GetEditLayerConst(0);
+	if (DestEditLayer == nullptr || !Cast<ULandscapeEditLayerPersistent>(DestEditLayer))
 	{
-		UTexture2D* Texture;
-		ULandscapeWeightmapUsage* Usage;
-	};
+		UE_LOG(LogLandscape, Error, TEXT("CopyDataToEditLayer: Failed to copy data to destination layer for proxy %s. Destination edit layer is not found in ALandscape or does not support persistent edit layer data"), *InProxy->GetName());
+		return;
+	}
 
-	TMap<UTexture2D*, FWeightmapTextureData> ProcessedWeightmaps;
-	TArray<UTexture2D*> ProcessedHeightmaps;
+	TSet<UTexture2D*> ProcessedWeightmaps;
+	TSet<UTexture2D*> ProcessedHeightmaps;
 
 	for (ULandscapeComponent* Component : InProxy->LandscapeComponents)
 	{
-		FLandscapeLayerComponentData* LayerData = Component->GetLayerData(DefaultEditLayer->GetGuid());
-
-		if (ensure(LayerData != nullptr && LayerData->IsInitialized()))
-		{
-			// If this component has data on the default edit layer it should not be copied again
-			check(LayerData->WeightmapData.LayerAllocations.Num() == 0);
-
-			// Heightmap
-			UTexture2D* ComponentHeightmap = Component->GetHeightmap();
-
-			if (!ProcessedHeightmaps.Contains(ComponentHeightmap))
-			{
-				ProcessedHeightmaps.Add(ComponentHeightmap);
-
-				UTexture* DefaultLayerHeightmap = LayerData->HeightmapData.Texture;
-				check(DefaultLayerHeightmap != nullptr);
-
-				// Only copy Mip0 as other mips will get regenerated
-				TArray64<uint8> ExistingMip0Data;
-				ComponentHeightmap->Source.GetMipData(ExistingMip0Data, 0);
-
-				// Calling modify here makes sure that async texture compilation finishes so we can Lock the mip
-				DefaultLayerHeightmap->Modify();
-				FColor* Mip0Data = (FColor*)DefaultLayerHeightmap->Source.LockMip(0);
-				FMemory::Memcpy(Mip0Data, ExistingMip0Data.GetData(), ExistingMip0Data.Num());
-				DefaultLayerHeightmap->Source.UnlockMip(0);
-
-				DefaultLayerHeightmap->UpdateResource();
-			}
-
-			// Weightmaps:
-			// Existing weightmap allocations are not freed after being copied
-			// Processes like PIE/Cook can run the non-edit layers conversion but will not run UpdateLayersContent/ReallocateWeightmaps
-			// Instead of having freed allocations, leave the data since the next ReallocateWeightmaps will restore a similar state anyways
-
-			const TArray<UTexture2D*>& ComponentWeightmapTextures = Component->GetWeightmapTextures();
-			const TArray<FWeightmapLayerAllocationInfo>& ComponentLayerAllocations = Component->GetWeightmapLayerAllocations();
-			const TArray<ULandscapeWeightmapUsage*>& ComponentWeightmapTexturesUsage = Component->GetWeightmapTexturesUsage();
-
-			LayerData->WeightmapData.Textures.AddDefaulted(ComponentWeightmapTextures.Num());
-			LayerData->WeightmapData.TextureUsages.AddDefaulted(ComponentWeightmapTexturesUsage.Num());
-
-			for (int32 TextureIndex = 0; TextureIndex < ComponentWeightmapTextures.Num(); ++TextureIndex)
-			{
-				UTexture2D* ComponentWeightmap = ComponentWeightmapTextures[TextureIndex];
-				const FWeightmapTextureData* WeightmapTextureData = ProcessedWeightmaps.Find(ComponentWeightmap);
-
-				if (WeightmapTextureData != nullptr)
-				{
-					LayerData->WeightmapData.Textures[TextureIndex] = WeightmapTextureData->Texture;
-					LayerData->WeightmapData.TextureUsages[TextureIndex] = WeightmapTextureData->Usage;
-					check(WeightmapTextureData->Usage->LayerGuid == DefaultEditLayer->GetGuid());
-
-					for (int32 ChannelIndex = 0; ChannelIndex < ULandscapeWeightmapUsage::NumChannels; ++ChannelIndex)
-					{
-						const ULandscapeComponent* ChannelLandscapeComponent = LayerData->WeightmapData.TextureUsages[TextureIndex]->ChannelUsage[ChannelIndex];
-
-						if (ChannelLandscapeComponent != nullptr && ChannelLandscapeComponent == Component)
-						{
-							for (const FWeightmapLayerAllocationInfo& Allocation : ComponentLayerAllocations)
-							{
-								if (Allocation.WeightmapTextureIndex == TextureIndex)
-								{
-									LayerData->WeightmapData.LayerAllocations.Add(Allocation);
-								}
-							}
-
-							break;
-						}
-					}
-				}
-				else
-				{
-					// No need for mip chain on edit layers : 
-					UTexture2D* NewLayerWeightmapTexture = InProxy->CreateLandscapeTexture(ComponentWeightmap->Source.GetSizeX(), ComponentWeightmap->Source.GetSizeY(), TEXTUREGROUP_Terrain_Weightmap, 
-						ComponentWeightmap->Source.GetFormat(), /*OptionalOverrideOuter = */ nullptr, /*bCompress = */false, /*bMipChain = */false);
-
-					// Only copy Mip0 as other mips will get regenerated
-					TArray64<uint8> ExistingMip0Data;
-					ComponentWeightmap->Source.GetMipData(ExistingMip0Data, 0);
-
-					FColor* Mip0Data = (FColor*)NewLayerWeightmapTexture->Source.LockMip(0);
-					FMemory::Memcpy(Mip0Data, ExistingMip0Data.GetData(), ExistingMip0Data.Num());
-					NewLayerWeightmapTexture->Source.UnlockMip(0);
-
-					LayerData->WeightmapData.Textures[TextureIndex] = NewLayerWeightmapTexture;
-					LayerData->WeightmapData.TextureUsages[TextureIndex] = InProxy->WeightmapUsageMap.Add(NewLayerWeightmapTexture, InProxy->CreateWeightmapUsage());
-
-					for (int32 ChannelIndex = 0; ChannelIndex < ULandscapeWeightmapUsage::NumChannels; ++ChannelIndex)
-					{
-						LayerData->WeightmapData.TextureUsages[TextureIndex]->ChannelUsage[ChannelIndex] = ComponentWeightmapTexturesUsage[TextureIndex]->ChannelUsage[ChannelIndex];
-					}
-
-					LayerData->WeightmapData.TextureUsages[TextureIndex]->LayerGuid = DefaultEditLayer->GetGuid();
-
-					// Create new Usage for the "final" layer as the other one will now be used by the Default layer
-					for (const FWeightmapLayerAllocationInfo& Allocation : ComponentLayerAllocations)
-					{
-						if (Allocation.WeightmapTextureIndex == TextureIndex)
-						{
-							LayerData->WeightmapData.LayerAllocations.Add(Allocation);
-						}
-					}
-
-					FWeightmapTextureData NewTextureData;
-					NewTextureData.Texture = NewLayerWeightmapTexture;
-					NewTextureData.Usage = LayerData->WeightmapData.TextureUsages[TextureIndex];
-
-					ProcessedWeightmaps.Add(ComponentWeightmap, NewTextureData);
-
-					NewLayerWeightmapTexture->UpdateResource();
-				}
-			}
-		}
+		Component->CopyComponentDataToEditLayer(InSourceEditLayerGuid, DestEditLayer->GetGuid(), bInUseObsoleteLayerData, ProcessedHeightmaps, ProcessedWeightmaps);
 	}
 
-	// Data is only copied for proxies with valid components
-	// After copying, ensure the InProxy's components are valid and have an edit layer in their LayersData map (Default edit layer)
-	check(!InProxy->LandscapeComponents.IsEmpty() && (InProxy->LandscapeComponents[0] != nullptr) && InProxy->LandscapeComponents[0]->HasLayersData());
+	RequestLayersContentUpdateForceAll();
 }
 
 void ALandscape::UpdateProxyLayersWeightmapUsage()
@@ -7932,47 +7811,45 @@ void ALandscapeProxy::DeleteUnusedLayers()
 	}
 }
 
+// Deprecated
 bool ALandscapeProxy::RemoveObsoleteLayers(const TSet<FGuid>& InExistingLayers)
+{
+	return RemoveObsoleteEditLayers(/*bShowMapCheckWarning =*/true);
+}
+
+bool ALandscapeProxy::RemoveObsoleteEditLayers(bool bShowMapCheckWarning)
 {
 	bool bModified = false;
 	ULandscapeSubsystem* LandscapeSubsystem = GetWorld()->GetSubsystem<ULandscapeSubsystem>();
-	if (!LandscapeSubsystem)
+	if (!LandscapeSubsystem || !GetLandscapeActor())
 	{
-		// Only available when the world has been initialized, and sometimes we are called from PostLoad
-		// before that happens. Skip for now and report no change made, it will be reexecuted later when
-		// registering the actor.
+		UE_LOG(LogLandscape, Error, TEXT("RemoveObsoleteEditLayers: Cannot remove obsolete edit layer data until proxy is registered with parent ALandscape."));
 		return bModified;
 	}
 
-	TSet<TPair<FGuid, FName>> ComponentLayers;
-	for (ULandscapeComponent* Component : LandscapeComponents)
+	for (const TPair<FGuid, FName>& LayerGuidToDebugName: GetObsoleteEditLayersInComponents())
 	{
-		if (Component != nullptr)
-		{
-			Component->ForEachLayer([&](const FGuid& Guid, FLandscapeLayerComponentData& ComponentData) { ComponentLayers.Add(TPair<FGuid, FName>(Guid, ComponentData.DebugName)); });
-		}
-	}
-
-	for (const TPair<FGuid, FName>& LayerGuidAndName : ComponentLayers)
-	{
-		if (!InExistingLayers.Contains(LayerGuidAndName.Key))
+		if (bShowMapCheckWarning)
 		{
 			FFormatNamedArguments Arguments;
-			Arguments.Add(TEXT("LayerName"), FText::FromString(LayerGuidAndName.Value.ToString()));
-			Arguments.Add(TEXT("LayerGuid"), FText::FromString(LayerGuidAndName.Key.ToString(EGuidFormats::HexValuesInBraces)));
+			Arguments.Add(TEXT("LayerName"), FText::FromString(LayerGuidToDebugName.Value.ToString()));
+			Arguments.Add(TEXT("LayerGuid"), FText::FromString(LayerGuidToDebugName.Key.ToString(EGuidFormats::HexValuesInBraces)));
 
 			FMessageLog("MapCheck").Info()
 				->AddToken(FUObjectToken::Create(this, FText::FromString(GetActorNameOrLabel())))
-				->AddToken(FTextToken::Create(FText::Format(LOCTEXT("MapCheck_Message_LandscapeProxyObsoleteLayer","Layer '{LayerName}' ('{LayerGuid}') was removed from LandscapeProxy because it doesn't match any of the Target Layers. Saving will remove this data for good."), Arguments)))
+				->AddToken(FTextToken::Create(FText::Format(LOCTEXT("MapCheck_Message_LandscapeProxyObsoleteLayer","Layer '{LayerName}' ('{LayerGuid}') was removed from LandscapeProxy because it doesn't match any Edit Layers in the parent Landscape. Saving will remove this data for good."), Arguments)))
 				->AddToken(FActionToken::Create(LOCTEXT("MapCheck_RemoveObsoleteLayers", "Save Modified Landscapes"), LOCTEXT("MapCheck_RemoveObsoleteLayers_Desc", "Saves the modified landscape proxy actors"),
 					FOnActionTokenExecuted::CreateUObject(LandscapeSubsystem, &ULandscapeSubsystem::SaveModifiedLandscapes, UE::Landscape::EBuildFlags::WriteFinalLog),
 					FCanExecuteActionToken::CreateUObject(LandscapeSubsystem, &ULandscapeSubsystem::HasModifiedLandscapes),
 					/*bInSingleUse = */false))
 				->AddToken(FMapErrorToken::Create(FMapErrors::LandscapeComponentPostLoad_Warning));
+		}
 
-			DeleteLayer(LayerGuidAndName.Key);
-			bModified = true;
+		for (ULandscapeComponent* Component : LandscapeComponents)
+		{
+			Component->RemoveObsoleteEditLayerData(LayerGuidToDebugName.Key);
 		}
+		bModified = true;
 	}
 
 	if (bModified)
@@ -7986,6 +7863,34 @@ bool ALandscapeProxy::RemoveObsoleteLayers(const TSet<FGuid>& InExistingLayers)
 	return bModified;
 }
 
+TMap<FGuid, FName> ALandscapeProxy::GetObsoleteEditLayersInComponents()
+{
+	TMap<FGuid, FName> ObsoleteEditLayers;
+
+	ULandscapeSubsystem* LandscapeSubsystem = GetWorld()->GetSubsystem<ULandscapeSubsystem>();
+	if (!LandscapeSubsystem || !GetLandscapeActor())
+	{
+		UE_LOG(LogLandscape, Error, TEXT("GetObsoleteEditLayersInComponents: Cannot get obsolete edit layer data until proxy is registered with parent ALandscape."));
+		return ObsoleteEditLayers;
+	}
+
+	for (ULandscapeComponent* Component : LandscapeComponents)
+	{
+		if (Component != nullptr)
+		{
+			for (TPair<FGuid, FLandscapeLayerComponentData> ObsoleteLayer : Component->GetObsoleteEditLayersData())
+			{
+				if (ObsoleteLayer.Key.IsValid())
+				{
+					ObsoleteEditLayers.Add(TPair<FGuid, FName>(ObsoleteLayer.Key, ObsoleteLayer.Value.DebugName));
+				}
+			}
+		}
+	}
+
+	return ObsoleteEditLayers;
+}
+
 bool ALandscapeProxy::AddLayer(const FGuid& InLayerGuid)
 {
 	bool bModified = false;
diff --git a/Engine/Source/Runtime/Landscape/Private/LandscapeUtilsPrivate.cpp b/Engine/Source/Runtime/Landscape/Private/LandscapeUtilsPrivate.cpp
--- a/Engine/Source/Runtime/Landscape/Private/LandscapeUtilsPrivate.cpp
+++ b/Engine/Source/Runtime/Landscape/Private/LandscapeUtilsPrivate.cpp
@@ -77,6 +77,26 @@ bool UseWeightmapTextureArray(EShaderPlatform InPlatform)
 {
 	return IsMobilePlatform(InPlatform) && (LandscapeMobileWeightTextureArray != 0);	
 }
+
+void CreateLandscapeComponentLayerDataDuplicate(const FLandscapeLayerComponentData& InSourceComponentData, FLandscapeLayerComponentData& OutDestComponentData)
+{
+	// Duplicate Heightmap Texture
+	OutDestComponentData.HeightmapData.Texture = DuplicateObject(InSourceComponentData.HeightmapData.Texture, InSourceComponentData.HeightmapData.Texture->GetOuter());
+	OutDestComponentData.HeightmapData.Texture->PostEditChange();
+
+	// Duplicate Weightmap Textures - avoid duplicating WeightmapTextureUsage since it can be regenerated based on the textures + allocations
+	OutDestComponentData.WeightmapData.Textures.Reset(InSourceComponentData.WeightmapData.Textures.Num());
+	for (TObjectPtr<UTexture2D> Texture : InSourceComponentData.WeightmapData.Textures)
+	{
+		TObjectPtr<UTexture2D> WeightmapTextureCopy = DuplicateObject(Texture, Texture->GetOuter());
+		WeightmapTextureCopy->PostEditChange();
+		OutDestComponentData.WeightmapData.Textures.Add(WeightmapTextureCopy);
+	}
+
+	// Copy Layer Allocations
+	OutDestComponentData.WeightmapData.LayerAllocations = InSourceComponentData.WeightmapData.LayerAllocations;
+}
+
 #endif //!WITH_EDITOR
 
 FIntPoint FLandscapeComponent2DIndexerKeyFuncs::GetKey(ULandscapeComponent* InComponent)
diff --git a/Engine/Source/Runtime/Landscape/Private/LandscapeUtilsPrivate.h b/Engine/Source/Runtime/Landscape/Private/LandscapeUtilsPrivate.h
--- a/Engine/Source/Runtime/Landscape/Private/LandscapeUtilsPrivate.h
+++ b/Engine/Source/Runtime/Landscape/Private/LandscapeUtilsPrivate.h
@@ -8,6 +8,7 @@
 enum EShaderPlatform : uint16;
 class ULandscapeComponent;
 class ULandscapeInfo;
+struct FLandscapeLayerComponentData;
 
 namespace UE::Landscape::Private
 {
@@ -22,6 +23,8 @@ namespace UE::Landscape::Private
 
 	/** Check if Landscape.MobileWeightTextureArray CVar set and we should attempt to use texture arrays for weight maps on mobile platforms  */
 	bool IsMobileWeightmapTextureArrayEnabled();
+
+	void CreateLandscapeComponentLayerDataDuplicate(const FLandscapeLayerComponentData& InSourceComponentData, FLandscapeLayerComponentData& OutDestComponentData);
 #endif //!WITH_EDITOR
 
 	/** 
@@ -187,6 +190,8 @@ namespace UE::Landscape::Private
 					}
 				}
 			}
+
+			return false;
 		}
 
 		TBitArray<> GetValidValueBitIndicesInBounds(const FIntRect& InBounds, bool bInInclusiveBounds)
