Landscape Editor - Ensure non-edit layer auto-conversion does run not for proxies with no landscape components

#DisabledValidations Tag, Preflight

diff --git a/Engine/Source/Runtime/Landscape/Private/Landscape.cpp b/Engine/Source/Runtime/Landscape/Private/Landscape.cpp
--- a/Engine/Source/Runtime/Landscape/Private/Landscape.cpp
+++ b/Engine/Source/Runtime/Landscape/Private/Landscape.cpp
@@ -6194,7 +6194,7 @@ void ULandscapeInfo::RegisterLandscapeActorWithProxyInternal(ALandscapeProxy* Pr
 			*/
 
 			const bool bIsLandscapeActor = Proxy == LandscapeActor;
-			const bool bProxyRequiresNonEditLayerConversion = bIsLandscapeActor ? LandscapeActor->GetEditLayersConst().IsEmpty() : !Proxy->HasEditLayersDataOnLoad();
+			const bool bProxyRequiresNonEditLayerConversion = bIsLandscapeActor ? LandscapeActor->GetEditLayersConst().IsEmpty() : (!Proxy->HasEditLayersDataOnLoad() && !Proxy->LandscapeComponents.IsEmpty());
 
 			if (!LandscapeActor->IsTemplate() && bProxyRequiresNonEditLayerConversion)
 			{
