Landscape - Event system for ULandscapeLayerInfoObject data updates
* Deprecates all public LandscapeInfo properties in favor of public accessors
* ALandscapeProxy subscribes to events to make landscape changes as needed

#FailedValidations Preflight

diff --git a/Engine/Plugins/Enterprise/DatasmithImporter/Source/DatasmithImporter/Private/DatasmithLandscapeImporter.cpp b/Engine/Plugins/Enterprise/DatasmithImporter/Source/DatasmithImporter/Private/DatasmithLandscapeImporter.cpp
--- a/Engine/Plugins/Enterprise/DatasmithImporter/Source/DatasmithImporter/Private/DatasmithLandscapeImporter.cpp
+++ b/Engine/Plugins/Enterprise/DatasmithImporter/Source/DatasmithImporter/Private/DatasmithLandscapeImporter.cpp
@@ -115,7 +115,7 @@ AActor* FDatasmithLandscapeImporter::ImportLandscapeActor( const TSharedRef< IDa
 	{
 		if(ImportLandscapeLayersList[i].LayerInfo != nullptr)
 		{
-			Landscape->AddTargetLayer(ImportLandscapeLayersList[i].LayerInfo->LayerName, FLandscapeTargetLayerSettings(ImportLandscapeLayersList[i].LayerInfo, ImportLandscapeLayersList[i].SourceFilePath));
+			Landscape->AddTargetLayer(ImportLandscapeLayersList[i].LayerInfo->GetLayerName(), FLandscapeTargetLayerSettings(ImportLandscapeLayersList[i].LayerInfo, ImportLandscapeLayersList[i].SourceFilePath));
 
 			int32 LayerInfoIndex = LandscapeInfo->GetLayerInfoIndex(ImportLandscapeLayersList[i].LayerName);
 			if(ensure(LayerInfoIndex != INDEX_NONE))
diff --git a/Engine/Plugins/Enterprise/GLTFExporter/Source/GLTFExporter/Private/Tasks/GLTFDelayedMeshTasks.cpp b/Engine/Plugins/Enterprise/GLTFExporter/Source/GLTFExporter/Private/Tasks/GLTFDelayedMeshTasks.cpp
--- a/Engine/Plugins/Enterprise/GLTFExporter/Source/GLTFExporter/Private/Tasks/GLTFDelayedMeshTasks.cpp
+++ b/Engine/Plugins/Enterprise/GLTFExporter/Source/GLTFExporter/Private/Tasks/GLTFDelayedMeshTasks.cpp
@@ -715,7 +715,7 @@ void FGLTFDelayedLandscapeTask::Process()
 		//Landscape Visibility Layer is named: __LANDSCAPE_VISIBILITY__
 		//based on: Engine/Source/Runtime/Landscape/Private/Materials/MaterialExpressionLandscapeVisibilityMask.cpp
 		//		FName UMaterialExpressionLandscapeVisibilityMask::ParameterName = FName("__LANDSCAPE_VISIBILITY__");
-		FString LayerName = AllocInfo.LayerInfo->LayerName.ToString();
+		FString LayerName = AllocInfo.LayerInfo->GetLayerName().ToString();
 		if (LayerName == TEXT("__LANDSCAPE_VISIBILITY__"))
 		{
 			CDI.GetWeightmapTextureData(AllocInfo.LayerInfo, CompVisData);
diff --git a/Engine/Source/Editor/LandscapeEditor/Private/LandscapeEdMode.cpp b/Engine/Source/Editor/LandscapeEditor/Private/LandscapeEdMode.cpp
--- a/Engine/Source/Editor/LandscapeEditor/Private/LandscapeEdMode.cpp
+++ b/Engine/Source/Editor/LandscapeEditor/Private/LandscapeEdMode.cpp
@@ -2912,11 +2912,13 @@ void FEdModeLandscape::UpdateLayerUsageInformation(TWeakObjectPtr<ULandscapeLaye
 		if (ULandscapeLayerInfoObject* LayerInfo = LayerInfoObj.Get())
 		{
 			bool bUsed = UsedLayerInfos.Contains(LayerInfo);
+PRAGMA_DISABLE_DEPRECATION_WARNINGS
 			if (LayerInfo->IsReferencedFromLoadedData != bUsed)
 			{
 				LayerInfo->IsReferencedFromLoadedData = bUsed;
 				DetailPanelRefreshRequired = true;
 			}
+PRAGMA_ENABLE_DEPRECATION_WARNINGS
 		}
 	}
 	
@@ -2943,7 +2945,9 @@ bool FEdModeLandscape::ShouldShowLayer(TSharedRef<FLandscapeTargetListInfo> Targ
 		return false;
 	}
 
+PRAGMA_DISABLE_DEPRECATION_WARNINGS
 	if (!UISettings->ShowUnusedLayers && (!Target->LayerInfoObj.IsValid() || !Target->LayerInfoObj.Get()->IsReferencedFromLoadedData))
+PRAGMA_ENABLE_DEPRECATION_WARNINGS
 	{
 		return false;
 	}
@@ -4351,7 +4355,7 @@ bool LandscapeEditorUtils::SetWeightmapData(ALandscapeProxy* Landscape, ULandsca
 
 FName FLandscapeTargetListInfo::GetLayerName() const
 {
-	return LayerInfoObj.IsValid() ? LayerInfoObj->LayerName : LayerName;
+	return LayerInfoObj.IsValid() ? LayerInfoObj->GetLayerName() : LayerName;
 }
 
 #undef LOCTEXT_NAMESPACE
diff --git a/Engine/Source/Editor/LandscapeEditor/Private/LandscapeEdMode.h b/Engine/Source/Editor/LandscapeEditor/Private/LandscapeEdMode.h
--- a/Engine/Source/Editor/LandscapeEditor/Private/LandscapeEdMode.h
+++ b/Engine/Source/Editor/LandscapeEditor/Private/LandscapeEdMode.h
@@ -145,6 +145,7 @@ struct FLandscapeTargetListInfo
 		{
 			check(LayerInfoObj.IsValid());
 			ALandscapeProxy* Proxy = LandscapeInfo->GetLandscapeProxy();
+			check(Proxy != nullptr);
 			const FName* TargetLayerName = Proxy->GetTargetLayers().FindKey(FLandscapeTargetLayerSettings(LayerInfoObj.Get()));
 			if (TargetLayerName)
 			{
@@ -152,7 +153,7 @@ struct FLandscapeTargetListInfo
 			}
 			else
 			{
-				return &Proxy->AddTargetLayer(LayerInfoObj->LayerName, FLandscapeTargetLayerSettings(LayerInfoObj.Get()));
+				return &Proxy->AddTargetLayer(LayerInfoObj->GetLayerName(), FLandscapeTargetLayerSettings(LayerInfoObj.Get()));
 			}
 		}
 		return nullptr;
diff --git a/Engine/Source/Editor/LandscapeEditor/Private/LandscapeEdModeComponentTools.cpp b/Engine/Source/Editor/LandscapeEditor/Private/LandscapeEdModeComponentTools.cpp
--- a/Engine/Source/Editor/LandscapeEditor/Private/LandscapeEdModeComponentTools.cpp
+++ b/Engine/Source/Editor/LandscapeEditor/Private/LandscapeEdModeComponentTools.cpp
@@ -1346,7 +1346,7 @@ public:
 				int32 LayerInfoIndex = LandscapeInfo->GetLayerInfoIndex(LayerInfo);
 				if (LayerInfoIndex == INDEX_NONE)
 				{
-					LayerInfoIndex = LandscapeInfo->GetLayerInfoIndex(LayerInfo->LayerName);
+					LayerInfoIndex = LandscapeInfo->GetLayerInfoIndex(LayerInfo->GetLayerName());
 					if (LayerInfoIndex != INDEX_NONE)
 					{
 						FLandscapeInfoLayerSettings& LayerSettings = LandscapeInfo->Layers[LayerInfoIndex];
diff --git a/Engine/Source/Editor/LandscapeEditor/Private/LandscapeEdModeErosionTools.cpp b/Engine/Source/Editor/LandscapeEditor/Private/LandscapeEdModeErosionTools.cpp
--- a/Engine/Source/Editor/LandscapeEditor/Private/LandscapeEdModeErosionTools.cpp
+++ b/Engine/Source/Editor/LandscapeEditor/Private/LandscapeEdModeErosionTools.cpp
@@ -184,7 +184,7 @@ public:
 									if (LayerInfo)
 									{
 										uint8 Weight = WeightDatas[Center*LayerNum + Idx];
-										Softness -= (float)(Weight) / 255.0f * LayerInfo->Hardness;
+										Softness -= (float)(Weight) / 255.0f * LayerInfo->GetHardness();
 									}
 								}
 							}
diff --git a/Engine/Source/Editor/LandscapeEditor/Private/LandscapeEditorDetailCustomization_CopyPaste.cpp b/Engine/Source/Editor/LandscapeEditor/Private/LandscapeEditorDetailCustomization_CopyPaste.cpp
--- a/Engine/Source/Editor/LandscapeEditor/Private/LandscapeEditorDetailCustomization_CopyPaste.cpp
+++ b/Engine/Source/Editor/LandscapeEditor/Private/LandscapeEditorDetailCustomization_CopyPaste.cpp
@@ -578,7 +578,7 @@ FReply FLandscapeEditorDetailCustomization_CopyPaste::OnGizmoExportButtonClicked
 						continue;
 					}
 
-					FName LayerName = Gizmo->LayerInfos[i]->LayerName;
+					FName LayerName = Gizmo->LayerInfos[i]->GetLayerName();
 					SaveDialogTitle = FText::Format(NSLOCTEXT("UnrealEd", "LandscapeExport_LayerFilename", "Choose filename for Layer {0} Export"), FText::FromString(LayerName.ToString())).ToString();
 					DefaultFilename = FString::Printf(TEXT("%s.raw"), *LayerName.ToString());
 					FileTypes = TEXT("Layer .raw files|*.raw|Layer .r8 files|*.r8|All files|*.*");
diff --git a/Engine/Source/Editor/LandscapeEditor/Private/LandscapeEditorDetailCustomization_ImportLayers.cpp b/Engine/Source/Editor/LandscapeEditor/Private/LandscapeEditorDetailCustomization_ImportLayers.cpp
--- a/Engine/Source/Editor/LandscapeEditor/Private/LandscapeEditorDetailCustomization_ImportLayers.cpp
+++ b/Engine/Source/Editor/LandscapeEditor/Private/LandscapeEditorDetailCustomization_ImportLayers.cpp
@@ -348,7 +348,7 @@ bool FLandscapeEditorStructCustomization_FLandscapeImportLayer::ShouldFilterLaye
 	}
 
 	ULandscapeLayerInfoObject* LayerInfo = CastChecked<ULandscapeLayerInfoObject>(AssetData.GetAsset());
-	return LayerInfo->LayerName != LayerName;
+	return LayerInfo->GetLayerName() != LayerName;
 }
 
 EVisibility FLandscapeEditorStructCustomization_FLandscapeImportLayer::GetImportLayerSelectionVisibility()
diff --git a/Engine/Source/Editor/LandscapeEditor/Private/LandscapeEditorDetailCustomization_NewLandscape.cpp b/Engine/Source/Editor/LandscapeEditor/Private/LandscapeEditorDetailCustomization_NewLandscape.cpp
--- a/Engine/Source/Editor/LandscapeEditor/Private/LandscapeEditorDetailCustomization_NewLandscape.cpp
+++ b/Engine/Source/Editor/LandscapeEditor/Private/LandscapeEditorDetailCustomization_NewLandscape.cpp
@@ -1007,7 +1007,7 @@ EVisibility FLandscapeEditorDetailCustomization_NewLandscape::GetLayerVisibility
 	{
 		FLandscapeImportLayer* ImportLayer = reinterpret_cast<FLandscapeImportLayer*>(RawData[0]);
 
-		if ((EdMode->NewLandscapePreviewMode != ENewLandscapePreviewMode::ImportLandscape) && (ImportLayer->LayerName == ALandscapeProxy::VisibilityLayer->LayerName))
+		if ((EdMode->NewLandscapePreviewMode != ENewLandscapePreviewMode::ImportLandscape) && (ImportLayer->LayerName == ALandscapeProxy::VisibilityLayer->GetLayerName()))
 		{
 			return EVisibility::Hidden;
 		}
@@ -1245,11 +1245,11 @@ FReply FLandscapeEditorDetailCustomization_NewLandscape::OnCreateButtonClicked()
 		{
 			if (LandscapeEdMode->NewLandscapePreviewMode == ENewLandscapePreviewMode::ImportLandscape)
 			{
-				Landscape->AddTargetLayer(LayerInfo->LayerName, FLandscapeTargetLayerSettings(LayerInfo, ImportLandscapeLayersList[i].SourceFilePath));
+				Landscape->AddTargetLayer(LayerInfo->GetLayerName(), FLandscapeTargetLayerSettings(LayerInfo, ImportLandscapeLayersList[i].SourceFilePath));
 			}
 			else
 			{
-				Landscape->AddTargetLayer(LayerInfo->LayerName, FLandscapeTargetLayerSettings(LayerInfo));
+				Landscape->AddTargetLayer(LayerInfo->GetLayerName(), FLandscapeTargetLayerSettings(LayerInfo));
 			}
 
 			int32 LayerInfoIndex = LandscapeInfo->GetLayerInfoIndex(ImportLandscapeLayersList[i].LayerName);
diff --git a/Engine/Source/Editor/LandscapeEditor/Private/LandscapeEditorDetailCustomization_TargetLayers.cpp b/Engine/Source/Editor/LandscapeEditor/Private/LandscapeEditorDetailCustomization_TargetLayers.cpp
--- a/Engine/Source/Editor/LandscapeEditor/Private/LandscapeEditorDetailCustomization_TargetLayers.cpp
+++ b/Engine/Source/Editor/LandscapeEditor/Private/LandscapeEditorDetailCustomization_TargetLayers.cpp
@@ -1398,13 +1398,11 @@ TSharedPtr<SWidget> FLandscapeEditorCustomNodeBuilder_TargetLayers::GenerateRow(
 												ALandscape* Landscape = Cast<ALandscape>(Target->Owner);
 												
 												const TMap<FName, FLandscapeTargetLayerSettings>& TargetLayers = Landscape->GetTargetLayers();
-												const FLandscapeTargetLayerSettings* LayerSettings = nullptr;
-												
 												Landscape->RemoveTargetLayer(FName(Target->TargetLayerDisplayName.ToString()));
 											
 												Target->TargetLayerDisplayName = Text;
 												Target->LayerName = FName(Text.ToString());
-												Landscape->AddTargetLayer(Target->LayerName, LayerSettings ? *LayerSettings : FLandscapeTargetLayerSettings());	
+												Landscape->AddTargetLayer(Target->LayerName, FLandscapeTargetLayerSettings());	
 												
 												Target->LandscapeInfo->UpdateLayerInfoMap();
 												if (FEdModeLandscape* LandscapeEdMode = GetEditorMode())
@@ -1740,8 +1738,9 @@ void FLandscapeEditorCustomNodeBuilder_TargetLayers::OnExportLayer(const TShared
 		}
 		else //if (Target->TargetType == ELandscapeToolTargetType::Weightmap)
 		{
-			SaveDialogTitle = FText::Format(LOCTEXT("ExportLayer", "Export Landscape Layer: {0}"), FText::FromName(LayerInfoObj->LayerName)).ToString();
-			DefaultFileName = LayerInfoObj->LayerName.ToString();
+			check(LayerInfoObj != nullptr);
+			SaveDialogTitle = FText::Format(LOCTEXT("ExportLayer", "Export Landscape Layer: {0}"), FText::FromName(LayerInfoObj->GetLayerName())).ToString();
+			DefaultFileName = LayerInfoObj->GetLayerName().ToString();
 			FileTypes = LandscapeEditorModule.GetWeightmapExportDialogTypeString();
 		}
 
@@ -1803,8 +1802,9 @@ void FLandscapeEditorCustomNodeBuilder_TargetLayers::OnImportLayer(const TShared
 		}
 		else //if (Target->TargetType == ELandscapeToolTargetType::Weightmap)
 		{
-			OpenDialogTitle = *FText::Format(LOCTEXT("ImportLayer", "Import Landscape Layer: {0}"), FText::FromName(LayerInfoObj->LayerName)).ToString();
-			DefaultFileName = *FString::Printf(TEXT("%s.png"), *(LayerInfoObj->LayerName.ToString()));
+			check(LayerInfoObj != nullptr);
+			OpenDialogTitle = *FText::Format(LOCTEXT("ImportLayer", "Import Landscape Layer: {0}"), FText::FromName(LayerInfoObj->GetLayerName())).ToString();
+			DefaultFileName = *FString::Printf(TEXT("%s.png"), *(LayerInfoObj->GetLayerName().ToString()));
 			FileTypes = LandscapeEditorModule.GetWeightmapImportDialogTypeString();
 		}
 
@@ -1958,7 +1958,7 @@ bool FLandscapeEditorCustomNodeBuilder_TargetLayers::ShouldFilterLayerInfo(const
 	}
 
 	ULandscapeLayerInfoObject* LayerInfo = CastChecked<ULandscapeLayerInfoObject>(AssetData.GetAsset());
-	return LayerInfo->LayerName != LayerName;
+	return LayerInfo->GetLayerName() != LayerName;
 }
 
 void FLandscapeEditorCustomNodeBuilder_TargetLayers::OnTargetLayerSetObject(const FAssetData& AssetData, const TSharedRef<FLandscapeTargetListInfo> Target)
@@ -1994,7 +1994,7 @@ void FLandscapeEditorCustomNodeBuilder_TargetLayers::OnTargetLayerSetObject(cons
 
 	if (SelectedLayerInfo != Target->LayerInfoObj.Get())
 	{
-		if (ensure(SelectedLayerInfo->LayerName == Target->GetLayerName()))
+		if (ensure(SelectedLayerInfo->GetLayerName() == Target->GetLayerName()))
 		{
 			ULandscapeInfo* LandscapeInfo = Target->LandscapeInfo.Get();
 			ALandscape* LandscapeActor = LandscapeInfo->LandscapeActor.Get();
@@ -2146,7 +2146,7 @@ FSlateColor FLandscapeEditorCustomNodeBuilder_TargetLayers::GetLayerUsageDebugCo
 {
 	if (GLandscapeViewMode == ELandscapeViewMode::LayerUsage && Target->TargetType != ELandscapeToolTargetType::Heightmap && ensure(Target->LayerInfoObj.IsValid()))
 	{
-		return FSlateColor(Target->LayerInfoObj->LayerUsageDebugColor);
+		return FSlateColor(Target->LayerInfoObj->GetLayerUsageDebugColor());
 	}
 	return FSlateColor(FLinearColor(0, 0, 0, 0));
 }
diff --git a/Engine/Source/Editor/LandscapeEditor/Private/LandscapeEditorObject.cpp b/Engine/Source/Editor/LandscapeEditor/Private/LandscapeEditorObject.cpp
--- a/Engine/Source/Editor/LandscapeEditor/Private/LandscapeEditorObject.cpp
+++ b/Engine/Source/Editor/LandscapeEditor/Private/LandscapeEditorObject.cpp
@@ -716,7 +716,7 @@ void ULandscapeEditorObject::RefreshLayerImport(FLandscapeImportLayer& ImportLay
 	ImportLayer.ErrorMessage = FText();
 	ImportLayer.ImportResult = ELandscapeImportResult::Success;
 
-	if (ImportLayer.LayerName == ALandscapeProxy::VisibilityLayer->LayerName)
+	if (ImportLayer.LayerName == ALandscapeProxy::VisibilityLayer->GetLayerName())
 	{
 		ImportLayer.LayerInfo = ALandscapeProxy::VisibilityLayer;
 	}
diff --git a/Engine/Source/Editor/LandscapeEditor/Private/LandscapeEditorUtils.cpp b/Engine/Source/Editor/LandscapeEditor/Private/LandscapeEditorUtils.cpp
--- a/Engine/Source/Editor/LandscapeEditor/Private/LandscapeEditorUtils.cpp
+++ b/Engine/Source/Editor/LandscapeEditor/Private/LandscapeEditorUtils.cpp
@@ -129,7 +129,7 @@ namespace LandscapeEditorUtils
 			FString CurrentPackagePath = LayerInfoAsset.PackagePath.ToString() + "/";
 
 			// Only include assets in the package or its sub folders
-			if (LayerInfo && LayerInfo->LayerName == LayerName && CurrentPackagePath.Contains(TargetLayerAssetPackagePath))
+			if (LayerInfo && LayerInfo->GetLayerName() == LayerName && CurrentPackagePath.Contains(TargetLayerAssetPackagePath))
 			{
 				float WorstCase = static_cast<float>(TargetLayerAssetPackagePath.Len() + CurrentPackagePath.Len());
 				WorstCase = FMath::Max(WorstCase, 1.f);
diff --git a/Engine/Source/Editor/WorldBrowser/Private/Tiles/WorldTileCollectionModel.cpp b/Engine/Source/Editor/WorldBrowser/Private/Tiles/WorldTileCollectionModel.cpp
--- a/Engine/Source/Editor/WorldBrowser/Private/Tiles/WorldTileCollectionModel.cpp
+++ b/Engine/Source/Editor/WorldBrowser/Private/Tiles/WorldTileCollectionModel.cpp
@@ -1797,12 +1797,12 @@ void FWorldTileCollectionModel::ReimportTiledLandscape_Executed(FName TargetLaye
 			for (const TPair<FName, FLandscapeTargetLayerSettings>& Pair : Landscape->GetTargetLayers())
 			{
 				const FLandscapeTargetLayerSettings& LayerSettings = Pair.Value;
-				if (LayerSettings.LayerInfoObj && (LayerSettings.LayerInfoObj->LayerName == TargetLayer || TargetLayer == NAME_None))
+				if (LayerSettings.LayerInfoObj && (LayerSettings.LayerInfoObj->GetLayerName() == TargetLayer || TargetLayer == NAME_None))
 				{
 					if (!LayerSettings.ReimportLayerFilePath.IsEmpty())
 					{
 						TArray<uint8> RawData;
-						ReadWeightmapFile(RawData, *LayerSettings.ReimportLayerFilePath, LayerSettings.LayerInfoObj->LayerName, NumSamplesX, NumSamplesY);
+						ReadWeightmapFile(RawData, *LayerSettings.ReimportLayerFilePath, LayerSettings.LayerInfoObj->GetLayerName(), NumSamplesX, NumSamplesY);
 						FScopedSetLandscapeEditingLayer Scope(Landscape->GetLandscapeActor(), Landscape->ReimportDestinationLayerGuid);
 						LandscapeEditorUtils::SetWeightmapData(Landscape, LayerSettings.LayerInfoObj, RawData);
 
diff --git a/Engine/Source/Editor/WorldBrowser/Private/Tiles/WorldTileModel.cpp b/Engine/Source/Editor/WorldBrowser/Private/Tiles/WorldTileModel.cpp
--- a/Engine/Source/Editor/WorldBrowser/Private/Tiles/WorldTileModel.cpp
+++ b/Engine/Source/Editor/WorldBrowser/Private/Tiles/WorldTileModel.cpp
@@ -1051,7 +1051,10 @@ ALandscapeProxy* FWorldTileModel::ImportLandscapeTile(const FLandscapeImportSett
 
 	for (const FLandscapeImportLayerInfo& ImportLayerInfo : Settings.ImportLayers)
 	{
-		LandscapeProxy->AddTargetLayer(ImportLayerInfo.LayerInfo->LayerName, FLandscapeTargetLayerSettings(ImportLayerInfo.LayerInfo, ImportLayerInfo.SourceFilePath));
+		if (ImportLayerInfo.LayerInfo != nullptr)
+		{
+			LandscapeProxy->AddTargetLayer(ImportLayerInfo.LayerInfo->GetLayerName(), FLandscapeTargetLayerSettings(ImportLayerInfo.LayerInfo, ImportLayerInfo.SourceFilePath));
+		}
 	}
 
 	return LandscapeProxy;
diff --git a/Engine/Source/Runtime/Landscape/Classes/Landscape.h b/Engine/Source/Runtime/Landscape/Classes/Landscape.h
--- a/Engine/Source/Runtime/Landscape/Classes/Landscape.h
+++ b/Engine/Source/Runtime/Landscape/Classes/Landscape.h
@@ -533,7 +533,7 @@ PRAGMA_ENABLE_DEPRECATION_WARNINGS
 	LANDSCAPE_API class ALandscapeBlueprintBrushBase* GetBrushForLayer(int32 InLayerIndex, int32 BrushIndex) const;
 	LANDSCAPE_API TArray<class ALandscapeBlueprintBrushBase*> GetBrushesForLayer(int32 InLayerIndex) const;
 	LANDSCAPE_API void OnBlueprintBrushChanged();
-	LANDSCAPE_API void OnLayerInfoSplineFalloffModulationChanged(ULandscapeLayerInfoObject* InLayerInfo);
+	LANDSCAPE_API void OnLayerInfoSplineFalloffModulationChanged(const ULandscapeLayerInfoObject* InLayerInfo);
 	LANDSCAPE_API void OnPreSave();
 
 	void ClearDirtyData(ULandscapeComponent* InLandscapeComponent);
diff --git a/Engine/Source/Runtime/Landscape/Classes/LandscapeLayerInfoObject.h b/Engine/Source/Runtime/Landscape/Classes/LandscapeLayerInfoObject.h
--- a/Engine/Source/Runtime/Landscape/Classes/LandscapeLayerInfoObject.h
+++ b/Engine/Source/Runtime/Landscape/Classes/LandscapeLayerInfoObject.h
@@ -13,6 +13,34 @@ class UTexture2D;
 struct FPropertyChangedEvent;
 
 
+struct FOnLandscapeLayerInfoDataChangedParams
+{
+	FOnLandscapeLayerInfoDataChangedParams(const ULandscapeLayerInfoObject& InLayerInfoObject, const FPropertyChangedEvent& InPropertyChangedEvent = FPropertyChangedEvent(/*InProperty = */nullptr))
+		: PropertyChangedEvent(InPropertyChangedEvent)
+		, LayerInfoObject(&InLayerInfoObject)
+	{
+	}
+
+	/** Provides some additional context about how data has changed (property, type of change...) */
+	FPropertyChangedEvent PropertyChangedEvent;
+
+	/** Indicates a user-initiated property change */
+	bool bUserTriggered = false;
+
+	/** Indicates the change requires a full landscape update (e.g. parameter affecting heightmap or weightmap...) */
+	bool bRequiresLandscapeUpdate = true;
+
+	/** The delegate is triggered each time a data change is requested, even when the data didn't actually change. This indicates that the
+	 * was actually modified. This can occur for example when several EPropertyChangeType::Interactive changes are triggered because of the user
+	 * manipulating a slider : this will be followed by a final EPropertyChangeType::ValueSet but when this occurs, the data usually is not actually
+	 * modified so, to be consistent, we'll still trigger the delegate but indicate that the value didn't actually change, to let the user react appropriately
+	 */
+	bool bHasValueChanged = true;
+
+	/** Pointer to the ULandscapeLayerInfoObject that broadcast the data change event */
+	const ULandscapeLayerInfoObject* LayerInfoObject = nullptr;
+};
+
 // ----------------------------------------------------------------------------------
 
 UENUM()
@@ -32,16 +60,20 @@ class ULandscapeLayerInfoObject : public UObject
 {
 	GENERATED_UCLASS_BODY()
 
+	UE_DEPRECATED(5.7, "Property will be made private. Use public Getters/Setter instead.")
 	UPROPERTY(VisibleAnywhere, Category = LandscapeLayerInfoObject, AssetRegistrySearchable)
 	FName LayerName;
 
+	UE_DEPRECATED(5.7, "Property will be made private. Use public Getters/Setter instead.")
 	UPROPERTY(EditAnywhere, Category = LandscapeLayerInfoObject, Meta = (DisplayName = "Physical Material", Tooltip = "Physical material to use when this layer is the predominant one at a given location. Note: this is ignored if the Landscape Physical Material node is used in the landscape material. "))
 	TObjectPtr<UPhysicalMaterial> PhysMaterial;
 
+	UE_DEPRECATED(5.7, "Property will be made private. Use public Getters/Setter instead.")
 	UPROPERTY(EditAnywhere, Category = LandscapeLayerInfoObject, Meta = (ClampMin = "0", ClampMax = "1", Tooltip = "Defines how much 'resistance' areas painted with this layer will offer to the Erosion tool. A hardness of 0 means the layer is fully affected by erosion, while 1 means fully unaffected."))
 	float Hardness = 0.0f;
 
 #if WITH_EDITORONLY_DATA
+	UE_DEPRECATED(5.7, "Property will be made private. Use public Getters/Setter instead.")
 	UPROPERTY(EditAnywhere, Category = LandscapeLayerInfoObject, Meta = (ClampMin = "0", ClampMax = "1", Tooltip = "The minimum weight that needs to be painted for that layer to be picked up as the dominant physical layer."))
 	float MinimumCollisionRelevanceWeight = 0.0f;
 
@@ -49,57 +81,191 @@ class ULandscapeLayerInfoObject : public UObject
 	UPROPERTY()
 	uint32 bNoWeightBlend_DEPRECATED:1;
 
+	UE_DEPRECATED(5.7, "Property will be made private. Use public Getters/Setter instead.")
 	UPROPERTY(EditAnywhere, Category = SplineFalloffModulation, Meta = (DisplayName="Texture", Tooltip = "Texture to modulate the Splines Falloff Layer Alpha"))
 	TObjectPtr<UTexture2D> SplineFalloffModulationTexture;
 
+	UE_DEPRECATED(5.7, "Property will be made private. Use public Getters/Setter instead.")
 	UPROPERTY(EditAnywhere, Category = SplineFalloffModulation, Meta = (DisplayName = "Color Mask", Tooltip = "Defines which channel of the Spline Falloff Modulation Texture to use."))
 	ESplineModulationColorMask SplineFalloffModulationColorMask = ESplineModulationColorMask::Red;
 
+	UE_DEPRECATED(5.7, "Property will be made private. Use public Getters/Setter instead.")
 	UPROPERTY(EditAnywhere, Category = SplineFalloffModulation, Meta = (DisplayName = "Tiling", ClampMin = "0.01", Tooltip = "Defines the tiling to use when sampling the Spline Falloff Modulation Texture."))
 	float SplineFalloffModulationTiling = 1.0f;
 
+	UE_DEPRECATED(5.7, "Property will be made private. Use public Getters/Setter instead.")
 	UPROPERTY(EditAnywhere, Category = SplineFalloffModulation, Meta = (DisplayName = "Bias", ClampMin = "0", Tooltip = "Defines the offset to use when sampling the Spline Falloff Modulation Texture."))
 	float SplineFalloffModulationBias = 0.5f;
 
+	UE_DEPRECATED(5.7, "Property will be made private. Use public Getters/Setter instead.")
 	UPROPERTY(EditAnywhere, Category = SplineFalloffModulation, Meta = (DisplayName = "Scale", ClampMin = "0", Tooltip = "Allows to scale the value sampled from the Spline Falloff Modulation Texture."))
 	float SplineFalloffModulationScale = 1.0f;
-				
+
+	// TODO [jared.ritchie] Property is used to hide the unused target layers in the UX when ShowUnusedLayers is false. This data should be per-landscape, so it's not valid to store it in the layer object 
+	// Two landscapes can loaded and reference the same layer info object with one of them saying "there is data associated with this layer" and the other, the contrary
+	UE_DEPRECATED(5.7, "This property will be removed in a future release.")
 	UPROPERTY(NonTransactional, Transient)
 	bool IsReferencedFromLoadedData = false;
 #endif // WITH_EDITORONLY_DATA
 
+	UE_DEPRECATED(5.7, "Property will be made private. Use public Getters/Setter instead.")
 	UPROPERTY(EditAnywhere, Category = LandscapeLayerInfoObject, Meta = (Tooltip = "The color to use for layer usage debug"))
 	FLinearColor LayerUsageDebugColor;
 
-private:
-	UPROPERTY(EditAnywhere, Category = LandscapeLayerInfoObject, Meta = (Tooltip = "Allows this layer's final weight to be adjusted against others."))
-	ELandscapeTargetLayerBlendMethod BlendMethod = ELandscapeTargetLayerBlendMethod::None;
-
-	UPROPERTY(EditAnywhere, Category = LandscapeLayerInfoObject, Meta = (EditCondition = "BlendMethod == ELandscapeTargetLayerBlendMethod::PremultipliedAlphaBlending", Tooltip = "Only available for Advanced Weight Blending. Allows target layers from the same group only to have their weight adjusted against one another."))
-	FName BlendGroup;
-
 public:
-	inline ELandscapeTargetLayerBlendMethod GetBlendMethod() const
+	inline ELandscapeTargetLayerBlendMethod GetBlendMethod() const 
 	{
 		return BlendMethod;
 	}
+	static FName GetBlendMethodMemberName() 
+	{
+		return GET_MEMBER_NAME_CHECKED(ULandscapeLayerInfoObject, BlendMethod); 
+	};
 	LANDSCAPE_API void SetBlendMethod(ELandscapeTargetLayerBlendMethod InBlendMethod, bool bInModify);
 
-	inline FName GetBlendGroup() const
-	{
-		return BlendGroup;
+	inline FName GetBlendGroup() const 
+	{	
+		return BlendGroup; 
 	}
+	static FName GetBlendGroupMemberName()
+	{
+		return GET_MEMBER_NAME_CHECKED(ULandscapeLayerInfoObject, BlendGroup);
+	};
 	LANDSCAPE_API void SetBlendGroup(const FName& InBlendGroup, bool bInModify);
 
+	// TODO [jared.ritchie] in 5.9 when public properties are made private, re-enable deprecation warnings
+	PRAGMA_DISABLE_DEPRECATION_WARNINGS
+	inline const FName& GetLayerName() const
+	{
+		return LayerName;
+	}
+	static FName GetLayerNameMemberName()
+	{
+		return GET_MEMBER_NAME_CHECKED(ULandscapeLayerInfoObject, LayerName);
+	};
+	LANDSCAPE_API void SetLayerName(const FName& InLayerName, bool bInModify);
+
+	inline TObjectPtr<UPhysicalMaterial> GetPhysicalMaterial() const
+	{
+		return PhysMaterial;
+	}
+	static FName GetPhysicalMaterialMemberName()
+	{
+		return GET_MEMBER_NAME_CHECKED(ULandscapeLayerInfoObject, PhysMaterial);
+	};
+	LANDSCAPE_API void SetPhysicalMaterial(TObjectPtr<UPhysicalMaterial> InPhysicalMaterial, bool bInModify);
+
+	inline float GetHardness() const
+	{
+		return Hardness;
+	}
+	static FName GetHardnessMemberName()
+	{
+		return GET_MEMBER_NAME_CHECKED(ULandscapeLayerInfoObject, Hardness);
+	};
+	LANDSCAPE_API void SetHardness(float InHardness, bool bInModify, EPropertyChangeType::Type InChangeType);
+
+	inline const FLinearColor& GetLayerUsageDebugColor() const
+	{
+		return LayerUsageDebugColor;
+	}
+	static FName GetLayerUsageDebugColorMemberName()
+	{
+		return GET_MEMBER_NAME_CHECKED(ULandscapeLayerInfoObject, LayerUsageDebugColor);
+	};
+	LANDSCAPE_API void SetLayerUsageDebugColor(const FLinearColor& InLayerUsageDebugColor, bool bInModify, EPropertyChangeType::Type InChangeType);
+
+#if WITH_EDITORONLY_DATA
+	inline float GetMinimumCollisionRelevanceWeight() const
+	{
+		return MinimumCollisionRelevanceWeight;
+	}
+	static FName GetMinimumCollisionRelevanceWeightMemberName()
+	{
+		return GET_MEMBER_NAME_CHECKED(ULandscapeLayerInfoObject, MinimumCollisionRelevanceWeight);
+	};
+	LANDSCAPE_API void SetMinimumCollisionRelevanceWeight(float InMinimumCollisionRelevanceWeight, bool bInModify, EPropertyChangeType::Type InChangeType);
+
+	inline TObjectPtr<UTexture2D> GetSplineFalloffModulationTexture() const
+	{
+		return SplineFalloffModulationTexture;
+	}
+	static FName GetSplineFalloffModulationTextureMemberName()
+	{
+		return GET_MEMBER_NAME_CHECKED(ULandscapeLayerInfoObject, SplineFalloffModulationTexture);
+	};
+	LANDSCAPE_API void SetSplineFalloffModulationTexture(TObjectPtr<UTexture2D> InSplineFalloffModulationTexture, bool bInModify);
+
+	inline ESplineModulationColorMask GetSplineFalloffModulationColorMask() const
+	{
+		return SplineFalloffModulationColorMask;
+	}
+	static FName GetSplineFalloffModulationColorMaskMemberName()
+	{
+		return GET_MEMBER_NAME_CHECKED(ULandscapeLayerInfoObject, SplineFalloffModulationColorMask);
+	};
+	LANDSCAPE_API void SetSplineFalloffModulationColorMask(ESplineModulationColorMask InSplineFalloffModulationColorMask, bool bInModify);
+
+	inline float GetSplineFalloffModulationTiling() const
+	{
+		return SplineFalloffModulationTiling;
+	}
+	static FName GetSplineFalloffModulationTilingMemberName()
+	{
+		return GET_MEMBER_NAME_CHECKED(ULandscapeLayerInfoObject, SplineFalloffModulationTiling);
+	};
+	LANDSCAPE_API void SetSplineFalloffModulationTiling(float InSplineFalloffModulationTiling, bool bInModify, EPropertyChangeType::Type InChangeType);
+
+	inline float GetSplineFalloffModulationBias() const
+	{
+		return SplineFalloffModulationBias;
+	}
+	static FName GetSplineFalloffModulationBiasMemberName()
+	{
+		return GET_MEMBER_NAME_CHECKED(ULandscapeLayerInfoObject, SplineFalloffModulationBias);
+	};
+	LANDSCAPE_API void SetSplineFalloffModulationBias(float InSplineFalloffModulationBias, bool bInModify, EPropertyChangeType::Type InChangeType);
+
+	inline float GetSplineFalloffModulationScale() const
+	{
+		return SplineFalloffModulationScale;
+	}
+	static FName GetSplineFalloffModulationScaleMemberName()
+	{
+		return GET_MEMBER_NAME_CHECKED(ULandscapeLayerInfoObject, SplineFalloffModulationScale);
+	};
+	LANDSCAPE_API void SetSplineFalloffModulationScale(float InSplineFalloffModulationScale, bool bInModify, EPropertyChangeType::Type InChangeType);
+#endif // WITH_EDITORONLY_DATA
+PRAGMA_ENABLE_DEPRECATION_WARNINGS
+
 	virtual void Serialize(FArchive& Ar) override;
 
 #if WITH_EDITOR
 	//~ Begin UObject Interface
 	virtual void PostEditChangeProperty(FPropertyChangedEvent& PropertyChangedEvent) override;
-	virtual void PostEditUndo() override;
 	virtual void PostLoad() override;
 	//~ End UObject Interface
+
+	/** Delegate triggered whenever a change occurrs in the layer info object data */
+	DECLARE_MULTICAST_DELEGATE_OneParam(FOnLayerInfoDataChanged, const FOnLandscapeLayerInfoDataChangedParams& /*InParams*/);
+	FOnLayerInfoDataChanged::RegistrationType& OnLayerInfoChanged() const
+	{
+		return OnLayerInfoObjectChangedDelegate;
+	}
+
+	void BroadcastOnLayerInfoObjectDataChanged(FName InPropertyName, bool bInUserTriggered, bool bInRequiresLandscapeUpdate, bool bInHasValueChanged, EPropertyChangeType::Type InChangeType);
 #endif
 
 	LANDSCAPE_API FLinearColor GenerateLayerUsageDebugColor() const;
+
+private:
+	UPROPERTY(EditAnywhere, Category = LandscapeLayerInfoObject, Meta = (Tooltip = "Allows this layer's final weight to be adjusted against others."))
+	ELandscapeTargetLayerBlendMethod BlendMethod = ELandscapeTargetLayerBlendMethod::None;
+
+	UPROPERTY(EditAnywhere, Category = LandscapeLayerInfoObject, Meta = (EditCondition = "BlendMethod == ELandscapeTargetLayerBlendMethod::PremultipliedAlphaBlending", Tooltip = "Only available for Advanced Weight Blending. Allows target layers from the same group only to have their weight adjusted against one another."))
+	FName BlendGroup;
+
+#if WITH_EDITOR
+	mutable FOnLayerInfoDataChanged OnLayerInfoObjectChangedDelegate;
+#endif // WITH_EDITOR
 };
diff --git a/Engine/Source/Runtime/Landscape/Classes/LandscapeProxy.h b/Engine/Source/Runtime/Landscape/Classes/LandscapeProxy.h
--- a/Engine/Source/Runtime/Landscape/Classes/LandscapeProxy.h
+++ b/Engine/Source/Runtime/Landscape/Classes/LandscapeProxy.h
@@ -46,6 +46,7 @@ class FMaterialUpdateContext;
 struct FAsyncGrassBuilder;
 struct FLandscapeInfoLayerSettings;
 struct FLandscapePerLODMaterialOverride;
+struct FOnLandscapeLayerInfoDataChangedParams;
 class FLandscapeProxyComponentDataChangedParams;
 struct FMeshDescription;
 
@@ -1316,6 +1317,7 @@ PRAGMA_ENABLE_DEPRECATION_WARNINGS
 	virtual void PostTransacted(const FTransactionObjectEvent& InTransactionEvent) override;
 	virtual void PostEditChangeProperty(FPropertyChangedEvent& PropertyChangedEvent) override;
 	virtual void PostEditImport() override;
+	virtual void PreEditUndo() override;
 	virtual void PostEditUndo() override;
 	//~ End UObject Interface
 
@@ -1613,6 +1615,7 @@ PRAGMA_ENABLE_DEPRECATION_WARNINGS
 
 	// All the Target layers 
 	LANDSCAPE_API const TMap<FName, FLandscapeTargetLayerSettings>& GetTargetLayers() const;
+	LANDSCAPE_API TSet<ULandscapeLayerInfoObject*> GetValidTargetLayerObjects() const;
 
 	// TargetLayer API End
 	// -------------------
@@ -1690,6 +1693,8 @@ private:
 	* @param bSetEmptyComponentsOnly When true, only update nanite components with no source landscape components
 	*/
 	void SetSourceComponentsForNaniteComponents(bool bSetEmptyComponentsOnly = false) const;
+
+	void OnLayerInfoObjectDataChanged(const FOnLandscapeLayerInfoDataChangedParams& InParams);
 #endif
 
 	void PostLoadFixupLandscapeGuidsIfInstanced();
diff --git a/Engine/Source/Runtime/Landscape/Private/Landscape.cpp b/Engine/Source/Runtime/Landscape/Private/Landscape.cpp
--- a/Engine/Source/Runtime/Landscape/Private/Landscape.cpp
+++ b/Engine/Source/Runtime/Landscape/Private/Landscape.cpp
@@ -6,7 +6,12 @@ Landscape.cpp: Terrain rendering
 
 #include "Landscape.h"
 
+#include "Algo/AllOf.h"
+#include "Algo/BinarySearch.h"
+#include "Algo/Count.h"
+#include "Algo/ForEach.h"
 #include "Algo/IsSorted.h"
+#include "Algo/Transform.h"
 #include "Serialization/MemoryWriter.h"
 #include "Serialization/BufferArchive.h"
 #include "Serialization/MemoryReader.h"
@@ -95,10 +100,6 @@ Landscape.cpp: Terrain rendering
 #include "LandscapeNotification.h"
 #include "LandscapeHeightfieldCollisionComponent.h"
 #include "SystemTextures.h"
-#include "Algo/BinarySearch.h"
-#include "Algo/Count.h"
-#include "Algo/Transform.h"
-#include "Algo/AllOf.h"
 #include "WorldPartition/WorldPartition.h"
 #include "WorldPartition/WorldPartitionHelpers.h"
 #include "WorldPartition/WorldPartitionHandle.h"
@@ -123,6 +124,7 @@ Landscape.cpp: Terrain rendering
 #include "AssetCompilingManager.h"
 #include "FileHelpers.h"
 #include "ScopedTransaction.h"
+#include "LandscapeSettings.h"
 #endif
 
 #include UE_INLINE_GENERATED_CPP_BY_NAME(Landscape)
@@ -1669,8 +1671,8 @@ ALandscapeProxy::ALandscapeProxy(const FObjectInitializer& ObjectInitializer)
 
 		VisibilityLayer = ConstructorStatics.LandscapeVisibilityLayerInfoFinder.Get();
 		check(VisibilityLayer);
-		VisibilityLayer->LayerName = UMaterialExpressionLandscapeVisibilityMask::ParameterName;
-		VisibilityLayer->LayerUsageDebugColor = FLinearColor(0, 0, 0, 0);
+		VisibilityLayer->SetLayerName(UMaterialExpressionLandscapeVisibilityMask::ParameterName, /*bInModify =*/ false);
+		VisibilityLayer->SetLayerUsageDebugColor(FLinearColor(0, 0, 0, 0), /*bInModify =*/ false, /*InChangeType = */EPropertyChangeType::ValueSet);
 		VisibilityLayer->AddToRoot();
 	}
 
@@ -4239,7 +4241,7 @@ void ALandscapeProxy::AddReferencedObjects(UObject* InThis, FReferenceCollector&
 
 FName FLandscapeInfoLayerSettings::GetLayerName() const
 {
-	checkSlow(LayerInfoObj == nullptr || LayerInfoObj->LayerName == LayerName);
+	checkSlow(LayerInfoObj == nullptr || LayerInfoObj->GetLayerName() == LayerName);
 
 	return LayerName;
 }
@@ -4262,7 +4264,7 @@ const FLandscapeTargetLayerSettings& ULandscapeInfo::GetTargetLayerSettings(ULan
 	}
 	else
 	{
-		return Proxy->AddTargetLayer(LayerInfo->LayerName, FLandscapeTargetLayerSettings(LayerInfo));
+		return Proxy->AddTargetLayer(LayerInfo->GetLayerName(), FLandscapeTargetLayerSettings(LayerInfo));
 	}
 }
 
@@ -4270,13 +4272,13 @@ void ULandscapeInfo::CreateTargetLayerSettingsFor(ULandscapeLayerInfoObject* Lay
 {
 	ForEachLandscapeProxy([LayerInfo](ALandscapeProxy* Proxy)
 	{
-		if (Proxy->HasTargetLayer(LayerInfo->LayerName))
+		if (Proxy->HasTargetLayer(LayerInfo->GetLayerName()))
 		{
-			Proxy->UpdateTargetLayer(LayerInfo->LayerName, FLandscapeTargetLayerSettings(LayerInfo));
+			Proxy->UpdateTargetLayer(LayerInfo->GetLayerName(), FLandscapeTargetLayerSettings(LayerInfo));
 		}
 		else
 		{
-			Proxy->AddTargetLayer(LayerInfo->LayerName, FLandscapeTargetLayerSettings(LayerInfo));
+			Proxy->AddTargetLayer(LayerInfo->GetLayerName(), FLandscapeTargetLayerSettings(LayerInfo));
 		}
 
 		return true;
@@ -4288,7 +4290,7 @@ ULandscapeLayerInfoObject* ULandscapeInfo::GetLayerInfoByName(FName LayerName, A
 	ULandscapeLayerInfoObject* LayerInfo = nullptr;
 	for (int32 j = 0; j < Layers.Num(); j++)
 	{
-		if (Layers[j].LayerInfoObj && Layers[j].LayerInfoObj->LayerName == LayerName
+		if (Layers[j].LayerInfoObj && Layers[j].LayerInfoObj->GetLayerName() == LayerName
 			&& (Owner == nullptr || Layers[j].Owner == Owner))
 		{
 			LayerInfo = Layers[j].LayerInfoObj;
@@ -4389,7 +4391,7 @@ bool ULandscapeInfo::UpdateLayerInfoMapInternal(ALandscapeProxy* Proxy)
 	// Add Visibility Layer info if not initialized
 	if (ALandscapeProxy::VisibilityLayer != nullptr)
 	{
-		int32 LayerInfoIndex = GetLayerInfoIndex(ALandscapeProxy::VisibilityLayer->LayerName);
+		int32 LayerInfoIndex = GetLayerInfoIndex(ALandscapeProxy::VisibilityLayer->GetLayerName());
 
 		if ((LayerInfoIndex != INDEX_NONE) && (Layers[LayerInfoIndex].LayerInfoObj == nullptr))
 		{
@@ -4506,6 +4508,13 @@ void ALandscapeProxy::PostLoad()
 	PostLoadFixupLandscapeGuidsIfInstanced();
 
 #if WITH_EDITOR
+	// Register to data change events on the layer info object so that we can update the landscape and react to changes : 
+	Algo::ForEach(GetValidTargetLayerObjects(), [this](ULandscapeLayerInfoObject* InLayerInfoObject)
+		{
+			check(!InLayerInfoObject->OnLayerInfoChanged().IsBoundToObject(this));
+			InLayerInfoObject->OnLayerInfoChanged().AddUObject(this, &ALandscapeProxy::OnLayerInfoObjectDataChanged);
+		});
+
 	FixupOverriddenSharedProperties();
 
 	ALandscape* LandscapeActor = GetLandscapeActor();
@@ -4585,7 +4594,7 @@ PRAGMA_DISABLE_DEPRECATION_WARNINGS
 		{
 			if (EditorLayerSetting.LayerInfoObj != nullptr)
 			{
-				TargetLayersForFixup.Add(EditorLayerSetting.LayerInfoObj->LayerName, EditorLayerSetting.LayerInfoObj);
+				TargetLayersForFixup.Add(EditorLayerSetting.LayerInfoObj->GetLayerName(), EditorLayerSetting.LayerInfoObj);
 			}
 		}
 		EditorLayerSettings_DEPRECATED.Reset();
@@ -5013,6 +5022,9 @@ TArray<FName> ALandscapeProxy::SynchronizeSharedProperties(ALandscapeProxy* InLa
 {
 	TArray<FName> SynchronizedProperties = SynchronizeUnmarkedSharedProperties(InLandscape);
 
+	// We need to unregister events from the target layers that we have lost because of the synchronization of the TargetLayers property : 
+	const TSet<ULandscapeLayerInfoObject*> ValidLayerInfoObjectsBefore = GetValidTargetLayerObjects();
+
 	for (TFieldIterator<FProperty> PropertyIterator(GetClass()); PropertyIterator; ++PropertyIterator)
 	{
 		FProperty* Property = *PropertyIterator;
@@ -5033,6 +5045,24 @@ TArray<FName> ALandscapeProxy::SynchronizeSharedProperties(ALandscapeProxy* InLa
 
 	if (!SynchronizedProperties.IsEmpty())
 	{
+		const TSet<ULandscapeLayerInfoObject*> ValidLayerInfoObjectsAfter = GetValidTargetLayerObjects();
+
+		// Add data changed events from layers we have gain
+		const TSet<ULandscapeLayerInfoObject*> LayerInfoObjectsToRegisterTo = ValidLayerInfoObjectsAfter.Difference(ValidLayerInfoObjectsBefore);
+		Algo::ForEach(LayerInfoObjectsToRegisterTo, [this](ULandscapeLayerInfoObject* InLayerInfoObject)
+			{
+				check(!InLayerInfoObject->OnLayerInfoChanged().IsBoundToObject(this));
+				InLayerInfoObject->OnLayerInfoChanged().AddUObject(this, &ALandscapeProxy::OnLayerInfoObjectDataChanged);
+			});
+
+		// Remove data changed events from layers we have lost
+		const TSet<ULandscapeLayerInfoObject*> LayerInfoObjectsToUnregisterFrom = ValidLayerInfoObjectsBefore.Difference(ValidLayerInfoObjectsAfter);
+		Algo::ForEach(LayerInfoObjectsToUnregisterFrom, [this](ULandscapeLayerInfoObject* InLayerInfoObject)
+			{
+				check(InLayerInfoObject->OnLayerInfoChanged().IsBoundToObject(this));
+				InLayerInfoObject->OnLayerInfoChanged().RemoveAll(this);
+			});
+
 		Modify();
 	}
 
@@ -7202,6 +7232,57 @@ void ALandscapeProxy::SetSourceComponentsForNaniteComponents(bool bSetEmptyCompo
 	}
 }
 
+void ALandscapeProxy::OnLayerInfoObjectDataChanged(const FOnLandscapeLayerInfoDataChangedParams& InParams)
+{
+	const ULandscapeSettings* LandscapeSettings = GetDefault<ULandscapeSettings>();
+	ALandscape* LandscapeActor = GetLandscapeActor();
+	check(LandscapeSettings != nullptr && LandscapeActor);
+	const bool bIsLandscapeActor = LandscapeActor == this;
+
+	const FName MemberPropertyName = InParams.PropertyChangedEvent.MemberProperty ? InParams.PropertyChangedEvent.MemberProperty->GetFName() : NAME_None;
+	// On Undo, we do not have access to the property. Force and update of everything
+	const bool bUpdateAll = MemberPropertyName.IsNone();
+
+	if (MemberPropertyName == ULandscapeLayerInfoObject::GetPhysicalMaterialMemberName()
+		|| MemberPropertyName == ULandscapeLayerInfoObject::GetMinimumCollisionRelevanceWeightMemberName()
+		|| bUpdateAll)
+	{
+		ChangedPhysMaterial();
+	}
+
+	if (MemberPropertyName == ULandscapeLayerInfoObject::GetLayerUsageDebugColorMemberName() || bUpdateAll)
+	{
+		if (GetWorld() && !GetWorld()->IsPlayInEditor())
+		{
+			MarkComponentsRenderStateDirty();
+		}
+	}
+
+#if WITH_EDITORONLY_DATA
+	if (MemberPropertyName == ULandscapeLayerInfoObject::GetSplineFalloffModulationTextureMemberName()
+		|| MemberPropertyName == ULandscapeLayerInfoObject::GetSplineFalloffModulationColorMaskMemberName()
+		|| MemberPropertyName == ULandscapeLayerInfoObject::GetSplineFalloffModulationBiasMemberName()
+		|| MemberPropertyName == ULandscapeLayerInfoObject::GetSplineFalloffModulationScaleMemberName()
+		|| MemberPropertyName == ULandscapeLayerInfoObject::GetSplineFalloffModulationTilingMemberName()
+		|| bUpdateAll)
+	{
+		const bool bAllowLandscapeSplineFalloffModulationUpdate = (InParams.PropertyChangedEvent.ChangeType != EPropertyChangeType::Interactive) || LandscapeSettings->GetShouldUpdateEditLayersDuringInteractiveChanges();
+		if (bIsLandscapeActor && bAllowLandscapeSplineFalloffModulationUpdate)
+		{
+			check(InParams.LayerInfoObject != nullptr);
+			LandscapeActor->OnLayerInfoSplineFalloffModulationChanged(InParams.LayerInfoObject);
+		}
+	}
+#endif // WITH_EDITORONLY_DATA
+
+	// Request landscape update if needed
+	const bool bAllowLandscapeUpdate = (InParams.PropertyChangedEvent.ChangeType != EPropertyChangeType::Interactive) || LandscapeSettings->GetShouldUpdateEditLayersDuringInteractiveChanges();
+	if (InParams.bRequiresLandscapeUpdate && bAllowLandscapeUpdate)
+	{
+		LandscapeActor->RequestLayersContentUpdateForceAll(ELandscapeLayerUpdateMode::Update_All, InParams.bUserTriggered);
+	}
+}
+
 void ALandscapeProxy::InvalidatePhysicalMaterial()
 {
 	for (ULandscapeComponent* Component : LandscapeComponents)
@@ -7311,6 +7392,16 @@ bool ALandscapeProxy::RemoveTargetLayer(const FName& Name,  bool bPostEditChange
 {
 	Modify();
 	
+	// Remove data changed events
+	if (FLandscapeTargetLayerSettings* LayerSettings = TargetLayers.Find(Name))
+	{
+		if (LayerSettings && LayerSettings->LayerInfoObj)
+		{
+			check(LayerSettings->LayerInfoObj->OnLayerInfoChanged().IsBoundToObject(this));
+			LayerSettings->LayerInfoObj->OnLayerInfoChanged().RemoveAll(this);
+		}
+	}
+
 	int32 NumItemsRemoved = TargetLayers.Remove(Name);
 	if (FProperty* Property = StaticClass()->FindPropertyByName(GET_MEMBER_NAME_CHECKED(ALandscapeProxy, TargetLayers)); Property != nullptr && bPostEditChange)
 	{
@@ -7337,6 +7428,13 @@ FLandscapeTargetLayerSettings& ALandscapeProxy::AddTargetLayer(const FName& Name
 		PostEditChangeProperty(PropertyChangedEvent);	
 	}
 	
+	// Register to data change events on the layer info object so that we can update the landscape accordingly : 
+	if (Settings.LayerInfoObj)
+	{
+		check(!Settings.LayerInfoObj->OnLayerInfoChanged().IsBoundToObject(this));
+		Settings.LayerInfoObj->OnLayerInfoChanged().AddUObject(this, &ALandscapeProxy::OnLayerInfoObjectDataChanged);
+	}
+	
 	return Settings;
 }
 
@@ -7360,6 +7458,14 @@ bool ALandscapeProxy::UpdateTargetLayer(const FName& Name, const FLandscapeTarge
 	if (TargetLayerSettings)
 	{
 		Modify();
+
+		// We may be updating the layer info here. Remove old event listeners
+		if (TargetLayerSettings->LayerInfoObj != nullptr)
+		{
+			check(TargetLayerSettings->LayerInfoObj->OnLayerInfoChanged().IsBoundToObject(this));
+			TargetLayerSettings->LayerInfoObj->OnLayerInfoChanged().RemoveAll(this);
+		}
+
 		*TargetLayerSettings = InTargetLayerSettings;
 		
 		if (FProperty* Property = StaticClass()->FindPropertyByName(GET_MEMBER_NAME_CHECKED(ALandscapeProxy, TargetLayers));  Property != nullptr && bPostEditChange)
@@ -7368,6 +7474,13 @@ bool ALandscapeProxy::UpdateTargetLayer(const FName& Name, const FLandscapeTarge
 			PostEditChangeProperty(PropertyChangedEvent);
 		}
 
+		// Register to data change events on the current layer info object so that we can update the landscape accordingly : 
+		if (TargetLayerSettings->LayerInfoObj != nullptr)
+		{
+			check(!TargetLayerSettings->LayerInfoObj->OnLayerInfoChanged().IsBoundToObject(this));
+			TargetLayerSettings->LayerInfoObj->OnLayerInfoChanged().AddUObject(this, &ALandscapeProxy::OnLayerInfoObjectDataChanged);
+		}
+
 		return true;
 	}
 
@@ -7400,6 +7513,21 @@ const TMap<FName, FLandscapeTargetLayerSettings>& ALandscapeProxy::GetTargetLaye
 {
 	return TargetLayers;
 }
+
+TSet<ULandscapeLayerInfoObject*> ALandscapeProxy::GetValidTargetLayerObjects() const
+{
+	TSet<ULandscapeLayerInfoObject*> ValidLayerInfoObjects;
+	Algo::ForEach(TargetLayers, [&ValidLayerInfoObjects](const TPair<FName, FLandscapeTargetLayerSettings>& InPair)
+		{
+			if (!InPair.Key.IsNone() && (InPair.Value.LayerInfoObj != nullptr))
+			{
+				check(!ValidLayerInfoObjects.Contains(InPair.Value.LayerInfoObj)); // There should be no duplicate
+				check(InPair.Value.LayerInfoObj->GetLayerName() == InPair.Key); // And the name in the layer info object should match with the key
+				ValidLayerInfoObjects.Add(InPair.Value.LayerInfoObj);
+			}
+		});
+	return ValidLayerInfoObjects;
+}
 #endif
 
 bool ALandscapeProxy::AreNaniteComponentsValid(const FGuid& InProxyContentId) const
diff --git a/Engine/Source/Runtime/Landscape/Private/LandscapeCollision.cpp b/Engine/Source/Runtime/Landscape/Private/LandscapeCollision.cpp
--- a/Engine/Source/Runtime/Landscape/Private/LandscapeCollision.cpp
+++ b/Engine/Source/Runtime/Landscape/Private/LandscapeCollision.cpp
@@ -1262,7 +1262,7 @@ bool ULandscapeHeightfieldCollisionComponent::GenerateCollisionData(const FWrite
 							}
 							else
 							{
-								UPhysicalMaterial* DominantMaterial = Layer && Layer->PhysMaterial ? ToRawPtr(Layer->PhysMaterial) : DefMaterial;
+								UPhysicalMaterial* DominantMaterial = Layer && Layer->GetPhysicalMaterial() ? ToRawPtr(Layer->GetPhysicalMaterial()) : DefMaterial;
 								MaterialIndex = InOutMaterials.AddUnique(DominantMaterial);
 							}
 							MaterialIndices.Add(IntCastChecked<uint8>(MaterialIndex));
diff --git a/Engine/Source/Runtime/Landscape/Private/LandscapeComponent.cpp b/Engine/Source/Runtime/Landscape/Private/LandscapeComponent.cpp
--- a/Engine/Source/Runtime/Landscape/Private/LandscapeComponent.cpp
+++ b/Engine/Source/Runtime/Landscape/Private/LandscapeComponent.cpp
@@ -21,7 +21,7 @@ FName FWeightmapLayerAllocationInfo::GetLayerName() const
 {
 	if (LayerInfo)
 	{
-		return LayerInfo->LayerName;
+		return LayerInfo->GetLayerName();
 	}
 	return NAME_None;
 }
diff --git a/Engine/Source/Runtime/Landscape/Private/LandscapeEdit.cpp b/Engine/Source/Runtime/Landscape/Private/LandscapeEdit.cpp
--- a/Engine/Source/Runtime/Landscape/Private/LandscapeEdit.cpp
+++ b/Engine/Source/Runtime/Landscape/Private/LandscapeEdit.cpp
@@ -537,12 +537,12 @@ bool ULandscapeComponent::ValidateCombinationMaterial(UMaterialInstanceConstant*
 		// Each weightmap allocation must have its equivalent in the material's TerrainLayerWeightParameters : 
 		if (!TerrainLayerWeightParameters.ContainsByPredicate(([&Allocation](const FStaticTerrainLayerWeightParameter& TerrainLayerWeightParameter)
 			{
-				return ((TerrainLayerWeightParameter.LayerName == Allocation.LayerInfo->LayerName)
+				return ((TerrainLayerWeightParameter.LayerName == Allocation.LayerInfo->GetLayerName())
 					&& (TerrainLayerWeightParameter.WeightmapIndex == Allocation.WeightmapTextureIndex));
 			})))
 		{
 			UE_LOG(LogLandscape, Display, TEXT("Material instance %s in landscape component %s doesn't match the expected shader combination: missing layer %s"),
-				*InCombinationMaterial->GetName(), *GetPathName(), *Allocation.LayerInfo->LayerName.ToString());
+				*InCombinationMaterial->GetName(), *GetPathName(), *Allocation.LayerInfo->GetLayerName().ToString());
 
 			return false;
 		}
@@ -1240,9 +1240,9 @@ void ULandscapeComponent::FixupWeightmaps(const FGuid& InEditLayerGuid)
 				}
 				else if(Allocation.LayerInfo != ALandscapeProxy::VisibilityLayer && Landscape && !Landscape->HasTargetLayer(Allocation.LayerInfo))
 				{
-					if (Landscape->HasTargetLayer(Allocation.LayerInfo->LayerName))
+					if (Landscape->HasTargetLayer(Allocation.LayerInfo->GetLayerName()))
 					{
-						if (ULandscapeLayerInfoObject* NewLayerInfo = Landscape->GetTargetLayers()[Allocation.LayerInfo->LayerName].LayerInfoObj)
+						if (ULandscapeLayerInfoObject* NewLayerInfo = Landscape->GetTargetLayers()[Allocation.LayerInfo->GetLayerName()].LayerInfoObj)
 						{
 							// Point the old layer info to the new one : 
 							LayersToRemap.Add(Allocation.LayerInfo, NewLayerInfo);
@@ -1267,7 +1267,7 @@ void ULandscapeComponent::FixupWeightmaps(const FGuid& InEditLayerGuid)
 			if (!LayersToDelete.IsEmpty())
 			{
 				TArray<FString> LayerStrings;
-				Algo::Transform(LayersToDelete, LayerStrings, [](ULandscapeLayerInfoObject* LayerInfo) { return LayerInfo ? LayerInfo->LayerName.ToString() : FString(TEXT("null")); });
+				Algo::Transform(LayersToDelete, LayerStrings, [](ULandscapeLayerInfoObject* LayerInfo) { return LayerInfo ? LayerInfo->GetLayerName().ToString() : FString(TEXT("null")); });
 				FString LayersToDeleteString = *FString::Join(LayerStrings, TEXT(", "));
 
 				// Report the deletion info to MapCheck (no need for a warning as the fix is automatic and trivial) :
@@ -1313,7 +1313,7 @@ void ULandscapeComponent::FixupWeightmaps(const FGuid& InEditLayerGuid)
 					{ 
 						ULandscapeLayerInfoObject* OldLayerInfo = LayerToRemap.Key;
 						ULandscapeLayerInfoObject* NewLayerInfo = LayerToRemap.Value;
-						return FString::Format(TEXT("{0} ({1} -> {2})"), {OldLayerInfo->LayerName.ToString(), OldLayerInfo->GetFullName(), NewLayerInfo->GetFullName()});
+						return FString::Format(TEXT("{0} ({1} -> {2})"), {OldLayerInfo->GetLayerName().ToString(), OldLayerInfo->GetFullName(), NewLayerInfo->GetFullName()});
 					});
 				FString LayersToRemapString = *FString::Join(LayerStrings, TEXT(", "));
 
@@ -1334,7 +1334,7 @@ void ULandscapeComponent::FixupWeightmaps(const FGuid& InEditLayerGuid)
 				{
 					if (ULandscapeLayerInfoObject** NewLayerInfo = LayersToRemap.Find(Allocation.LayerInfo))
 					{
-						check(Allocation.GetLayerName() == (*NewLayerInfo)->LayerName);
+						check(Allocation.GetLayerName() == (*NewLayerInfo)->GetLayerName());
 						Allocation.LayerInfo = *NewLayerInfo;
 					}
 				}
@@ -1870,7 +1870,7 @@ void ULandscapeComponent::UpdateDominantLayerBuffer(int32 InComponentX1, int32 I
 					for (uint8 LayerIdx = 0; LayerIdx < NumLayers; LayerIdx++)
 					{
 						const uint8 LayerWeight = InCollisionDataPtrs[LayerIdx][DataOffset];
-						const uint8 LayerMinimumWeight = InLayerInfos[LayerIdx] ? (uint8)(InLayerInfos[LayerIdx]->MinimumCollisionRelevanceWeight * 255) :  0;
+						const uint8 LayerMinimumWeight = InLayerInfos[LayerIdx] ? (uint8)(InLayerInfos[LayerIdx]->GetMinimumCollisionRelevanceWeight() * 255) :  0;
 
 						if (LayerIdx == InVisibilityLayerIdx) // Override value for visibility layer holes
 						{
@@ -3853,9 +3853,9 @@ void ALandscapeProxy::Import(const FGuid& InGuid, int32 InMinX, int32 InMinY, in
 					LandscapeEdit.SetAlphaData(MaterialLayerInfo.LayerInfo, InMinX, InMinY, InMaxX, InMaxY, MaterialLayerInfo.LayerData.GetData(), /*Stride = */0, ELandscapeLayerPaintingRestriction::None);
 
 					// Add to TargetLayers if needed.
-					if (!HasTargetLayer(MaterialLayerInfo.LayerInfo->LayerName))
+					if (!HasTargetLayer(MaterialLayerInfo.LayerInfo->GetLayerName()))
 					{
-						AddTargetLayer(MaterialLayerInfo.LayerInfo->LayerName, FLandscapeTargetLayerSettings(MaterialLayerInfo.LayerInfo));
+						AddTargetLayer(MaterialLayerInfo.LayerInfo->GetLayerName(), FLandscapeTargetLayerSettings(MaterialLayerInfo.LayerInfo));
 					}
 				}
 			}
@@ -5002,7 +5002,7 @@ void ULandscapeInfo::ExportLayer(ULandscapeLayerInfoObject* LayerInfo, const FSt
 	const ILandscapeWeightmapFileFormat* WeightmapFormat = LandscapeEditorModule.GetWeightmapFormatByExtension(*FPaths::GetExtension(Filename, true));
 	if (WeightmapFormat)
 	{
-		WeightmapFormat->Export(*Filename, LayerInfo->LayerName, WeightData, { (uint32)ExportWidth, (uint32)ExportHeight }, DrawScale * FVector(1, 1, LANDSCAPE_ZSCALE));
+		WeightmapFormat->Export(*Filename, LayerInfo->GetLayerName(), WeightData, {(uint32)ExportWidth, (uint32)ExportHeight}, DrawScale * FVector(1, 1, LANDSCAPE_ZSCALE));
 	}
 }
 
@@ -5053,11 +5053,11 @@ void ULandscapeInfo::ReplaceLayer(ULandscapeLayerInfoObject* FromLayerInfo, ULan
 			}
 		}
 
-		LandscapeActor->RemoveTargetLayer(FromLayerInfo->LayerName);
+		LandscapeActor->RemoveTargetLayer(FromLayerInfo->GetLayerName());
 
 		// When a layer is cleared ToLayerInfo will be invalid
 		const FLandscapeTargetLayerSettings NewSettings = ToLayerInfo != nullptr ? FLandscapeTargetLayerSettings(ToLayerInfo) : FLandscapeTargetLayerSettings();
-		const FName NewName = ToLayerInfo != nullptr ? ToLayerInfo->LayerName : FromLayerInfo->LayerName;
+		const FName NewName = ToLayerInfo != nullptr ? ToLayerInfo->GetLayerName() : FromLayerInfo->GetLayerName();
 		LandscapeActor->AddTargetLayer(NewName, NewSettings);
 		UpdateLayerInfoMap();
 
diff --git a/Engine/Source/Runtime/Landscape/Private/LandscapeEditInterface.cpp b/Engine/Source/Runtime/Landscape/Private/LandscapeEditInterface.cpp
--- a/Engine/Source/Runtime/Landscape/Private/LandscapeEditInterface.cpp
+++ b/Engine/Source/Runtime/Landscape/Private/LandscapeEditInterface.cpp
@@ -1635,7 +1635,9 @@ void FLandscapeEditDataInterface::FillLayer(ULandscapeLayerInfoObject* LayerInfo
 		return;
 	}
 
+PRAGMA_DISABLE_DEPRECATION_WARNINGS
 	LayerInfo->IsReferencedFromLoadedData = true;
+PRAGMA_ENABLE_DEPRECATION_WARNINGS
 
 	for (auto& XYComponentPair : LandscapeInfo->XYtoComponentMap)
 	{
diff --git a/Engine/Source/Runtime/Landscape/Private/LandscapeEditLayerMergeRenderContext.cpp b/Engine/Source/Runtime/Landscape/Private/LandscapeEditLayerMergeRenderContext.cpp
--- a/Engine/Source/Runtime/Landscape/Private/LandscapeEditLayerMergeRenderContext.cpp
+++ b/Engine/Source/Runtime/Landscape/Private/LandscapeEditLayerMergeRenderContext.cpp
@@ -835,7 +835,7 @@ void FMergeRenderContext::GenericBlendLayer(const FBlendParams& InBlendParams, F
 						{ 
 							return InTargetLayerGroup.GetWeightmapTargetLayerBitIndices()[TargetLayerIndex];
 						});
-					checkf(TargetLayerGroupIndex != INDEX_NONE, TEXT("Couldn't find a target layer group for layer %s. All premultiplied alpha blending must belong to a group (and one only)"), *InLayerInfo->LayerName.ToString());
+					checkf(TargetLayerGroupIndex != INDEX_NONE, TEXT("Couldn't find a target layer group for layer %s. All premultiplied alpha blending must belong to a group (and one only)"), *InLayerInfo->GetLayerName().ToString());
 					WeightmapTargetLayerInfo.BlendGroupIndex = TargetLayerGroupIndex;
 				}
 				if (InLayerInfo == ALandscapeProxy::VisibilityLayer)
diff --git a/Engine/Source/Runtime/Landscape/Private/LandscapeEditLayers.cpp b/Engine/Source/Runtime/Landscape/Private/LandscapeEditLayers.cpp
--- a/Engine/Source/Runtime/Landscape/Private/LandscapeEditLayers.cpp
+++ b/Engine/Source/Runtime/Landscape/Private/LandscapeEditLayers.cpp
@@ -1488,11 +1488,30 @@ void ALandscapeProxy::UpdateProxyLayersWeightmapUsage()
 	check(!bNeedsWeightmapUsagesUpdate);
 }
 
+void ALandscapeProxy::PreEditUndo()
+{
+	Super::PreEditUndo();
+
+	// Unregister all ULandscapeLayerInfoObjects events. PostEditUndo will re-register objects that still exist post-undo
+	Algo::ForEach(GetValidTargetLayerObjects(), [this](ULandscapeLayerInfoObject* InLayerInfoObject)
+		{
+			check(InLayerInfoObject->OnLayerInfoChanged().IsBoundToObject(this));
+			InLayerInfoObject->OnLayerInfoChanged().RemoveAll(this);
+		});
+}
+
 void ALandscapeProxy::PostEditUndo()
 {
 	check(ULandscapeComponent::UndoRedoModifiedComponentCount == 0);
 
 	Super::PostEditUndo();
+
+	// Undo operation may have added/removed ULandscapeLayerInfoObjects. Re-register to events here
+	Algo::ForEach(GetValidTargetLayerObjects(), [this](ULandscapeLayerInfoObject* InLayerInfoObject)
+		{
+			check(!InLayerInfoObject->OnLayerInfoChanged().IsBoundToObject(this));
+			InLayerInfoObject->OnLayerInfoChanged().AddUObject(this, &ALandscapeProxy::OnLayerInfoObjectDataChanged);
+		});
 }
 
 void ALandscape::InitializeLandscapeLayersWeightmapUsage()
@@ -1828,7 +1847,7 @@ TArray<FLandscapeLayersCopyReadbackTextureParams> PrepareLandscapeLayersCopyRead
 						UTexture2D* PaintLayerTexture = WeightmapTextures[AllocInfo.WeightmapTextureIndex];
 						if (PaintLayerTexture == Texture)
 						{
-							PerChannelLayerNames[AllocInfo.WeightmapTextureChannel] = AllocInfo.LayerInfo->LayerName;
+							PerChannelLayerNames[AllocInfo.WeightmapTextureChannel] = AllocInfo.LayerInfo->GetLayerName();
 						}
 					}
 				}
@@ -5651,7 +5670,7 @@ bool ALandscape::ResolveLayersTexture(
 							if (WeightmapLayerAllocation.IsAllocated())
 							{
 								UTexture2D* Texture = WeightmapTextures[WeightmapLayerAllocation.WeightmapTextureIndex];
-								if ((Texture == InOutputTexture) && (AllZerosLayerInfoName == WeightmapLayerAllocation.LayerInfo->LayerName))
+								if ((Texture == InOutputTexture) && (AllZerosLayerInfoName == WeightmapLayerAllocation.LayerInfo->GetLayerName()))
 								{
 									FLandscapeEditLayerComponentReadbackResult* ComponentReadbackResult = InOutComponentReadbackResults.FindByPredicate([LandscapeComponent = *Component](const FLandscapeEditLayerComponentReadbackResult& Element) { return Element.LandscapeComponent == LandscapeComponent; });
 									check(ComponentReadbackResult != nullptr);
@@ -9252,7 +9271,7 @@ void ALandscape::OnBlueprintBrushChanged()
 #endif
 }
 
-void ALandscape::OnLayerInfoSplineFalloffModulationChanged(ULandscapeLayerInfoObject* InLayerInfo)
+void ALandscape::OnLayerInfoSplineFalloffModulationChanged(const ULandscapeLayerInfoObject* InLayerInfo)
 {
 	ULandscapeInfo* LandscapeInfo = GetLandscapeInfo();
 
diff --git a/Engine/Source/Runtime/Landscape/Private/LandscapeInfo.cpp b/Engine/Source/Runtime/Landscape/Private/LandscapeInfo.cpp
--- a/Engine/Source/Runtime/Landscape/Private/LandscapeInfo.cpp
+++ b/Engine/Source/Runtime/Landscape/Private/LandscapeInfo.cpp
@@ -8,7 +8,7 @@
 
 LANDSCAPE_API FLandscapeInfoLayerSettings::FLandscapeInfoLayerSettings(ULandscapeLayerInfoObject* InLayerInfo, class ALandscapeProxy* InProxy)
 	: LayerInfoObj(InLayerInfo)
-	, LayerName((InLayerInfo != nullptr) ? InLayerInfo->LayerName : NAME_None)
+	, LayerName((InLayerInfo != nullptr) ? InLayerInfo->GetLayerName() : NAME_None)
 #if WITH_EDITORONLY_DATA
 	, ThumbnailMIC(nullptr)
 	, Owner(InProxy)
diff --git a/Engine/Source/Runtime/Landscape/Private/LandscapeLayerInfoObject.cpp b/Engine/Source/Runtime/Landscape/Private/LandscapeLayerInfoObject.cpp
--- a/Engine/Source/Runtime/Landscape/Private/LandscapeLayerInfoObject.cpp
+++ b/Engine/Source/Runtime/Landscape/Private/LandscapeLayerInfoObject.cpp
@@ -18,8 +18,10 @@ ULandscapeLayerInfoObject::ULandscapeLayerInfoObject(const FObjectInitializer& O
 {
 	if (!IsTemplate())
 	{
+PRAGMA_DISABLE_DEPRECATION_WARNINGS
 		// Assign initial random LayerUsageDebugColor
 		LayerUsageDebugColor = GenerateLayerUsageDebugColor();
+PRAGMA_ENABLE_DEPRECATION_WARNINGS
 
 		// The default blend method comes from the settings :
 		BlendMethod = GetDefault<ULandscapeSettings>()->GetTargetLayerDefaultBlendMethod();
@@ -47,7 +49,9 @@ void ULandscapeLayerInfoObject::SetBlendMethod(ELandscapeTargetLayerBlendMethod
 		BlendMethod = InBlendMethod;
 	}
 
-	// TODO [jonathan.bard] : send event for data change
+#if WITH_EDITOR
+	BroadcastOnLayerInfoObjectDataChanged(GET_MEMBER_NAME_CHECKED(ULandscapeLayerInfoObject, BlendMethod), /*bInUserTriggered = */true, /*bInRequiresLandscapeUpdate =*/ true, bHasValueChanged, /*InChangeType = */EPropertyChangeType::ValueSet);
+#endif // WITH_EDITOR
 }
 
 void ULandscapeLayerInfoObject::SetBlendGroup(const FName& InBlendGroup, bool bInModify)
@@ -62,106 +66,205 @@ void ULandscapeLayerInfoObject::SetBlendGroup(const FName& InBlendGroup, bool bI
 		BlendGroup = InBlendGroup;
 	}
 
-	// TODO [jonathan.bard] : send event for data change
+#if WITH_EDITOR
+	BroadcastOnLayerInfoObjectDataChanged(GET_MEMBER_NAME_CHECKED(ULandscapeLayerInfoObject, BlendGroup), /*bInUserTriggered = */true, /*bInRequiresLandscapeUpdate =*/ true, bHasValueChanged, /*InChangeType = */EPropertyChangeType::ValueSet);
+#endif // WITH_EDITOR
 }
 
-void ULandscapeLayerInfoObject::Serialize(FArchive& Ar)
+// TODO [jared.ritchie] when deprecated public properties are made private, re-enable deprecation warnings
+PRAGMA_DISABLE_DEPRECATION_WARNINGS
+void ULandscapeLayerInfoObject::SetLayerName(const FName& InLayerName, bool bInModify)
 {
-	Super::Serialize(Ar);
+	const bool bHasValueChanged = (InLayerName != LayerName);
+	if (bHasValueChanged)
+	{
+		if (bInModify)
+		{
+			Modify();
+		}
+		LayerName = InLayerName;
+	}
 
-	Ar.UsingCustomVersion(FFortniteMainBranchObjectVersion::GUID);
+#if WITH_EDITOR
+	BroadcastOnLayerInfoObjectDataChanged(GET_MEMBER_NAME_CHECKED(ULandscapeLayerInfoObject, LayerName), /*bInUserTriggered = */true, /*bInRequiresLandscapeUpdate =*/ true, bHasValueChanged, /*InChangeType = */EPropertyChangeType::ValueSet);
+#endif // WITH_EDITOR
 }
 
+void ULandscapeLayerInfoObject::SetPhysicalMaterial(TObjectPtr<UPhysicalMaterial> InPhysicalMaterial, bool bInModify)
+{
+	const bool bHasValueChanged = (InPhysicalMaterial != PhysMaterial);
+	if (bHasValueChanged)
+	{
+		if (bInModify)
+		{
+			Modify();
+		}
+		PhysMaterial = InPhysicalMaterial;
+	}
+
 #if WITH_EDITOR
+	BroadcastOnLayerInfoObjectDataChanged(GET_MEMBER_NAME_CHECKED(ULandscapeLayerInfoObject, PhysMaterial), /*bInUserTriggered = */true, /*bInRequiresLandscapeUpdate =*/ true, bHasValueChanged, /*InChangeType = */EPropertyChangeType::ValueSet);
+#endif // WITH_EDITOR
+}
 
-void ULandscapeLayerInfoObject::PostEditChangeProperty(FPropertyChangedEvent& PropertyChangedEvent)
+void ULandscapeLayerInfoObject::SetHardness(float InHardness, bool bInModify, EPropertyChangeType::Type InChangeType)
 {
-	Super::PostEditChangeProperty(PropertyChangedEvent);
+	const bool bHasValueChanged = (InHardness != Hardness);
+	if (bHasValueChanged)
+	{
+		if (bInModify)
+		{
+			Modify();
+		}
+		Hardness = InHardness;
+	}
 
-	const FName MemberPropertyName = PropertyChangedEvent.MemberProperty ? PropertyChangedEvent.MemberProperty->GetFName() : NAME_None;
+#if WITH_EDITOR
+	BroadcastOnLayerInfoObjectDataChanged(GET_MEMBER_NAME_CHECKED(ULandscapeLayerInfoObject, Hardness), /*bInUserTriggered = */true, /*bInRequiresLandscapeUpdate =*/ false, bHasValueChanged, InChangeType);
+#endif // WITH_EDITOR
+}
 
-	// TODO [jonathan.bard] : automate the data changes via a data change event, so that we don't have to perform ad-hoc fixups like this 
+void ULandscapeLayerInfoObject::SetLayerUsageDebugColor(const FLinearColor& InLayerUsageDebugColor, bool bInModify, EPropertyChangeType::Type InChangeType)
+{
+	LayerUsageDebugColor.A = 1;
+	const bool bHasValueChanged = (InLayerUsageDebugColor != LayerUsageDebugColor);
 
-	if ((MemberPropertyName == GET_MEMBER_NAME_CHECKED(ULandscapeLayerInfoObject, PhysMaterial))
-		|| (MemberPropertyName == GET_MEMBER_NAME_CHECKED(ULandscapeLayerInfoObject, MinimumCollisionRelevanceWeight)))
+	if (bHasValueChanged)
 	{
-		for (TObjectIterator<ALandscapeProxy> It(/*AdditionalExclusionFlags = */RF_ClassDefaultObject, /*bIncludeDerivedClasses = */true, /*InInternalExclusionFlags = */EInternalObjectFlags::Garbage); It; ++It)
+		if (bInModify)
 		{
-			ALandscapeProxy* Proxy = *It;
-			if (Proxy->GetWorld() && !Proxy->GetWorld()->IsPlayInEditor())
-			{
-				ULandscapeInfo* Info = Proxy->GetLandscapeInfo();
-				if (Info)
-				{
-					for (int32 i = 0; i < Info->Layers.Num(); ++i)
-					{
-						if (Info->Layers[i].LayerInfoObj == this)
-						{
-							Proxy->ChangedPhysMaterial();
-							break;
-						}
-					}
-				}
-			}
+			Modify();
 		}
+		LayerUsageDebugColor = InLayerUsageDebugColor;
 	}
-	else if (MemberPropertyName == GET_MEMBER_NAME_CHECKED(ULandscapeLayerInfoObject, LayerUsageDebugColor))
+
+#if WITH_EDITOR
+	BroadcastOnLayerInfoObjectDataChanged(GET_MEMBER_NAME_CHECKED(ULandscapeLayerInfoObject, LayerUsageDebugColor), /*bInUserTriggered = */true, /*bInRequiresLandscapeUpdate =*/false, bHasValueChanged, InChangeType);
+#endif // WITH_EDITOR
+}
+
+#if WITH_EDITORONLY_DATA
+void ULandscapeLayerInfoObject::SetMinimumCollisionRelevanceWeight(float InMinimumCollisionRelevanceWeight, bool bInModify, EPropertyChangeType::Type InChangeType)
+{
+	const bool bHasValueChanged = (InMinimumCollisionRelevanceWeight != MinimumCollisionRelevanceWeight);
+	if (bHasValueChanged)
 	{
-		LayerUsageDebugColor.A = 1.0f;
-		for (TObjectIterator<ALandscapeProxy> It(/*AdditionalExclusionFlags = */RF_ClassDefaultObject, /*bIncludeDerivedClasses = */true, /*InInternalExclusionFlags = */EInternalObjectFlags::Garbage); It; ++It)
+		if (bInModify)
 		{
-			ALandscapeProxy* Proxy = *It;
-			if (Proxy->GetWorld() && !Proxy->GetWorld()->IsPlayInEditor())
-			{
-				Proxy->MarkComponentsRenderStateDirty();
-			}
+			Modify();
 		}
+		MinimumCollisionRelevanceWeight = InMinimumCollisionRelevanceWeight;
 	}
-	else if ((MemberPropertyName == GET_MEMBER_NAME_CHECKED(ULandscapeLayerInfoObject, SplineFalloffModulationTexture))
-		|| (MemberPropertyName == GET_MEMBER_NAME_CHECKED(ULandscapeLayerInfoObject, SplineFalloffModulationColorMask))
-		|| (MemberPropertyName == GET_MEMBER_NAME_CHECKED(ULandscapeLayerInfoObject, SplineFalloffModulationBias))
-		|| (MemberPropertyName == GET_MEMBER_NAME_CHECKED(ULandscapeLayerInfoObject, SplineFalloffModulationScale))
-		|| (MemberPropertyName == GET_MEMBER_NAME_CHECKED(ULandscapeLayerInfoObject, SplineFalloffModulationTiling)))
+
+	BroadcastOnLayerInfoObjectDataChanged(GET_MEMBER_NAME_CHECKED(ULandscapeLayerInfoObject, MinimumCollisionRelevanceWeight), /*bInUserTriggered = */true, /*bInRequiresLandscapeUpdate =*/ true, bHasValueChanged, InChangeType);
+}
+
+void ULandscapeLayerInfoObject::SetSplineFalloffModulationTexture(TObjectPtr<UTexture2D> InSplineFalloffModulationTexture, bool bInModify)
+{
+	const bool bHasValueChanged = (InSplineFalloffModulationTexture != SplineFalloffModulationTexture);
+	if (bHasValueChanged)
 	{
-		for (TObjectIterator<ULandscapeInfo> It(/*AdditionalExclusionFlags = */RF_ClassDefaultObject, /*bIncludeDerivedClasses = */true, /*InInternalExclusionFlags = */EInternalObjectFlags::Garbage); It; ++It)
+		if (bInModify)
 		{
-			if (ALandscape* Landscape = It->LandscapeActor.Get())
-			{
-				Landscape->OnLayerInfoSplineFalloffModulationChanged(this);
-			}
+			Modify();
 		}
+		SplineFalloffModulationTexture = InSplineFalloffModulationTexture;
 	}
-	else if (MemberPropertyName == GET_MEMBER_NAME_CHECKED(ULandscapeLayerInfoObject, BlendMethod))
+
+	BroadcastOnLayerInfoObjectDataChanged(GET_MEMBER_NAME_CHECKED(ULandscapeLayerInfoObject, SplineFalloffModulationTexture), /*bInUserTriggered = */true, /*bInRequiresLandscapeUpdate =*/ true, bHasValueChanged, /*InChangeType = */EPropertyChangeType::ValueSet);
+}
+
+void ULandscapeLayerInfoObject::SetSplineFalloffModulationColorMask(ESplineModulationColorMask InSplineFalloffModulationColorMask, bool bInModify)
+{
+	const bool bHasValueChanged = (InSplineFalloffModulationColorMask != SplineFalloffModulationColorMask);
+	if (bHasValueChanged)
 	{
-		SetBlendMethod(BlendMethod, /*bInModify = */true);
+		if (bInModify)
+		{
+			Modify();
+		}
+		SplineFalloffModulationColorMask = InSplineFalloffModulationColorMask;
 	}
-	else if (MemberPropertyName == GET_MEMBER_NAME_CHECKED(ULandscapeLayerInfoObject, BlendGroup))
+
+	BroadcastOnLayerInfoObjectDataChanged(GET_MEMBER_NAME_CHECKED(ULandscapeLayerInfoObject, SplineFalloffModulationColorMask), /*bInUserTriggered = */true, /*bInRequiresLandscapeUpdate =*/ true, bHasValueChanged, /*InChangeType = */EPropertyChangeType::ValueSet);
+}
+
+void ULandscapeLayerInfoObject::SetSplineFalloffModulationTiling(float InSplineFalloffModulationTiling, bool bInModify, EPropertyChangeType::Type InChangeType)
+{
+	const bool bHasValueChanged = (InSplineFalloffModulationTiling != SplineFalloffModulationTiling);
+	if (bHasValueChanged)
 	{
-		SetBlendGroup(BlendGroup, /*bInModify = */true);
+		if (bInModify)
+		{
+			Modify();
+		}
+		SplineFalloffModulationTiling = InSplineFalloffModulationTiling;
 	}
+
+	BroadcastOnLayerInfoObjectDataChanged(GET_MEMBER_NAME_CHECKED(ULandscapeLayerInfoObject, SplineFalloffModulationTiling), /*bInUserTriggered = */true, /*bInRequiresLandscapeUpdate =*/ true, bHasValueChanged, InChangeType);
 }
 
-void ULandscapeLayerInfoObject::PostEditUndo()
+void ULandscapeLayerInfoObject::SetSplineFalloffModulationBias(float InSplineFalloffModulationBias, bool bInModify, EPropertyChangeType::Type InChangeType)
 {
-	Super::PostEditUndo();
+	const bool bHasValueChanged = (InSplineFalloffModulationBias != SplineFalloffModulationBias);
+	if (bHasValueChanged)
+	{
+		if (bInModify)
+		{
+			Modify();
+		}
+		SplineFalloffModulationBias = InSplineFalloffModulationBias;
+	}
 
-	// TODO [jonathan.bard] : send event for data change and remove this object iterator nonsense : 
+	BroadcastOnLayerInfoObjectDataChanged(GET_MEMBER_NAME_CHECKED(ULandscapeLayerInfoObject, SplineFalloffModulationBias), /*bInUserTriggered = */true, /*bInRequiresLandscapeUpdate =*/ true, bHasValueChanged, InChangeType);
+}
 
-	// Force the update of spline data for the (potentially) affected landscapes in case of an undo
-	for (TObjectIterator<ULandscapeInfo> It(/*AdditionalExclusionFlags = */RF_ClassDefaultObject, /*bIncludeDerivedClasses = */true, /*InInternalExclusionFlags = */EInternalObjectFlags::Garbage); It; ++It)
+void ULandscapeLayerInfoObject::SetSplineFalloffModulationScale(float InSplineFalloffModulationScale, bool bInModify, EPropertyChangeType::Type InChangeType)
+{
+	const bool bHasValueChanged = (InSplineFalloffModulationScale != SplineFalloffModulationScale);
+	if (bHasValueChanged)
 	{
-		if (ALandscape* Landscape = It->LandscapeActor.Get())
+		if (bInModify)
 		{
-			Landscape->OnLayerInfoSplineFalloffModulationChanged(this);
+			Modify();
 		}
+		SplineFalloffModulationScale = InSplineFalloffModulationScale;
 	}
+
+	BroadcastOnLayerInfoObjectDataChanged(GET_MEMBER_NAME_CHECKED(ULandscapeLayerInfoObject, SplineFalloffModulationScale), /*bInUserTriggered = */true, /*bInRequiresLandscapeUpdate =*/ true, bHasValueChanged, InChangeType);
+}
+#endif // WITH_EDITORONLY_DATA
+PRAGMA_ENABLE_DEPRECATION_WARNINGS
+
+void ULandscapeLayerInfoObject::Serialize(FArchive& Ar)
+{
+	Super::Serialize(Ar);
+
+	Ar.UsingCustomVersion(FFortniteMainBranchObjectVersion::GUID);
+}
+
+#if WITH_EDITOR
+
+void ULandscapeLayerInfoObject::PostEditChangeProperty(FPropertyChangedEvent& PropertyChangedEvent)
+{
+	Super::PostEditChangeProperty(PropertyChangedEvent);
+
+	const FName MemberPropertyName = PropertyChangedEvent.MemberProperty ? PropertyChangedEvent.MemberProperty->GetFName() : NAME_None;
+
+	bool bRequiresLandscapeUpdate = true;
+	// Some properties don't need the landscape to be updated :
+	if (MemberPropertyName == GetHardnessMemberName() || MemberPropertyName == GetLayerUsageDebugColorMemberName())
+	{
+		bRequiresLandscapeUpdate = false;
+	}
+
+	BroadcastOnLayerInfoObjectDataChanged(MemberPropertyName, /*bInUserTriggered = */true, bRequiresLandscapeUpdate, /*bHasValueChanged =*/true, PropertyChangedEvent.ChangeType);
 }
 
 void ULandscapeLayerInfoObject::PostLoad()
 {
 	Super::PostLoad();
 
-#if WITH_EDITOR
 	if (GetLinkerCustomVersion(FFortniteMainBranchObjectVersion::GUID) < FFortniteMainBranchObjectVersion::LandscapeAdvancedWeightBlending)
 	{
 PRAGMA_DISABLE_DEPRECATION_WARNINGS
@@ -175,9 +278,24 @@ PRAGMA_ENABLE_DEPRECATION_WARNINGS
 		{
 			SetFlags(RF_Standalone);
 		}
-		Hardness = FMath::Clamp<float>(Hardness, 0.0f, 1.0f);
+		SetHardness(FMath::Clamp<float>(GetHardness(), 0.0f, 1.0f), /*bInModify =*/false, /*InChangeType =*/ EPropertyChangeType::ValueSet);
 	}
-#endif // WITH_EDITOR
+}
+
+void ULandscapeLayerInfoObject::BroadcastOnLayerInfoObjectDataChanged(FName InPropertyName, bool bInUserTriggered, bool bInRequiresLandscapeUpdate, bool bInHasValueChanged, EPropertyChangeType::Type InChangeType)
+{
+	FProperty* Property = nullptr;
+	if (InPropertyName != NAME_None)
+	{
+		Property = FindFProperty<FProperty>(StaticClass(), InPropertyName);
+		check(Property != nullptr);
+	}
+
+	FOnLandscapeLayerInfoDataChangedParams OnLayerInfoChangedParams(*this, FPropertyChangedEvent(Property, InChangeType));
+	OnLayerInfoChangedParams.bUserTriggered = bInUserTriggered;
+	OnLayerInfoChangedParams.bRequiresLandscapeUpdate = bInRequiresLandscapeUpdate;
+	OnLayerInfoChangedParams.bHasValueChanged = bInHasValueChanged;
+	OnLayerInfoObjectChangedDelegate.Broadcast(OnLayerInfoChangedParams);
 }
 
 #endif // WITH_EDITOR
diff --git a/Engine/Source/Runtime/Landscape/Private/LandscapeRender.cpp b/Engine/Source/Runtime/Landscape/Private/LandscapeRender.cpp
--- a/Engine/Source/Runtime/Landscape/Private/LandscapeRender.cpp
+++ b/Engine/Source/Runtime/Landscape/Private/LandscapeRender.cpp
@@ -1713,7 +1713,7 @@ FLandscapeComponentSceneProxy::FLandscapeComponentSceneProxy(ULandscapeComponent
 	{
 		if (Allocation.LayerInfo != nullptr)
 		{
-			LayerColors.Add(Allocation.LayerInfo->LayerUsageDebugColor);
+			LayerColors.Add(Allocation.LayerInfo->GetLayerUsageDebugColor());
 		}
 	}
 #endif
diff --git a/Engine/Source/Runtime/Landscape/Private/LandscapeSplineRaster.cpp b/Engine/Source/Runtime/Landscape/Private/LandscapeSplineRaster.cpp
--- a/Engine/Source/Runtime/Landscape/Private/LandscapeSplineRaster.cpp
+++ b/Engine/Source/Runtime/Landscape/Private/LandscapeSplineRaster.cpp
@@ -130,14 +130,14 @@ public:
 	{
 		if (CVarLandscapeSplineFalloffModulation.GetValueOnAnyThread() == 0 
 			|| (InLayerInfo == nullptr) 
-			|| (InLayerInfo->SplineFalloffModulationTexture == nullptr))
+			|| (InLayerInfo->GetSplineFalloffModulationTexture() == nullptr))
 		{
 			return nullptr;
 		}
 
-		if (!InLayerInfo->SplineFalloffModulationTexture->Source.IsValid())
+		if (!InLayerInfo->GetSplineFalloffModulationTexture()->Source.IsValid())
 		{
-			UE_LOG(LogLandscape, Error, TEXT("Invalid source data for spline falloff modulation texture (%s). Alpha modulation will be disabled."), *InLayerInfo->SplineFalloffModulationTexture->GetPathName());
+			UE_LOG(LogLandscape, Error, TEXT("Invalid source data for spline falloff modulation texture (%s). Alpha modulation will be disabled."), *InLayerInfo->GetSplineFalloffModulationTexture()->GetPathName());
 			return nullptr;
 		}
 
@@ -172,20 +172,20 @@ public:
 private:
 	FModulateAlpha(ULandscapeLayerInfoObject* InLayerInfo, int32 InLandscapeMinX, int32 InLandscapeMinY)
 	{
-		check(InLayerInfo && InLayerInfo->SplineFalloffModulationTexture);
+		check(InLayerInfo && InLayerInfo->GetSplineFalloffModulationTexture());
 
-		verify(InLayerInfo->SplineFalloffModulationTexture->Source.GetMipData(MipData, 0));
-		TextureWidth = InLayerInfo->SplineFalloffModulationTexture->Source.GetSizeX();
-		TextureHeight = InLayerInfo->SplineFalloffModulationTexture->Source.GetSizeY();
+		verify(InLayerInfo->GetSplineFalloffModulationTexture()->Source.GetMipData(MipData, 0));
+		TextureWidth = InLayerInfo->GetSplineFalloffModulationTexture()->Source.GetSizeX();
+		TextureHeight = InLayerInfo->GetSplineFalloffModulationTexture()->Source.GetSizeY();
 
-		Tiling = 1.0f / InLayerInfo->SplineFalloffModulationTiling;
-		Bias = InLayerInfo->SplineFalloffModulationBias;
-		Scale = InLayerInfo->SplineFalloffModulationScale;
+		Tiling = 1.0f / InLayerInfo->GetSplineFalloffModulationTiling();
+		Bias = InLayerInfo->GetSplineFalloffModulationBias();
+		Scale = InLayerInfo->GetSplineFalloffModulationScale();
 
 		LandscapeMinX = InLandscapeMinX;
 		LandscapeMinY = InLandscapeMinY;
 
-		ChannelOffset = ChannelOffsets[(int32)InLayerInfo->SplineFalloffModulationColorMask];
+		ChannelOffset = ChannelOffsets[(int32)InLayerInfo->GetSplineFalloffModulationColorMask()];
 	}
 
 private:
diff --git a/Engine/Source/Runtime/Landscape/Private/LandscapeSplines.cpp b/Engine/Source/Runtime/Landscape/Private/LandscapeSplines.cpp
--- a/Engine/Source/Runtime/Landscape/Private/LandscapeSplines.cpp
+++ b/Engine/Source/Runtime/Landscape/Private/LandscapeSplines.cpp
@@ -945,7 +945,7 @@ bool ULandscapeSplinesComponent::IsUsingLayerInfo(const ULandscapeLayerInfoObjec
 {
 	for (ULandscapeSplineControlPoint* ControlPoint : ControlPoints)
 	{
-		if (ControlPoint->LayerName == LayerInfo->LayerName)
+		if (ControlPoint->LayerName == LayerInfo->GetLayerName())
 		{
 			return true;
 		}
@@ -953,7 +953,7 @@ bool ULandscapeSplinesComponent::IsUsingLayerInfo(const ULandscapeLayerInfoObjec
 
 	for (ULandscapeSplineSegment* Segment : Segments)
 	{
-		if (Segment->LayerName == LayerInfo->LayerName)
+		if (Segment->LayerName == LayerInfo->GetLayerName())
 		{
 			return true;
 		}
diff --git a/Engine/Source/Runtime/Landscape/Private/LandscapeSubsystem.cpp b/Engine/Source/Runtime/Landscape/Private/LandscapeSubsystem.cpp
--- a/Engine/Source/Runtime/Landscape/Private/LandscapeSubsystem.cpp
+++ b/Engine/Source/Runtime/Landscape/Private/LandscapeSubsystem.cpp
@@ -141,7 +141,8 @@ namespace UE::Landscape
 								{
 									if (AllocInfo.IsAllocated())
 									{
-										FName LayerName = AllocInfo.LayerInfo->LayerName;
+										check(AllocInfo.LayerInfo != nullptr);
+										FName LayerName = AllocInfo.LayerInfo->GetLayerName();
 										TargetLayerToNumComponents.FindOrAdd(LayerName)++;
 										ComponentLayerNameSet.Add(LayerName);
 									}
diff --git a/Engine/Source/Runtime/Landscape/Private/LandscapeUtils.cpp b/Engine/Source/Runtime/Landscape/Private/LandscapeUtils.cpp
--- a/Engine/Source/Runtime/Landscape/Private/LandscapeUtils.cpp
+++ b/Engine/Source/Runtime/Landscape/Private/LandscapeUtils.cpp
@@ -299,8 +299,8 @@ ULandscapeLayerInfoObject* CreateTargetLayerInfo(const FName& InLayerName, const
 	}
 
 	check(LayerInfo != nullptr);
-	LayerInfo->LayerName = InLayerName;
-	LayerInfo->LayerUsageDebugColor = LayerInfo->GenerateLayerUsageDebugColor();
+	LayerInfo->SetLayerName(InLayerName, /*bInModify =*/false);
+	LayerInfo->SetLayerUsageDebugColor(LayerInfo->GenerateLayerUsageDebugColor(), /*bInModify =*/false, /*InChangeType =*/ EPropertyChangeType::ValueSet);
 
 	// Notify the asset registry
 	FAssetRegistryModule::AssetCreated(LayerInfo);
