Landscape Editor - Fixed a missing LayerInfoObjectAsset event registration when duplicating a landscape (non-WP duplication)

diff --git a/Engine/Source/Runtime/Landscape/Classes/LandscapeProxy.h b/Engine/Source/Runtime/Landscape/Classes/LandscapeProxy.h
--- a/Engine/Source/Runtime/Landscape/Classes/LandscapeProxy.h
+++ b/Engine/Source/Runtime/Landscape/Classes/LandscapeProxy.h
@@ -1319,6 +1319,7 @@ PRAGMA_ENABLE_DEPRECATION_WARNINGS
 	virtual void PostEditImport() override;
 	virtual void PreEditUndo() override;
 	virtual void PostEditUndo() override;
+	virtual void PostDuplicate(bool bDuplicateForPIE) override;
 	//~ End UObject Interface
 
 	LANDSCAPE_API void InitializeProxyLayersWeightmapUsage();
diff --git a/Engine/Source/Runtime/Landscape/Private/LandscapeEditLayers.cpp b/Engine/Source/Runtime/Landscape/Private/LandscapeEditLayers.cpp
--- a/Engine/Source/Runtime/Landscape/Private/LandscapeEditLayers.cpp
+++ b/Engine/Source/Runtime/Landscape/Private/LandscapeEditLayers.cpp
@@ -1514,6 +1514,21 @@ void ALandscapeProxy::PostEditUndo()
 		});
 }
 
+void ALandscapeProxy::PostDuplicate(bool bDuplicateForPIE)
+{
+	Super::PostDuplicate(bDuplicateForPIE);
+
+	// Duplicated target layers with layer info object assets need to register so that we can update the landscape and react to asset changes : 
+	if (!bDuplicateForPIE)
+	{
+		Algo::ForEach(GetValidTargetLayerObjects(), [this](ULandscapeLayerInfoObject* InLayerInfoObject)
+			{
+				check(!InLayerInfoObject->OnLayerInfoChanged().IsBoundToObject(this));
+				InLayerInfoObject->OnLayerInfoChanged().AddUObject(this, &ALandscapeProxy::OnLayerInfoObjectDataChanged);
+			});
+	}
+}
+
 void ALandscape::InitializeLandscapeLayersWeightmapUsage()
 {
 	ULandscapeInfo* Info = GetLandscapeInfo();
