Fixed ensure loading old landscape data with landscape.StripLayerMipsOnLoad = false

diff --git a/Engine/Source/Runtime/Landscape/Private/LandscapeEdit.cpp b/Engine/Source/Runtime/Landscape/Private/LandscapeEdit.cpp
--- a/Engine/Source/Runtime/Landscape/Private/LandscapeEdit.cpp
+++ b/Engine/Source/Runtime/Landscape/Private/LandscapeEdit.cpp
@@ -7180,12 +7180,6 @@ TArray<UTexture*> ULandscapeComponent::ReallocateWeightmaps(FLandscapeEditDataIn
 			}
 		}
 
-		if ((CurrentWeightmapTexture != nullptr) && CVarStripLayerTextureMipsOnLoad->GetBool())
-		{
-			// sanity check - only final weightmaps are allowed more than one mip
-			ensure((CurrentWeightmapTexture->Source.GetNumMips() == 1) || bIsFinalWeightmap);
-		}
-
 		bool NeedsUpdateResource = false;
 		// No suitable weightmap texture
 		if (CurrentWeightmapTexture == nullptr)
@@ -7209,7 +7203,7 @@ TArray<UTexture*> ULandscapeComponent::ReallocateWeightmaps(FLandscapeEditDataIn
 			}
 			else
 			{
-				// sanity check - only final weightmaps are allowed more than one mip
+				// sanity check - only final weightmaps are allowed more than one mip after creation
 				check(CurrentWeightmapTexture->Source.GetNumMips() == 1);
 			}
 
@@ -7222,6 +7216,11 @@ TArray<UTexture*> ULandscapeComponent::ReallocateWeightmaps(FLandscapeEditDataIn
 			CurrentWeightmapUsage->LayerGuid = InEditLayerGuid;
 			// UE_LOG(LogLandscape, Log, TEXT("Making a new texture %s"), *CurrentWeightmapTexture->GetName());
 		}
+		else
+		{
+			// sanity check - only final weightmaps are allowed more than one mip (except if we load old data with CVarStripLayerTextureMipsOnLoad set to false)
+			check(!CVarStripLayerTextureMipsOnLoad->GetBool() || (CurrentWeightmapTexture->Source.GetNumMips() == 1) || bIsFinalWeightmap);
+		}
 
 		NewComponentWeightmapTexturesUsage.Add(CurrentWeightmapUsage);
 		NewWeightmapTextures.Add(CurrentWeightmapTexture);
@@ -7250,8 +7249,6 @@ TArray<UTexture*> ULandscapeComponent::ReallocateWeightmaps(FLandscapeEditDataIn
 					// Copy the data
 					if (ensure(DataInterface != nullptr)) // it's not safe to skip the copy
 					{
-						// sanity check - only final weightmaps are allowed more than one mip
-						ensure((CurrentWeightmapTexture->Source.GetNumMips() == 1) || bIsFinalWeightmap);
 						DataInterface->CopyTextureChannel(CurrentWeightmapTexture, ChanIdx, OldWeightmapTexture, AllocInfo.WeightmapTextureChannel);
 						DataInterface->ZeroTextureChannel(OldWeightmapTexture, AllocInfo.WeightmapTextureChannel);
 						// UE_LOG(LogLandscape, Log, TEXT("Copying old channel (%s).%d to new channel (%s).%d"), *OldWeightmapTexture->GetName(), AllocInfo.WeightmapTextureChannel, *CurrentWeightmapTexture->GetName(), ChanIdx);
