Landscape Editor - Changed ALandscape SelectedEditLayerIndex default value to 0, since non-edit layer landscapes are no longer supported
* Fixed Duplicating landscape crash (non-WP). UI attemps to access SelectedIndex on new Landscape before PostDuplicate runs

diff --git a/Engine/Source/Runtime/Landscape/Classes/Landscape.h b/Engine/Source/Runtime/Landscape/Classes/Landscape.h
--- a/Engine/Source/Runtime/Landscape/Classes/Landscape.h
+++ b/Engine/Source/Runtime/Landscape/Classes/Landscape.h
@@ -679,8 +679,8 @@ public:
 	FGuid EditingLayer;
 
 	/** Current Selected Edit Layer of this landscape. Used by landscape editor mode to track the current selection */
-	UPROPERTY(Transient)
-	int32 SelectedEditLayerIndex = INDEX_NONE;
+	UPROPERTY(Transient, DuplicateTransient)
+	int32 SelectedEditLayerIndex = 0;
 
 	/** Used to temporarily disable Grass Update in Editor */
 	bool bGrassUpdateEnabled;
diff --git a/Engine/Source/Runtime/Landscape/Private/Landscape.cpp b/Engine/Source/Runtime/Landscape/Private/Landscape.cpp
--- a/Engine/Source/Runtime/Landscape/Private/Landscape.cpp
+++ b/Engine/Source/Runtime/Landscape/Private/Landscape.cpp
@@ -3254,13 +3254,8 @@ PRAGMA_ENABLE_DEPRECATION_WARNINGS
 			CreateDefaultLayer();
 		}
 
-		// Ensure a valid edit layer is selected post load
-		if (!LandscapeEditLayers.IsValidIndex(SelectedEditLayerIndex))
-		{
-			SetSelectedEditLayerIndex(0);
-		}
-
 		ensure(!LandscapeEditLayers.IsEmpty());
+		ensure(SelectedEditLayerIndex == 0);
 	}
 #endif // WITH_EDITOR
 
diff --git a/Engine/Source/Runtime/Landscape/Private/LandscapeEdit.cpp b/Engine/Source/Runtime/Landscape/Private/LandscapeEdit.cpp
--- a/Engine/Source/Runtime/Landscape/Private/LandscapeEdit.cpp
+++ b/Engine/Source/Runtime/Landscape/Private/LandscapeEdit.cpp
@@ -3822,8 +3822,7 @@ void ALandscapeProxy::Import(const FGuid& InGuid, int32 InMinX, int32 InMinY, in
 
 		ComponentsToProcess.Append(ToRawPtrTArrayUnsafe(LandscapeComponents));
 	}
-	// CreateDefaultLayer or DuplicateLayerAndMoveBrushes must set the SelectedEditLayerIndex during import
-	check(LandscapeActor->GetSelectedEditLayerIndex() != INDEX_NONE);
+	check(LandscapeActor->GetEditLayersConst().IsValidIndex(LandscapeActor->GetSelectedEditLayerIndex()));
 	check(LayerImportSettings.Num() != 0);
 	// Currently only supports reimporting heightmap data into a single edit layer, which will always be the default layer
 	ReimportDestinationLayerGuid = LayerImportSettings[0].DestinationLayerGuid;
@@ -5328,9 +5327,6 @@ void ALandscape::PostDuplicate(bool bDuplicateForPIE)
 		// Some edit layers could be affected by BP brushes, which might need to be updated when the landscape is transformed :
 		RequestLayersContentUpdate(ELandscapeLayerUpdateMode::Update_All);
 	}
-
-	// Ensure the first edit layer is selected after duplicating
-	SetSelectedEditLayerIndex(0);
 }
 #endif	//WITH_EDITOR
 
diff --git a/Engine/Source/Runtime/Landscape/Private/LandscapeEditLayers.cpp b/Engine/Source/Runtime/Landscape/Private/LandscapeEditLayers.cpp
--- a/Engine/Source/Runtime/Landscape/Private/LandscapeEditLayers.cpp
+++ b/Engine/Source/Runtime/Landscape/Private/LandscapeEditLayers.cpp
@@ -8408,7 +8408,6 @@ bool ALandscape::DeleteLayer(int32 InLayerIndex)
 	Modify();
 
 	// If the layer to delete is below the current selected layer index, shift the selected index down
-	// Ignore index update when toggling from an layer to a non-layer state 
 	if (SelectedEditLayerIndex >= InLayerIndex)
 	{
 		if (LandscapeEditLayers.IsValidIndex(InLayerIndex - 1))
@@ -9005,13 +9004,7 @@ FLandscapeLayer* ALandscape::DuplicateLayerAndMoveBrushes(const FLandscapeLayer&
 	}
 
 	int32 AddedIndex = LandscapeEditLayers.Add(NewLayer);
-
-	// If this is the first edit layer duplicated, auto select it so our landscape actor maintains a valid selection state
-	if (LandscapeEditLayers.Num() == 1 && SelectedEditLayerIndex == INDEX_NONE)
-	{
-		SetSelectedEditLayerIndex(0);
-	}
-	check(SelectedEditLayerIndex != INDEX_NONE);
+	check(LandscapeEditLayers.IsValidIndex(SelectedEditLayerIndex));
 
 	OnLayerCreatedInternal(NewLayer.EditLayer);
 
