* Prevent landscape edge fixup from running on world composition maps, where heightmap sharing is still possible
* Fixed all logs involving landscape LandcapeGroupKey and LODGroupKey

diff --git a/Engine/Source/Runtime/Landscape/Private/LandscapeEdgeFixup.cpp b/Engine/Source/Runtime/Landscape/Private/LandscapeEdgeFixup.cpp
--- a/Engine/Source/Runtime/Landscape/Private/LandscapeEdgeFixup.cpp
+++ b/Engine/Source/Runtime/Landscape/Private/LandscapeEdgeFixup.cpp
@@ -754,7 +754,7 @@ void ULandscapeHeightmapTextureEdgeFixup::SetActiveComponent(ULandscapeComponent
 			// deactivate the active component/group, add it to the disabled list
 			ActiveGroup->DisableAndUnmap(this);
 
-			FIXUP_DEBUG_LOG_DETAIL(TEXT("    DisableAndUnmap %p (%d,%d) <Group %d Key %d> edgefixup:%p MAPPED:%d TOTAL:%d"),
+			FIXUP_DEBUG_LOG_DETAIL(TEXT("    DisableAndUnmap %p (%d,%d) <Group %u Key %u> edgefixup:%p MAPPED:%d TOTAL:%d"),
 				ActiveComponent.Get(), GroupCoord.X, GroupCoord.Y, ActiveComponent->GetLandscapeProxy()->LODGroupKey, ActiveGroup->LandscapeGroupKey, this,
 				ActiveGroup->XYToEdgeFixupMap.Num(), ActiveGroup->AllRegisteredFixups.Num());
 		}
@@ -763,7 +763,7 @@ void ULandscapeHeightmapTextureEdgeFixup::SetActiveComponent(ULandscapeComponent
 			// unmap, don't flag as disabled
 			ActiveGroup->Unmap(this);
 
-			FIXUP_DEBUG_LOG_DETAIL(TEXT("    Unmap %p (%d,%d) <Group %d Key %d> edgefixup:%p MAPPED:%d TOTAL:%d"),
+			FIXUP_DEBUG_LOG_DETAIL(TEXT("    Unmap %p (%d,%d) <Group %u Key %u> edgefixup:%p MAPPED:%d TOTAL:%d"),
 				ActiveComponent.Get(), GroupCoord.X, GroupCoord.Y, ActiveComponent->GetLandscapeProxy()->LODGroupKey, ActiveGroup->LandscapeGroupKey, this,
 				ActiveGroup->XYToEdgeFixupMap.Num(), ActiveGroup->AllRegisteredFixups.Num());
 		}
@@ -780,7 +780,7 @@ void ULandscapeHeightmapTextureEdgeFixup::SetActiveComponent(ULandscapeComponent
 	{
 		InGroup->Map(this, InComponent);
 
-		FIXUP_DEBUG_LOG_DETAIL(TEXT("    Map %p (%d,%d) <Group %d Key %d> edgefixup:%p MAPPED:%d TOTAL:%d "),
+		FIXUP_DEBUG_LOG_DETAIL(TEXT("    Map %p (%d,%d) <Group %u Key %u> edgefixup:%p MAPPED:%d TOTAL:%d "),
 			InComponent, GroupCoord.X, GroupCoord.Y, InComponent->GetLandscapeProxy()->LODGroupKey, InGroup->LandscapeGroupKey, this,
 			InGroup->XYToEdgeFixupMap.Num(), InGroup->AllRegisteredFixups.Num());
 	}
diff --git a/Engine/Source/Runtime/Landscape/Private/LandscapeGroup.cpp b/Engine/Source/Runtime/Landscape/Private/LandscapeGroup.cpp
--- a/Engine/Source/Runtime/Landscape/Private/LandscapeGroup.cpp
+++ b/Engine/Source/Runtime/Landscape/Private/LandscapeGroup.cpp
@@ -88,8 +88,11 @@ void FLandscapeGroup::RegisterComponent(ULandscapeComponent* Component)
 	ALandscapeProxy* LandscapeProxy = Component->GetLandscapeProxy();
 	check(LandscapeProxy != nullptr);
 
-	// TODO [jonathan.bard] : Remove this once we remove shared heightmaps, so we can enable edge fixup on non-WP landscapes.
-	if (!LandscapeProxy->IsA<ALandscapeStreamingProxy>())
+	const UWorld* ComponentWorld = Component->GetWorld();
+	check(ComponentWorld != nullptr);
+
+	// TODO [jonathan.bard] : Remove this once we remove shared heightmaps, so we can enable edge fixup on non-WP landscapes (and world composition levels too).
+	if (!LandscapeProxy->IsA<ALandscapeStreamingProxy>() || (ComponentWorld->WorldComposition != nullptr))
 	{
 		return;
 	}
@@ -127,9 +130,8 @@ void FLandscapeGroup::RegisterComponent(ULandscapeComponent* Component)
 			check(ActiveComponent != Component);
 			check(ActiveComponent->RegisteredLandscapeGroup != nullptr);
 
-			UWorld* NewWorld = Component->GetWorld();
-			UWorld* ActiveWorld = ActiveComponent->GetWorld();
-			check(NewWorld != ActiveWorld);
+			const UWorld* ActiveComponentWorld = ActiveComponent->GetWorld();
+			check(ComponentWorld != ActiveComponentWorld);
 			
 			// since there is already an active component, edge fixup is already created and registered, we can just grab it
 			EdgeFixup = ActiveComponent->RegisteredEdgeFixup;
@@ -450,7 +452,7 @@ FIntPoint FLandscapeGroup::Map(ULandscapeHeightmapTextureEdgeFixup* Fixup, ULand
 
 			if (!(bResolutionMatches && bXVectorMatches && bYVectorMatches))
 			{
-				UE_LOG(LogLandscape, Warning, TEXT("Landscapes in LOD Group with Key %d do not have matching resolution (%d == %d), scale (%f == %f, %f == %f) and/or rotation; seam artifacts may appear."),
+				UE_LOG(LogLandscape, Warning, TEXT("Landscapes in LOD Group with Key %u do not have matching resolution (%d == %d), scale (%f == %f, %f == %f) and/or rotation; seam artifacts may appear."),
 					LandscapeGroupKey,
 					ComponentResolution, Component->ComponentSizeQuads,
 					GroupCoordXVector.Length(), ComponentXVectorWorldSpace.Length(),
@@ -460,7 +462,7 @@ FIntPoint FLandscapeGroup::Map(ULandscapeHeightmapTextureEdgeFixup* Fixup, ULand
 
 		if (!(LandscapeGridScale - ComponentLandscapeGridScale).IsNearlyZero())
 		{
-			UE_LOG(LogLandscape, Warning, TEXT("Landscapes in LOD Group with Key %d do not have matching grid scale (%s == %s); seam artifacts may appear."),
+			UE_LOG(LogLandscape, Warning, TEXT("Landscapes in LOD Group with Key %u do not have matching grid scale (%s == %s); seam artifacts may appear."),
 				LandscapeGroupKey,
 				*LandscapeGridScale.ToString(), *ComponentLandscapeGridScale.ToString());
 		}
@@ -477,7 +479,7 @@ FIntPoint FLandscapeGroup::Map(ULandscapeHeightmapTextureEdgeFixup* Fixup, ULand
 
 	if ((FMath::Abs(GroupCoordX - GroupCoord.X) > 0.01) || (FMath::Abs(GroupCoordY - GroupCoord.Y) > 0.01))
 	{
-		UE_LOG(LogLandscape, Warning, TEXT("Landscape component %s is not spatially aligned with the group grid (%d), seam artifacts may appear."), *Component->GetPathName(), LandscapeGroupKey);
+		UE_LOG(LogLandscape, Warning, TEXT("Landscape component %s is not spatially aligned with the group grid (LOD Group with Key %u), seam artifacts may appear."), *Component->GetPathName(), LandscapeGroupKey);
 	}
 
 	// Add to the Map
@@ -496,7 +498,7 @@ FIntPoint FLandscapeGroup::Map(ULandscapeHeightmapTextureEdgeFixup* Fixup, ULand
 		}
 		else
 		{
-			UE_LOG(LogLandscape, Warning, TEXT("Two landscape components in group (%d) occupy the same group grid cell: %s and %s, artifacts may appear. Please move one."),
+			UE_LOG(LogLandscape, Warning, TEXT("Two landscape components in group (LOD Group with Key %u) occupy the same group grid cell: %s and %s, artifacts may appear. Please move one."),
 				LandscapeGroupKey, *Component->GetPathName(), *(*MapEntryPtrPtr)->ActiveComponent->GetPathName());
 
 			// attempt to give a temporary location to the overlapped component
@@ -704,13 +706,13 @@ void FLandscapeGroup::TickEdgeFixup(ULandscapeSubsystem* LandscapeSubsystem, boo
 			FVector CurrentLandscapeGridScale = Component->GetLandscapeProxy()->GetRootComponent()->GetRelativeScale3D();
 			if (!ensure((CurrentLandscapeGridScale - LandscapeGridScale).IsNearlyZero()))
 			{
-				UE_LOG(LogLandscape, Warning, TEXT("Landscape component %p (%s) has a different scale than other components in the Landscape Group, normal seam artifacts may occur"), Component, *Component->GetName());
+				UE_LOG(LogLandscape, Warning, TEXT("Landscape component %p (%s) has a different scale than other components in the Landscape Group with Key %u, normal seam artifacts may occur"), Component, *Component->GetName(), this->LandscapeGroupKey);
 			}
 
 			FLandscapeGroup* NewGroup = LandscapeSubsystem->GetLandscapeGroupForComponent(Component);
 			if (!ensure(NewGroup == this))
 			{
-				UE_LOG(LogLandscape, Warning, TEXT("Landscape component %p (%s) changed groups unexpectedly (%x ==> %x) - attempting to fix automatically by re-registering it"), Component, *Component->GetName(), this->LandscapeGroupKey, NewGroup->LandscapeGroupKey);
+				UE_LOG(LogLandscape, Warning, TEXT("Landscape component %p (%s) changed groups unexpectedly (%u ==> %u) - attempting to fix automatically by re-registering it"), Component, *Component->GetName(), this->LandscapeGroupKey, NewGroup->LandscapeGroupKey);
 				this->UnregisterComponent(Component);
 				NewGroup->RegisterComponent(Component);
 			}
@@ -865,7 +867,7 @@ FLandscapeGroup::~FLandscapeGroup()
 {
 	if (AllRegisteredFixups.Num() > 0)
 	{
-		UE_LOG(LogLandscape, Warning, TEXT("Landscape Group (%d) had registered components at destruction, this indicates some components were not properly unregistered"), this->LandscapeGroupKey);
+		UE_LOG(LogLandscape, Warning, TEXT("Landscape Group with Key %u had registered components at destruction, this indicates some components were not properly unregistered"), this->LandscapeGroupKey);
 		
 		// Avoid UnregisterComponent removing from AllRegisteredFixups while iterating through it
 		TSet<TObjectPtr<ULandscapeHeightmapTextureEdgeFixup>> LocalAllRegisteredFixups(AllRegisteredFixups);
diff --git a/Engine/Source/Runtime/Landscape/Private/LandscapeRender.cpp b/Engine/Source/Runtime/Landscape/Private/LandscapeRender.cpp
--- a/Engine/Source/Runtime/Landscape/Private/LandscapeRender.cpp
+++ b/Engine/Source/Runtime/Landscape/Private/LandscapeRender.cpp
@@ -733,7 +733,7 @@ void FLandscapeRenderSystem::CreateResources_Internal(FRHICommandListBase& RHICm
 			bool bYVectorMatches = (SectionYVector - ComponentYVector).IsNearlyZero();
 			if (!(bResolutionMatches && bXVectorMatches && bYVectorMatches))
 			{
-				UE_LOG(LogLandscape, Warning, TEXT("Landscapes in LOD Group %d do not have matching resolution (%d == %d), scale (%f == %f, %f == %f) and/or rotation; geometry seam artifacts may appear. If using HLOD, it may need to be updated to reflect changed resolutions or transforms."),
+				UE_LOG(LogLandscape, Warning, TEXT("Landscapes in LOD Group %u do not have matching resolution (%d == %d), scale (%f == %f, %f == %f) and/or rotation; geometry seam artifacts may appear. If using HLOD, it may need to be updated to reflect changed resolutions or transforms."),
 					SectionInfo->LODGroupKey,
 					ComponentResolution, SectionComponentResolution,
 					ComponentXVector.Length(), SectionXVector.Length(),
