Fixed checkSlow in landscape code + misc comment fixes

diff --git a/Engine/Source/Runtime/Landscape/Private/LandscapeEdit.cpp b/Engine/Source/Runtime/Landscape/Private/LandscapeEdit.cpp
--- a/Engine/Source/Runtime/Landscape/Private/LandscapeEdit.cpp
+++ b/Engine/Source/Runtime/Landscape/Private/LandscapeEdit.cpp
@@ -330,6 +330,7 @@ namespace UE::Landscape::Private
 bool ULandscapeComponent::UpdateCachedBounds(bool bInApproximateBounds)
 {
 	TRACE_CPUPROFILER_EVENT_SCOPE(ULandscapeComponent::UpdateCachedBounds);
+	check(IsInParallelGameThread() || IsInGameThread());
 
 	if ((SubsectionSizeQuads == 0) || (NumSubsections == 0) || (ComponentSizeQuads == 0))
 	{
@@ -402,6 +403,8 @@ bool ULandscapeComponent::UpdateCachedBounds(bool bInApproximateBounds)
 				ParallelForWithExistingTaskContext(MakeArrayView(ContextValues), NumQuadsForMip, MaxBatchSize,
 					[&MipQuadInfos, &CDI, MipQuadsStride, MipTextureSubSectionSize, MipIndex](FMinMaxType& Context, int32 MipQuadIndex)
 					{
+						FTaskTagScope TaskTag(ETaskTag::EParallelGameThread);
+
 						int32 QuadY = MipQuadIndex / MipQuadsStride;
 						int32 QuadX = MipQuadIndex - QuadY * MipQuadsStride;
 						int32 X = QuadX * 2;
diff --git a/Engine/Source/Runtime/Landscape/Public/LandscapeEdit.h b/Engine/Source/Runtime/Landscape/Public/LandscapeEdit.h
--- a/Engine/Source/Runtime/Landscape/Public/LandscapeEdit.h
+++ b/Engine/Source/Runtime/Landscape/Public/LandscapeEdit.h
@@ -229,7 +229,7 @@ struct FLandscapeEditDataInterface : public FLandscapeTextureDataInterface
 	//void GetWeightData(FName LayerName, int32& X1, int32& Y1, int32& X2, int32& Y2, TMap<uint64, TArray<uint8>>& SparseData);
 	LANDSCAPE_API void GetWeightDataFast(ULandscapeLayerInfoObject* LayerInfo, const int32 X1, const int32 Y1, const int32 X2, const int32 Y2, TMap<FIntPoint, uint8>& SparseData);
 	LANDSCAPE_API void GetWeightDataFast(ULandscapeLayerInfoObject* LayerInfo, const int32 X1, const int32 Y1, const int32 X2, const int32 Y2, TMap<FIntPoint, TArray<uint8>>& SparseData);
-	// Updates weightmap for LayerInfo, optionally adjusting all other weightmaps.
+	// Updates weightmap for LayerInfo
 	LANDSCAPE_API void SetAlphaData(ULandscapeLayerInfoObject* LayerInfo, const int32 X1, const int32 Y1, const int32 X2, const int32 Y2, const uint8* Data, int32 Stride, ELandscapeLayerPaintingRestriction PaintingRestriction = ELandscapeLayerPaintingRestriction::None);
 	UE_DEPRECATED(5.7, "Support for bWeightAdjust and bTotalWeightAdjust have been removed, this function shall be used anymore")
 	LANDSCAPE_API void SetAlphaData(ULandscapeLayerInfoObject* LayerInfo, const int32 X1, const int32 Y1, const int32 X2, const int32 Y2, const uint8* Data, int32 Stride, ELandscapeLayerPaintingRestriction PaintingRestriction, bool bWeightAdjust, bool bTotalWeightAdjust);
