Landscape Editor - Moved LayerInfoObjectAsset event registration code  into generic PostRegister workflow to fix failing asserts
* Fixed failng assert when using Save Current Level as on a landscape with target layers
* Fixed failing assert when using the Undo operation after deleting a non-WP landscape with target layers

diff --git a/Engine/Source/Runtime/Landscape/Classes/LandscapeProxy.h b/Engine/Source/Runtime/Landscape/Classes/LandscapeProxy.h
--- a/Engine/Source/Runtime/Landscape/Classes/LandscapeProxy.h
+++ b/Engine/Source/Runtime/Landscape/Classes/LandscapeProxy.h
@@ -1317,9 +1317,7 @@ PRAGMA_ENABLE_DEPRECATION_WARNINGS
 	virtual void PostTransacted(const FTransactionObjectEvent& InTransactionEvent) override;
 	virtual void PostEditChangeProperty(FPropertyChangedEvent& PropertyChangedEvent) override;
 	virtual void PostEditImport() override;
-	virtual void PreEditUndo() override;
 	virtual void PostEditUndo() override;
-	virtual void PostDuplicate(bool bDuplicateForPIE) override;
 	//~ End UObject Interface
 
 	LANDSCAPE_API void InitializeProxyLayersWeightmapUsage();
diff --git a/Engine/Source/Runtime/Landscape/Private/Landscape.cpp b/Engine/Source/Runtime/Landscape/Private/Landscape.cpp
--- a/Engine/Source/Runtime/Landscape/Private/Landscape.cpp
+++ b/Engine/Source/Runtime/Landscape/Private/Landscape.cpp
@@ -6,8 +6,6 @@ Landscape.cpp: Terrain rendering
 
 #include "Landscape.h"
 
-#include "Algo/AllOf.h"
-#include "Algo/BinarySearch.h"
 #include "Algo/Count.h"
 #include "Algo/ForEach.h"
 #include "Algo/IsSorted.h"
@@ -3091,6 +3089,15 @@ void ALandscapeProxy::PostRegisterAllComponents()
 	if ((LandscapeInfo != nullptr) && !IsPendingKillPending() && LandscapeGuid.IsValid())
 	{
 		LandscapeInfo->FixupProxiesTransform();
+
+		if (LandscapeInfo->IsLandscapeEditableWorld())
+		{
+			// Register to data change events on the layer info object so that we can update the landscape and react to changes : 
+			Algo::ForEach(GetValidTargetLayerObjects(), [this](ULandscapeLayerInfoObject* InLayerInfoObject)
+				{
+					InLayerInfoObject->OnLayerInfoChanged().AddUObject(this, &ALandscapeProxy::OnLayerInfoObjectDataChanged);
+				});
+		}
 	}
 #endif // WITH_EDITOR
 }
@@ -3106,6 +3113,16 @@ void ALandscapeProxy::UnregisterAllComponents(const bool bForReregister)
 		ULandscapeInfo* LandscapeInfo = GetLandscapeInfo();
 		if (LandscapeInfo)
 		{
+#if WITH_EDITOR
+			if (LandscapeInfo->IsLandscapeEditableWorld())
+			{
+				// Unregister from data change events on the layer info objects
+				Algo::ForEach(GetValidTargetLayerObjects(), [this](ULandscapeLayerInfoObject* InLayerInfoObject)
+					{
+						InLayerInfoObject->OnLayerInfoChanged().RemoveAll(this);
+					});
+			}
+#endif // WITH_EDITOR
 			LandscapeInfo->UnregisterActor(this);
 		}
 
@@ -4503,13 +4520,6 @@ void ALandscapeProxy::PostLoad()
 	PostLoadFixupLandscapeGuidsIfInstanced();
 
 #if WITH_EDITOR
-	// Register to data change events on the layer info object so that we can update the landscape and react to changes : 
-	Algo::ForEach(GetValidTargetLayerObjects(), [this](ULandscapeLayerInfoObject* InLayerInfoObject)
-		{
-			check(!InLayerInfoObject->OnLayerInfoChanged().IsBoundToObject(this));
-			InLayerInfoObject->OnLayerInfoChanged().AddUObject(this, &ALandscapeProxy::OnLayerInfoObjectDataChanged);
-		});
-
 	FixupOverriddenSharedProperties();
 
 	ALandscape* LandscapeActor = GetLandscapeActor();
@@ -5017,9 +5027,6 @@ TArray<FName> ALandscapeProxy::SynchronizeSharedProperties(ALandscapeProxy* InLa
 {
 	TArray<FName> SynchronizedProperties = SynchronizeUnmarkedSharedProperties(InLandscape);
 
-	// We need to unregister events from the target layers that we have lost because of the synchronization of the TargetLayers property : 
-	const TSet<ULandscapeLayerInfoObject*> ValidLayerInfoObjectsBefore = GetValidTargetLayerObjects();
-
 	for (TFieldIterator<FProperty> PropertyIterator(GetClass()); PropertyIterator; ++PropertyIterator)
 	{
 		FProperty* Property = *PropertyIterator;
@@ -5040,24 +5047,6 @@ TArray<FName> ALandscapeProxy::SynchronizeSharedProperties(ALandscapeProxy* InLa
 
 	if (!SynchronizedProperties.IsEmpty())
 	{
-		const TSet<ULandscapeLayerInfoObject*> ValidLayerInfoObjectsAfter = GetValidTargetLayerObjects();
-
-		// Add data changed events from layers we have gain
-		const TSet<ULandscapeLayerInfoObject*> LayerInfoObjectsToRegisterTo = ValidLayerInfoObjectsAfter.Difference(ValidLayerInfoObjectsBefore);
-		Algo::ForEach(LayerInfoObjectsToRegisterTo, [this](ULandscapeLayerInfoObject* InLayerInfoObject)
-			{
-				check(!InLayerInfoObject->OnLayerInfoChanged().IsBoundToObject(this));
-				InLayerInfoObject->OnLayerInfoChanged().AddUObject(this, &ALandscapeProxy::OnLayerInfoObjectDataChanged);
-			});
-
-		// Remove data changed events from layers we have lost
-		const TSet<ULandscapeLayerInfoObject*> LayerInfoObjectsToUnregisterFrom = ValidLayerInfoObjectsBefore.Difference(ValidLayerInfoObjectsAfter);
-		Algo::ForEach(LayerInfoObjectsToUnregisterFrom, [this](ULandscapeLayerInfoObject* InLayerInfoObject)
-			{
-				check(InLayerInfoObject->OnLayerInfoChanged().IsBoundToObject(this));
-				InLayerInfoObject->OnLayerInfoChanged().RemoveAll(this);
-			});
-
 		Modify();
 	}
 
@@ -7387,12 +7376,11 @@ bool ALandscapeProxy::RemoveTargetLayer(const FName& Name,  bool bPostEditChange
 {
 	Modify();
 	
-	// Remove data changed events
+	// Remove data changed event listeners. The LayerInfoObjectAsset will not be reachable from the TargetLayers list in the next UnregisterAllComponents
 	if (FLandscapeTargetLayerSettings* LayerSettings = TargetLayers.Find(Name))
 	{
 		if (LayerSettings && LayerSettings->LayerInfoObj)
 		{
-			check(LayerSettings->LayerInfoObj->OnLayerInfoChanged().IsBoundToObject(this));
 			LayerSettings->LayerInfoObj->OnLayerInfoChanged().RemoveAll(this);
 		}
 	}
@@ -7416,6 +7404,7 @@ FLandscapeTargetLayerSettings& ALandscapeProxy::AddTargetLayer(const FName& Name
 {
 	Modify();
 
+	// LayerInfoObjectAsset events will be registered in PostRegisterAllComponents after PostEditChangeProperty runs
 	FLandscapeTargetLayerSettings& Settings = TargetLayers.Add(Name.IsNone() ? GenerateUniqueTargetLayerName() : Name, TargetLayerSettings);
 	if (FProperty* Property = StaticClass()->FindPropertyByName(GET_MEMBER_NAME_CHECKED(ALandscapeProxy, TargetLayers)); Property != nullptr && bPostEditChange)
 	{
@@ -7423,13 +7412,6 @@ FLandscapeTargetLayerSettings& ALandscapeProxy::AddTargetLayer(const FName& Name
 		PostEditChangeProperty(PropertyChangedEvent);	
 	}
 	
-	// Register to data change events on the layer info object so that we can update the landscape accordingly : 
-	if (Settings.LayerInfoObj)
-	{
-		check(!Settings.LayerInfoObj->OnLayerInfoChanged().IsBoundToObject(this));
-		Settings.LayerInfoObj->OnLayerInfoChanged().AddUObject(this, &ALandscapeProxy::OnLayerInfoObjectDataChanged);
-	}
-	
 	return Settings;
 }
 
@@ -7454,28 +7436,21 @@ bool ALandscapeProxy::UpdateTargetLayer(const FName& Name, const FLandscapeTarge
 	{
 		Modify();
 
-		// We may be updating the layer info here. Remove old event listeners
+		// We may be updating the layer info asset here. The removed LayerInfoObjectAsset will not be reachable from the TargetLayers list in the next UnregisterAllComponents
 		if (TargetLayerSettings->LayerInfoObj != nullptr)
 		{
-			check(TargetLayerSettings->LayerInfoObj->OnLayerInfoChanged().IsBoundToObject(this));
 			TargetLayerSettings->LayerInfoObj->OnLayerInfoChanged().RemoveAll(this);
 		}
 
 		*TargetLayerSettings = InTargetLayerSettings;
 		
+		// LayerInfoObjectAsset event listeners will be re-registered in PostRegisterAllComponents after PostEditChangeProperty runs
 		if (FProperty* Property = StaticClass()->FindPropertyByName(GET_MEMBER_NAME_CHECKED(ALandscapeProxy, TargetLayers));  Property != nullptr && bPostEditChange)
 		{
 			FPropertyChangedEvent PropertyChangedEvent(Property);
 			PostEditChangeProperty(PropertyChangedEvent);
 		}
 
-		// Register to data change events on the current layer info object so that we can update the landscape accordingly : 
-		if (TargetLayerSettings->LayerInfoObj != nullptr)
-		{
-			check(!TargetLayerSettings->LayerInfoObj->OnLayerInfoChanged().IsBoundToObject(this));
-			TargetLayerSettings->LayerInfoObj->OnLayerInfoChanged().AddUObject(this, &ALandscapeProxy::OnLayerInfoObjectDataChanged);
-		}
-
 		return true;
 	}
 
diff --git a/Engine/Source/Runtime/Landscape/Private/LandscapeEdit.cpp b/Engine/Source/Runtime/Landscape/Private/LandscapeEdit.cpp
--- a/Engine/Source/Runtime/Landscape/Private/LandscapeEdit.cpp
+++ b/Engine/Source/Runtime/Landscape/Private/LandscapeEdit.cpp
@@ -2935,7 +2935,7 @@ bool ULandscapeInfo::IsLandscapeEditableWorld() const
 {
 	if (GIsEditor)
 	{
-		UWorld* World = LandscapeActor->GetWorld();
+		UWorld* World = LandscapeActor.IsValid() ? LandscapeActor->GetWorld() : nullptr;
 		return World && SupportsLandscapeEditing() && !IsRunningCommandlet() && !World->IsPlayInEditor() && GEditor->PlayWorld == nullptr;
 	}
 	return false;
diff --git a/Engine/Source/Runtime/Landscape/Private/LandscapeEditLayers.cpp b/Engine/Source/Runtime/Landscape/Private/LandscapeEditLayers.cpp
--- a/Engine/Source/Runtime/Landscape/Private/LandscapeEditLayers.cpp
+++ b/Engine/Source/Runtime/Landscape/Private/LandscapeEditLayers.cpp
@@ -1488,45 +1488,11 @@ void ALandscapeProxy::UpdateProxyLayersWeightmapUsage()
 	check(!bNeedsWeightmapUsagesUpdate);
 }
 
-void ALandscapeProxy::PreEditUndo()
-{
-	Super::PreEditUndo();
-
-	// Unregister all ULandscapeLayerInfoObjects events. PostEditUndo will re-register objects that still exist post-undo
-	Algo::ForEach(GetValidTargetLayerObjects(), [this](ULandscapeLayerInfoObject* InLayerInfoObject)
-		{
-			check(InLayerInfoObject->OnLayerInfoChanged().IsBoundToObject(this));
-			InLayerInfoObject->OnLayerInfoChanged().RemoveAll(this);
-		});
-}
-
 void ALandscapeProxy::PostEditUndo()
 {
 	check(ULandscapeComponent::UndoRedoModifiedComponentCount == 0);
 
 	Super::PostEditUndo();
-
-	// Undo operation may have added/removed ULandscapeLayerInfoObjects. Re-register to events here
-	Algo::ForEach(GetValidTargetLayerObjects(), [this](ULandscapeLayerInfoObject* InLayerInfoObject)
-		{
-			check(!InLayerInfoObject->OnLayerInfoChanged().IsBoundToObject(this));
-			InLayerInfoObject->OnLayerInfoChanged().AddUObject(this, &ALandscapeProxy::OnLayerInfoObjectDataChanged);
-		});
-}
-
-void ALandscapeProxy::PostDuplicate(bool bDuplicateForPIE)
-{
-	Super::PostDuplicate(bDuplicateForPIE);
-
-	// Duplicated target layers with layer info object assets need to register so that we can update the landscape and react to asset changes : 
-	if (!bDuplicateForPIE)
-	{
-		Algo::ForEach(GetValidTargetLayerObjects(), [this](ULandscapeLayerInfoObject* InLayerInfoObject)
-			{
-				check(!InLayerInfoObject->OnLayerInfoChanged().IsBoundToObject(this));
-				InLayerInfoObject->OnLayerInfoChanged().AddUObject(this, &ALandscapeProxy::OnLayerInfoObjectDataChanged);
-			});
-	}
 }
 
 void ALandscape::InitializeLandscapeLayersWeightmapUsage()
