Backed out removal of "paint-time weight balancing" (CL 46382265) for getting back the old painting behavior in UE5.7, as a temporary measure while preparing a more robust version for UE5.8

#tests Editor, Cooked game

diff --git a/Engine/Source/Editor/LandscapeEditor/Private/LandscapeEdModeMirrorTool.cpp b/Engine/Source/Editor/LandscapeEditor/Private/LandscapeEdModeMirrorTool.cpp
--- a/Engine/Source/Editor/LandscapeEditor/Private/LandscapeEdModeMirrorTool.cpp
+++ b/Engine/Source/Editor/LandscapeEditor/Private/LandscapeEdModeMirrorTool.cpp
@@ -634,7 +634,9 @@ public:
 				TempMaxY = SourceMaxY;
 				LandscapeEdit.GetWeightData(LayerInfo, TempMinX, TempMinY, TempMaxX, TempMaxY, &SourceWeightData[0], SourceSizeX);
 				ApplyMirrorInternal<uint8>(SourceWeightData, DestWeightData, SourceSizeX, SourceSizeY, DestSizeX, DestSizeY, MirrorPos, BlendWidth);
-				LandscapeEdit.SetAlphaData(LayerInfo, DestMinX, DestMinY, DestMaxX, DestMaxY, &DestWeightData[0], DestSizeX, ELandscapeLayerPaintingRestriction::None);
+PRAGMA_DISABLE_DEPRECATION_WARNINGS
+				LandscapeEdit.SetAlphaData(LayerInfo, DestMinX, DestMinY, DestMaxX, DestMaxY, &DestWeightData[0], DestSizeX, ELandscapeLayerPaintingRestriction::None, /*bWeightAdjust = */false,  /*bTotalWeightAdjust = */false);
+PRAGMA_ENABLE_DEPRECATION_WARNINGS
 				//LayerInfo->IsReferencedFromLoadedData = true;
 			}
 		}
diff --git a/Engine/Source/Runtime/Landscape/Private/LandscapeEdit.cpp b/Engine/Source/Runtime/Landscape/Private/LandscapeEdit.cpp
--- a/Engine/Source/Runtime/Landscape/Private/LandscapeEdit.cpp
+++ b/Engine/Source/Runtime/Landscape/Private/LandscapeEdit.cpp
@@ -3359,6 +3359,8 @@ void ALandscapeProxy::Import(const FGuid& InGuid, int32 InMinX, int32 InMinY, in
 			WeightValues.Empty(EditingAlphaLayerData.Num());
 			WeightValues.AddZeroed(EditingAlphaLayerData.Num());
 			ComponentWeightmapLayerAllocations.Empty(EditingAlphaLayerData.Num());
+			// TODO [jonathan.bard] : remove all this once we get rid of non-edit layers, weight balancing at the data level is not reliable and weight blending with edit layers active 
+			//  is providing the same service in a non-destructive way
 			TArray<bool, TInlineAllocator<16>> IsNoBlendArray;
 			IsNoBlendArray.Empty(EditingAlphaLayerData.Num());
 			IsNoBlendArray.AddZeroed(EditingAlphaLayerData.Num());
@@ -3368,13 +3370,12 @@ void ALandscapeProxy::Import(const FGuid& InGuid, int32 InMinX, int32 InMinY, in
 				// Lookup the original layer name
 				WeightValues[WeightLayerIndex] = EditingAlphaLayerData[WeightLayerIndex].AlphaValues;
 				new(ComponentWeightmapLayerAllocations) FWeightmapLayerAllocationInfo(ImportLayerInfos[EditingAlphaLayerData[WeightLayerIndex].LayerIndex].LayerInfo);
-				IsNoBlendArray[WeightLayerIndex] = (ImportLayerInfos[EditingAlphaLayerData[WeightLayerIndex].LayerIndex].LayerInfo->GetBlendMethod() == ELandscapeTargetLayerBlendMethod::None);
+				IsNoBlendArray[WeightLayerIndex] = (ImportLayerInfos[EditingAlphaLayerData[WeightLayerIndex].LayerIndex].LayerInfo->GetBlendMethod() != ELandscapeTargetLayerBlendMethod::FinalWeightBlending);
 			}
 
 			// Discard the temporary alpha data
 			EditingAlphaLayerData.Empty();
-			
-			// Handle the special import method where each alpha map specifies the reminder from previous layers : apply it only to weight-blended layers
+
 			if (InImportMaterialLayerType == ELandscapeImportAlphamapType::Layered)
 			{
 				// For each layer...
@@ -3412,6 +3413,51 @@ void ALandscapeProxy::Import(const FGuid& InGuid, int32 InMinX, int32 InMinY, in
 					}
 				}
 			}
+
+			// Weight normalization for total should be 255...
+			if (WeightValues.Num())
+			{
+				for (int32 Idx = 0; Idx < WeightValues[0].Num(); Idx++)
+				{
+					int32 TotalWeight = 0;
+					int32 MaxLayerIdx = -1;
+					int32 MaxWeight = INT_MIN;
+
+					for (int32 WeightLayerIndex = 0; WeightLayerIndex < WeightValues.Num(); WeightLayerIndex++)
+					{
+						if (!IsNoBlendArray[WeightLayerIndex])
+						{
+							int32 Weight = WeightValues[WeightLayerIndex][Idx];
+							TotalWeight += Weight;
+							if (MaxWeight < Weight)
+							{
+								MaxWeight = Weight;
+								MaxLayerIdx = WeightLayerIndex;
+							}
+						}
+					}
+
+					if (TotalWeight > 0 && TotalWeight != 255)
+					{
+						// normalization...
+						float Factor = 255.0f / TotalWeight;
+						TotalWeight = 0;
+						for (int32 WeightLayerIndex = 0; WeightLayerIndex < WeightValues.Num(); WeightLayerIndex++)
+						{
+							if (!IsNoBlendArray[WeightLayerIndex])
+							{
+								WeightValues[WeightLayerIndex][Idx] = (uint8)(Factor * WeightValues[WeightLayerIndex][Idx]);
+								TotalWeight += WeightValues[WeightLayerIndex][Idx];
+							}
+						}
+
+						if (255 - TotalWeight && MaxLayerIdx >= 0)
+						{
+							WeightValues[MaxLayerIdx][Idx] += static_cast<uint8>(255 - TotalWeight);
+						}
+					}
+				}
+			}
 		}
 	}
 
@@ -3852,7 +3898,9 @@ void ALandscapeProxy::Import(const FGuid& InGuid, int32 InMinX, int32 InMinY, in
 			{
 				if (MaterialLayerInfo.LayerInfo != nullptr && MaterialLayerInfo.LayerData.Num() > 0)
 				{
-					LandscapeEdit.SetAlphaData(MaterialLayerInfo.LayerInfo, InMinX, InMinY, InMaxX, InMaxY, MaterialLayerInfo.LayerData.GetData(), /*Stride = */0, ELandscapeLayerPaintingRestriction::None);
+PRAGMA_DISABLE_DEPRECATION_WARNINGS
+					LandscapeEdit.SetAlphaData(MaterialLayerInfo.LayerInfo, InMinX, InMinY, InMaxX, InMaxY, MaterialLayerInfo.LayerData.GetData(), /*Stride = */0, ELandscapeLayerPaintingRestriction::None, /*bWeightAdjust = */true, false);
+PRAGMA_ENABLE_DEPRECATION_WARNINGS
 
 					// Add to TargetLayers if needed.
 					if (!HasTargetLayer(MaterialLayerInfo.LayerInfo->GetLayerName()))
diff --git a/Engine/Source/Runtime/Landscape/Private/LandscapeEditInterface.cpp b/Engine/Source/Runtime/Landscape/Private/LandscapeEditInterface.cpp
--- a/Engine/Source/Runtime/Landscape/Private/LandscapeEditInterface.cpp
+++ b/Engine/Source/Runtime/Landscape/Private/LandscapeEditInterface.cpp
@@ -1500,6 +1500,113 @@ void ULandscapeComponent::DeleteLayerInternal(ULandscapeLayerInfoObject* LayerIn
 
 	Component->DeleteLayerAllocation(InEditLayerGuid, DeleteAllocIdx, LandscapeEdit.GetShouldDirtyPackage());
 
+	// See if the deleted layer is a NoWeightBlend layer - if not, we don't have to worry about normalization
+	// If the layer doesn't exist, assume it is a blended layer (so renormalization will run whether needed or not)
+	// TODO [jonathan.bard] : remove all this once we get rid of non-edit layers, weight balancing at the data level is not reliable and weight blending with edit layers active 
+	//  is providing the same service in a non-destructive way
+	if ((LayerInfo == nullptr) || (LayerInfo->GetBlendMethod() == ELandscapeTargetLayerBlendMethod::FinalWeightBlending))
+	{
+		// Lock data for all the weightmaps
+		TArray<FLandscapeTextureDataInfo*> TexDataInfos;
+		for (int32 WeightmapIdx = 0; WeightmapIdx < ComponentWeightmapTextures.Num(); WeightmapIdx++)
+		{
+			TexDataInfos.Add(LandscapeEdit.GetTextureDataInfo(ComponentWeightmapTextures[WeightmapIdx]));
+		}
+
+		TArray<bool> LayerNoWeightBlends;	// Array of NoWeightBlend flags
+		TArray<uint8*> LayerDataPtrs;		// Pointers to all layers' data 
+
+		// Get the data for each layer
+		for (int32 LayerIdx = 0; LayerIdx < ComponentWeightmapLayerAllocations.Num(); LayerIdx++)
+		{
+			const FWeightmapLayerAllocationInfo& Allocation = ComponentWeightmapLayerAllocations[LayerIdx];
+			LayerDataPtrs.Add((uint8*)TexDataInfos[Allocation.WeightmapTextureIndex]->GetMipData(0) + ChannelOffsets[Allocation.WeightmapTextureChannel]);
+
+			// Find the layer info and record if it is a "no weight blend" layer.
+			const bool bPerformLegacyWeightBalancing = (Allocation.LayerInfo == nullptr) || (Allocation.LayerInfo->GetBlendMethod() == ELandscapeTargetLayerBlendMethod::FinalWeightBlending);
+			LayerNoWeightBlends.Add(!bPerformLegacyWeightBalancing);
+		}
+
+		// Find the texture data corresponding to this vertex
+		const int32 SizeU = (SubsectionSizeQuads + 1) * NumSubsections;
+		const int32 SizeV = (SubsectionSizeQuads + 1) * NumSubsections;
+		const int32 WeightmapOffsetX = static_cast<int32>(Component->WeightmapScaleBias.Z * SizeU);
+		const int32 WeightmapOffsetY = static_cast<int32>(Component->WeightmapScaleBias.W * SizeV);
+
+		for (int32 SubIndexY = 0; SubIndexY < NumSubsections; SubIndexY++)
+		{
+			for (int32 SubIndexX = 0; SubIndexX < NumSubsections; SubIndexX++)
+			{
+				for (int32 SubY = 0; SubY <= SubsectionSizeQuads; SubY++)
+				{
+					for (int32 SubX = 0; SubX <= SubsectionSizeQuads; SubX++)
+					{
+						const int32 TexX = WeightmapOffsetX + (SubsectionSizeQuads + 1) * SubIndexX + SubX;
+						const int32 TexY = WeightmapOffsetY + (SubsectionSizeQuads + 1) * SubIndexY + SubY;
+						const int32 TexDataIndex = 4 * (TexX + TexY * SizeU);
+
+						// Calculate the sum of other layer weights
+						int32 OtherLayerWeightSum = 0;
+						for (int32 LayerIdx = 0; LayerIdx < ComponentWeightmapLayerAllocations.Num(); LayerIdx++)
+						{
+							if (LayerIdx != DeleteAllocIdx && LayerNoWeightBlends[LayerIdx] == false)
+							{
+								OtherLayerWeightSum += LayerDataPtrs[LayerIdx][TexDataIndex];
+							}
+						}
+
+						if (OtherLayerWeightSum == 0)
+						{
+							// Set the first other weight-blend layer we can find to 255 to avoid a black hole
+							// This isn't ideal but it's the best option
+							// There's nothing we can easily do if this was the only weight-blend layer on this component
+							for (int32 LayerIdx = 0; LayerIdx < ComponentWeightmapLayerAllocations.Num(); LayerIdx++)
+							{
+								if (LayerIdx != DeleteAllocIdx && LayerNoWeightBlends[LayerIdx] == false)
+								{
+									uint8& Weight = LayerDataPtrs[LayerIdx][TexDataIndex];
+									Weight = 255;
+									break;
+								}
+							}
+						}
+						else
+						{
+							// Adjust other layer weights
+							for (int32 LayerIdx = 0; LayerIdx < ComponentWeightmapLayerAllocations.Num(); LayerIdx++)
+							{
+								if (LayerIdx != DeleteAllocIdx && LayerNoWeightBlends[LayerIdx] == false)
+								{
+									uint8& Weight = LayerDataPtrs[LayerIdx][TexDataIndex];
+									Weight = static_cast<uint8>(FMath::Clamp<int32>(FMath::RoundToInt32(255.0f * (float)Weight / (float)OtherLayerWeightSum), 0, 255));
+								}
+							}
+						}
+					}
+				}
+			}
+		}
+
+		// Update all the textures and mips
+		for (int32 Idx = 0; Idx < ComponentWeightmapTextures.Num(); Idx++)
+		{
+			UTexture2D* WeightmapTexture = ComponentWeightmapTextures[Idx];
+			FLandscapeTextureDataInfo* WeightmapDataInfo = TexDataInfos[Idx];
+
+			const int32 NumMips = WeightmapTexture->Source.GetNumMips();
+			TArray<FColor*> WeightmapTextureMipData;
+			WeightmapTextureMipData.AddUninitialized(NumMips);
+			for (int32 MipIdx = 0; MipIdx < NumMips; MipIdx++)
+			{
+				WeightmapTextureMipData[MipIdx] = (FColor*)WeightmapDataInfo->GetMipData(MipIdx);
+			}
+
+			ULandscapeComponent::UpdateWeightmapMips(Component->NumSubsections, Component->SubsectionSizeQuads, WeightmapTexture, WeightmapTextureMipData, 0, 0, MAX_int32, MAX_int32, WeightmapDataInfo);
+
+			WeightmapDataInfo->AddMipUpdateRegion(0, 0, 0, WeightmapTexture->Source.GetSizeX() - 1, WeightmapTexture->Source.GetSizeY() - 1);
+		}
+	}
+
 	Component->RequestWeightmapUpdate();
 }
 
@@ -1545,6 +1652,11 @@ void ULandscapeComponent::FillLayer(ULandscapeLayerInfoObject* LayerInfo, FLands
 	Component->Modify(LandscapeEdit.GetShouldDirtyPackage());
 	Proxy->Modify(LandscapeEdit.GetShouldDirtyPackage());
 
+	// TODO [jonathan.bard] : remove all this once we get rid of non-edit layers, weight balancing at the data level is not reliable and weight blending with edit layers active 
+	//  is providing the same service in a non-destructive way
+	const bool bPerformLegacyWeightBalancing = (LayerInfo->GetBlendMethod() == ELandscapeTargetLayerBlendMethod::FinalWeightBlending);
+	bool bClearOtherWeightBlendLayers = bPerformLegacyWeightBalancing;
+
 	TArray<FWeightmapLayerAllocationInfo>& ComponentWeightmapLayerAllocations = Component->GetWeightmapLayerAllocations(EditLayerGuid);
 	TArray<TObjectPtr<UTexture2D>>& ComponentWeightmapTextures = Component->GetWeightmapTextures(EditLayerGuid);
 	TArray<TObjectPtr<ULandscapeWeightmapUsage>>& ComponentWeightmapTexturesUsage = Component->GetWeightmapTexturesUsage(EditLayerGuid);
@@ -1553,6 +1665,26 @@ void ULandscapeComponent::FillLayer(ULandscapeLayerInfoObject* LayerInfo, FLands
 	int32 FillLayerIdx = ComponentWeightmapLayerAllocations.IndexOfByPredicate(
 		[LayerInfo](const FWeightmapLayerAllocationInfo& Allocation) { return Allocation.LayerInfo == LayerInfo; });
 
+	// if the layer isn't used on this component yet but is a weight-blend layer, then simply steal the allocation of another weight-blend layer!
+	if (FillLayerIdx == INDEX_NONE && bPerformLegacyWeightBalancing)
+	{
+		FillLayerIdx = ComponentWeightmapLayerAllocations.IndexOfByPredicate(
+			[](const FWeightmapLayerAllocationInfo& Allocation) 
+			{ 
+				return (Allocation.LayerInfo == nullptr) || (Allocation.LayerInfo->GetBlendMethod() == ELandscapeTargetLayerBlendMethod::FinalWeightBlending);
+			});
+
+		if (FillLayerIdx != INDEX_NONE)
+		{
+			ComponentWeightmapLayerAllocations[FillLayerIdx].LayerInfo = LayerInfo;
+		}
+		else
+		{
+			// no other weight-blend layers exist
+			bClearOtherWeightBlendLayers = false;
+		}
+	}
+
 	// if the layer is still not found then we are forced to make a new allocation
 	if (FillLayerIdx == INDEX_NONE)
 	{
@@ -1585,6 +1717,61 @@ void ULandscapeComponent::FillLayer(ULandscapeLayerInfoObject* LayerInfo, FLands
 		}
 	}
 
+	// clear other layers
+	if (bClearOtherWeightBlendLayers)
+	{
+		for (int32 LayerIdx = 0; LayerIdx < ComponentWeightmapLayerAllocations.Num(); ++LayerIdx)
+		{
+			if (LayerIdx == FillLayerIdx)
+			{
+				continue;
+			}
+			FWeightmapLayerAllocationInfo& Allocation = ComponentWeightmapLayerAllocations[LayerIdx];
+			if ((Allocation.LayerInfo == nullptr) || (Allocation.LayerInfo->GetBlendMethod() != ELandscapeTargetLayerBlendMethod::FinalWeightBlending))
+			{
+				continue;
+			}
+
+			ULandscapeWeightmapUsage* Usage = ComponentWeightmapTexturesUsage[Allocation.WeightmapTextureIndex];
+			if (Usage) // can be null if WeightmapUsageMap hasn't been built yet
+			{
+				Usage->ChannelUsage[Allocation.WeightmapTextureChannel] = nullptr;
+			}
+
+			Allocation.Free();
+		}
+
+		ComponentWeightmapLayerAllocations.RemoveAll(
+			[](const FWeightmapLayerAllocationInfo& Allocation) { return !Allocation.IsAllocated(); });
+
+		// remove any textures we're no longer using
+		for (int32 TextureIdx = 0; TextureIdx < ComponentWeightmapTextures.Num(); ++TextureIdx)
+		{
+			if (!ComponentWeightmapLayerAllocations.ContainsByPredicate(
+				[TextureIdx](const FWeightmapLayerAllocationInfo& Allocation) { return Allocation.WeightmapTextureIndex == TextureIdx; }))
+			{
+				ComponentWeightmapTextures[TextureIdx]->Modify(LandscapeEdit.GetShouldDirtyPackage());
+
+				for (int32 LayerIdx = 0; LayerIdx < ComponentWeightmapLayerAllocations.Num(); ++LayerIdx)
+				{
+					FWeightmapLayerAllocationInfo& Allocation = ComponentWeightmapLayerAllocations[LayerIdx];
+					if (Allocation.WeightmapTextureIndex > TextureIdx)
+					{
+						--Allocation.WeightmapTextureIndex;
+					}
+				}
+
+				ComponentWeightmapTextures.RemoveAt(TextureIdx);
+				if (ComponentWeightmapTexturesUsage.IsValidIndex(TextureIdx))
+				{
+					ComponentWeightmapTexturesUsage.RemoveAt(TextureIdx);
+				}
+
+				TextureIdx--;
+			}
+		}
+	}
+
 	// todo - normalize texture usage: it's possible to end up with two textures each using one channel at this point
 	// e.g. if you start with 4 blended layers (in one texture) and a non-blended layer (in a 2nd), and fill one weight-blended layer (deleting the other three)
 	// this can also happen with normal painting I believe
@@ -1748,6 +1935,14 @@ void ULandscapeComponent::ReplaceLayerInternal(ULandscapeLayerInfoObject* FromLa
 		}
 	}
 
+	// See if the deleted layer is a NoWeightBlend layer - if not, we don't have to worry about normalization
+	const bool bFromLayerIsNoWeightBlend = (FromLayerInfo && (FromLayerInfo->GetBlendMethod() != ELandscapeTargetLayerBlendMethod::FinalWeightBlending));
+	const bool bToLayerIsNoWeightBlend = (ToLayerInfo && (ToLayerInfo->GetBlendMethod() != ELandscapeTargetLayerBlendMethod::FinalWeightBlending));
+
+	// TODO [jonathan.bard] : remove all this once we get rid of non-edit layers, weight balancing at the data level is not reliable and weight blending with edit layers active 
+	//  is providing the same service in a non-destructive way
+	const bool bRequireNormalization = (bFromLayerIsNoWeightBlend != bToLayerIsNoWeightBlend);
+
 	if (bMerging)
 	{
 		FWeightmapLayerAllocationInfo& ToLayerAllocation = ComponentWeightmapLayerAllocations[ToLayerIdx];
@@ -1807,6 +2002,11 @@ void ULandscapeComponent::ReplaceLayerInternal(ULandscapeLayerInfoObject* FromLa
 		}
 	}
 
+	if (bRequireNormalization)
+	{
+		// TODO
+	}
+
 	// if merging into an existing layer, remove the layer and potentially the texture
 	if (bMerging)
 	{
@@ -2107,13 +2307,170 @@ inline bool FLandscapeEditDataInterface::IsLayerAllowed(const ULandscapeLayerInf
 	return true;
 }
 
-void FLandscapeEditDataInterface::SetAlphaData(ULandscapeLayerInfoObject* const LayerInfo, const int32 X1, const int32 Y1, const int32 X2, const int32 Y2, const uint8* Data, int32 Stride, ELandscapeLayerPaintingRestriction PaintingRestriction /*= None*/)
+inline TMap<const ULandscapeLayerInfoObject*, uint32> FLandscapeEditDataInterface::CountWeightBlendedLayerInfluence(const int32 ComponentIndexX, const int32 ComponentIndexY, TOptional<TArrayView<const uint8* const>> InOptionalLayerDataPtrs)
+{
+	// the counts should easily fit in a uint32, a 255x255 x2x2 Component with weights of all 255 only totals 26 bits
+	checkSlow(FMath::CeilLogTwo(ComponentSizeQuads + 1) * 2 + 8 /*ceillog2(255)*/ <= 32);
+	TMap<const ULandscapeLayerInfoObject*, uint32> LayerInfluenceMap;
+
+	ULandscapeComponent* Component = LandscapeInfo->XYtoComponentMap.FindChecked(FIntPoint(ComponentIndexX,ComponentIndexY));
+	const TArray<FWeightmapLayerAllocationInfo>& ComponentWeightmapLayerAllocations = Component->GetWeightmapLayerAllocations(GetEditLayer());
+	const TArray<UTexture2D*>& ComponentWeightmapTextures = Component->GetWeightmapTextures(GetEditLayer());
+
+	// used if InOptionalLayerDataPtrs is null
+	TArray<FLandscapeTextureDataInfo*, TInlineAllocator<2>> InternalTexDataInfos;
+	TArray<const uint8*, TInlineAllocator<8>> InternalLayerDataPtrs;
+	TArrayView<const uint8* const> LayerDataPtrs;
+	if (InOptionalLayerDataPtrs)
+	{
+		check(InOptionalLayerDataPtrs->Num() == ComponentWeightmapLayerAllocations.Num());
+		LayerDataPtrs = InOptionalLayerDataPtrs.GetValue();
+	}
+	else
+	{
+		InternalTexDataInfos.AddUninitialized(ComponentWeightmapTextures.Num());
+		for (int32 WeightmapIdx = 0; WeightmapIdx < ComponentWeightmapTextures.Num(); ++WeightmapIdx)
+		{
+			InternalTexDataInfos[WeightmapIdx] = GetTextureDataInfo(ComponentWeightmapTextures[WeightmapIdx]);
+		}
+
+		InternalLayerDataPtrs.AddUninitialized(ComponentWeightmapLayerAllocations.Num());
+		for (int32 LayerIdx = 0; LayerIdx < ComponentWeightmapLayerAllocations.Num(); LayerIdx++)
+		{
+			const FWeightmapLayerAllocationInfo& Allocation = ComponentWeightmapLayerAllocations[LayerIdx];
+			InternalLayerDataPtrs[LayerIdx] = (uint8*)InternalTexDataInfos[Allocation.WeightmapTextureIndex]->GetMipData(0) + ChannelOffsets[Allocation.WeightmapTextureChannel];
+		}
+
+		LayerDataPtrs = InternalLayerDataPtrs;
+	}
+
+	const int32 ScanlineSize = (SubsectionSizeQuads + 1) * ComponentNumSubsections * 4;
+
+	for (int32 LayerIdx = 0; LayerIdx < ComponentWeightmapLayerAllocations.Num(); LayerIdx++)
+	{
+		const FWeightmapLayerAllocationInfo& Allocation = ComponentWeightmapLayerAllocations[LayerIdx];
+		if ((Allocation.LayerInfo == nullptr) || (Allocation.LayerInfo->GetBlendMethod() != ELandscapeTargetLayerBlendMethod::FinalWeightBlending))
+		{
+			continue;
+		}
+		auto& Count = LayerInfluenceMap.Add(Allocation.LayerInfo, 0);
+
+		for (int32 SubIndexY = 0; SubIndexY < ComponentNumSubsections; ++SubIndexY)
+		{
+			const int32 YStart = SubIndexY * (SubsectionSizeQuads + 1);
+			const int32 YEnd = YStart + (SubsectionSizeQuads + 1);
+			for (int32 Y = YStart; Y < YEnd; ++Y)
+			{
+				for (int32 SubIndexX = 0; SubIndexX < ComponentNumSubsections; ++SubIndexX)
+				{
+					const int32 XStart = SubIndexX * (SubsectionSizeQuads + 1);
+					const int32 XEnd = XStart + (SubsectionSizeQuads + 1);
+					for (int32 X = XStart; X < XEnd; ++X)
+					{
+						const int32 TexDataIndex = Y * ScanlineSize + X * 4;
+						const uint8 Weight = LayerDataPtrs[LayerIdx][TexDataIndex];
+						Count += Weight;
+					}
+				}
+			}
+		}
+	}
+
+	LayerInfluenceMap.ValueSort(TGreater<uint32>());
+	return LayerInfluenceMap;
+}
+
+const ULandscapeLayerInfoObject* FLandscapeEditDataInterface::ChooseReplacementLayer(const ULandscapeLayerInfoObject* const LayerInfo, const int32 ComponentIndexX, const int32 SubIndexX, const int32 SubX, const int32 ComponentIndexY, const int32 SubIndexY, const int32 SubY, TMap<FIntPoint, TMap<const ULandscapeLayerInfoObject*, uint32>>& LayerInfluenceCache, TArrayView<const uint8* const> LayerDataPtrs)
+{
+	const TMap<const ULandscapeLayerInfoObject*, uint32>* LayerInfluenceMapCacheEntry = LayerInfluenceCache.Find(FIntPoint(ComponentIndexX, ComponentIndexY));
+	if (!LayerInfluenceMapCacheEntry)
+	{
+		LayerInfluenceMapCacheEntry = &LayerInfluenceCache.Add(FIntPoint(ComponentIndexX, ComponentIndexY), CountWeightBlendedLayerInfluence(ComponentIndexX, ComponentIndexY, LayerDataPtrs));
+	}
+
+	if (!(SubIndexX == 0 && SubX == 0) &&
+		!(SubIndexX == ComponentNumSubsections - 1 && SubX == SubsectionSizeQuads) &&
+		!(SubIndexY == 0 && SubY == 0) &&
+		!(SubIndexY == ComponentNumSubsections - 1 && SubY == SubsectionSizeQuads))
+	{
+		for (const auto& LayerInfluenceMapPair : *LayerInfluenceMapCacheEntry)
+		{
+			if (LayerInfluenceMapPair.Key != LayerInfo)
+			{
+				return LayerInfluenceMapPair.Key;
+			}
+		}
+		return nullptr;
+	}
+
+	TMap<const ULandscapeLayerInfoObject*, uint32, TInlineSetAllocator<8>> LayerInfluenceMap = *LayerInfluenceMapCacheEntry;
+
+	const int32 ComponentXStart = (SubIndexX == 0 && SubX == 0) ? ComponentIndexX - 1 : ComponentIndexX;
+	const int32 ComponentXEnd = (SubIndexX == ComponentNumSubsections - 1 && SubX == SubsectionSizeQuads) ? ComponentIndexX + 1 : ComponentIndexX;
+	const int32 ComponentYStart = (SubIndexY == 0 && SubY == 0) ? ComponentIndexY - 1 : ComponentIndexY;
+	const int32 ComponentYEnd = (SubIndexY == ComponentNumSubsections - 1 && SubY == SubsectionSizeQuads) ? ComponentIndexY + 1 : ComponentIndexY;
+	for (int32 Y = ComponentYStart; Y <= ComponentYEnd; ++Y)
+	{
+		for (int32 X = ComponentXStart; X <= ComponentXEnd; ++X)
+		{
+			if (X == ComponentIndexX && Y == ComponentIndexY)
+			{
+				// skip the current component, it is already included above
+				continue;
+			}
+
+			ULandscapeComponent* Component = LandscapeInfo->XYtoComponentMap.FindRef(FIntPoint(X, Y));
+			if (!Component)
+			{
+				// skip missing components
+				continue;
+			}
+
+			const TMap<const ULandscapeLayerInfoObject*, uint32>* OtherLayerInfluenceMapCacheEntry = LayerInfluenceCache.Find(FIntPoint(X, Y));
+			if (!OtherLayerInfluenceMapCacheEntry)
+			{
+				OtherLayerInfluenceMapCacheEntry = &LayerInfluenceCache.Add(FIntPoint(X, Y), CountWeightBlendedLayerInfluence(X, Y, {}));
+			}
+
+			for (auto LayerInfluenceMapIt = LayerInfluenceMap.CreateIterator(); LayerInfluenceMapIt; ++LayerInfluenceMapIt)
+			{
+				const uint32* Value = OtherLayerInfluenceMapCacheEntry->Find(LayerInfluenceMapIt->Key);
+				if (Value)
+				{
+					LayerInfluenceMapIt->Value += *Value;
+				}
+				else
+				{
+					// only allow layers that exist in *all* the touched components
+					LayerInfluenceMapIt.RemoveCurrent();
+				}
+			}
+		}
+	}
+
+	LayerInfluenceMap.ValueSort(TGreater<uint32>());
+	for (const auto& LayerInfluenceMapPair : LayerInfluenceMap)
+	{
+		if (LayerInfluenceMapPair.Key != LayerInfo)
+		{
+			return LayerInfluenceMapPair.Key;
+		}
+	}
+	return nullptr;
+}
+
+void FLandscapeEditDataInterface::SetAlphaData(ULandscapeLayerInfoObject* LayerInfo, const int32 X1, const int32 Y1, const int32 X2, const int32 Y2, const uint8* Data, int32 Stride, ELandscapeLayerPaintingRestriction PaintingRestriction /*= None*/, bool bWeightAdjust /*= true*/, bool bTotalWeightAdjust /*= false*/)
 {
 	TRACE_CPUPROFILER_EVENT_SCOPE(LandscapeEditDataInterface_SetAlphaData);
 
 	check(LayerInfo != nullptr);
 	check(Data != nullptr);
 
+	if ((LayerInfo == nullptr) || (LayerInfo->GetBlendMethod() != ELandscapeTargetLayerBlendMethod::FinalWeightBlending))
+	{
+		bWeightAdjust = false;
+	}
+
 	if (Stride == 0)
 	{
 		Stride = (1+X2-X1);
@@ -2128,11 +2485,14 @@ void FLandscapeEditDataInterface::SetAlphaData(ULandscapeLayerInfoObject* const
 
 	TArray<FLandscapeTextureDataInfo*, TInlineAllocator<2>> TexDataInfos;
 	TArray<uint8*, TInlineAllocator<8>> LayerDataPtrs;      // Pointers to all layers' data
+	TArray<bool, TInlineAllocator<8>> LayerNoWeightBlends;  // NoWeightBlend flags
 	TArray<bool, TInlineAllocator<8>> LayerEditDataAllZero; // Whether the data we are editing for this layer is all zero
 	TArray<FColor*> CollisionWeightmapMipData;
 	TArray<FColor*> SimpleCollisionWeightmapMipData;
 	TArray<FColor*> WeightmapTextureMipData;
 
+	TMap<FIntPoint, TMap<const ULandscapeLayerInfoObject*, uint32>> LayerInfluenceCache;
+
 	for (int32 ComponentIndexY = ComponentIndexY1; ComponentIndexY <= ComponentIndexY2; ComponentIndexY++)
 	{
 		for (int32 ComponentIndexX = ComponentIndexX1; ComponentIndexX <= ComponentIndexX2; ComponentIndexX++)
@@ -2247,6 +2607,8 @@ void FLandscapeEditDataInterface::SetAlphaData(ULandscapeLayerInfoObject* const
 
 			LayerDataPtrs.Reset();        // Pointers to all layers' data
 			LayerDataPtrs.AddUninitialized(ComponentWeightmapLayerAllocations.Num());
+			LayerNoWeightBlends.Reset();  // NoWeightBlend flags
+			LayerNoWeightBlends.AddUninitialized(ComponentWeightmapLayerAllocations.Num());
 			LayerEditDataAllZero.Reset(); // Whether the data we are editing for this layer is all zero 
 			LayerEditDataAllZero.AddUninitialized(ComponentWeightmapLayerAllocations.Num());
 
@@ -2257,11 +2619,13 @@ void FLandscapeEditDataInterface::SetAlphaData(ULandscapeLayerInfoObject* const
 				if (Allocation.LayerInfo != nullptr) // only take into account valid layer
 				{
 					LayerDataPtrs[LayerIdx] = (uint8*)TexDataInfos[Allocation.WeightmapTextureIndex]->GetMipData(0) + ChannelOffsets[Allocation.WeightmapTextureChannel];
+					LayerNoWeightBlends[LayerIdx] = (Allocation.LayerInfo->GetBlendMethod() != ELandscapeTargetLayerBlendMethod::FinalWeightBlending);
 					LayerEditDataAllZero[LayerIdx] = true;
 				}
 				else
 				{
 					LayerDataPtrs[LayerIdx] = nullptr;
+					LayerNoWeightBlends[LayerIdx] = true;
 					LayerEditDataAllZero[LayerIdx] = false;
 				}
 			}
@@ -2324,12 +2688,164 @@ void FLandscapeEditDataInterface::SetAlphaData(ULandscapeLayerInfoObject* const
 								}
 							}
 
-							// No need to check for nullptr here because UpdateLayerIdx is always valid
-							uint8& Weight = LayerDataPtrs[UpdateLayerIdx][TexDataIndex];
-							Weight = NewWeight;
-							if (Weight != 0)
+							// Adjust all layer weights
+							// (bWeightAdjust implies that this is a weight-blended layer)
+							int32 OtherLayerWeightSum = 0;
+							if (bWeightAdjust)
 							{
-								LayerEditDataAllZero[UpdateLayerIdx] = false;
+								// Normalize all layers including the painted one
+								// TODO [jonathan.bard] : this isn't used. TODO: Remove
+								if (bTotalWeightAdjust)
+								{
+									int32 MaxLayerIdx = -1;
+									int32 MaxWeight = INT_MIN;
+
+									// Adjust other layers' weights accordingly
+									for (int32 LayerIdx = 0; LayerIdx < ComponentWeightmapLayerAllocations.Num(); LayerIdx++)
+									{
+										if (LayerDataPtrs[LayerIdx] == nullptr)
+										{
+											continue;
+										}
+
+										uint8& ExistingWeight = LayerDataPtrs[LayerIdx][TexDataIndex];
+
+										if (LayerIdx == UpdateLayerIdx)
+										{
+											ExistingWeight = NewWeight;
+										}
+										// Exclude bNoWeightBlend layers
+										if (LayerNoWeightBlends[LayerIdx] == false)
+										{
+											OtherLayerWeightSum += ExistingWeight;
+											if (MaxWeight < ExistingWeight)
+											{
+												MaxWeight = ExistingWeight;
+												MaxLayerIdx = LayerIdx;
+											}
+										}
+									}
+
+									if (OtherLayerWeightSum > 0 && OtherLayerWeightSum != 255)
+									{
+										const float Factor = 255.0f / OtherLayerWeightSum;
+										OtherLayerWeightSum = 0;
+
+										// Normalize
+										for (int32 LayerIdx = 0; LayerIdx < ComponentWeightmapLayerAllocations.Num(); LayerIdx++)
+										{
+											if (LayerDataPtrs[LayerIdx] == nullptr)
+											{
+												continue;
+											}
+
+											uint8& ExistingWeight = LayerDataPtrs[LayerIdx][TexDataIndex];
+
+											if (LayerNoWeightBlends[LayerIdx] == false)
+											{
+												// normalization...
+												ExistingWeight = (uint8)(Factor * ExistingWeight);
+												OtherLayerWeightSum += ExistingWeight;
+
+												if (ExistingWeight != 0)
+												{
+													LayerEditDataAllZero[LayerIdx] = false;
+												}
+											}
+										}
+
+										if ((255 - OtherLayerWeightSum) && MaxLayerIdx >= 0)
+										{
+											// No need to check for nullptr here because MaxLayerIdx can only be set to a valid layer
+											LayerDataPtrs[MaxLayerIdx][TexDataIndex] += static_cast<uint8>(255 - OtherLayerWeightSum);
+										}
+									}
+								}
+								else
+								{
+									// Adjust other layers' weights accordingly
+									for (int32 LayerIdx = 0; LayerIdx < ComponentWeightmapLayerAllocations.Num(); LayerIdx++)
+									{
+										// Exclude bNoWeightBlend layers 
+										if (LayerIdx != UpdateLayerIdx && LayerNoWeightBlends[LayerIdx] == false)
+										{
+											// No need to check for nullptr here because invalid layers have LayerNoWeightBlends set to true.
+											OtherLayerWeightSum += LayerDataPtrs.IsValidIndex(LayerIdx) ? LayerDataPtrs[LayerIdx][TexDataIndex] : 0;
+										}
+									}
+
+									if (OtherLayerWeightSum == 0 && NewWeight < 255)
+									{
+										if (NewWeight < CurrentWeight)
+										{
+											// When reducing the layer weight from 255, we need to choose another layer to fill to avoid a black hole
+											const ULandscapeLayerInfoObject* ReplacementLayer = ChooseReplacementLayer(LayerInfo, ComponentIndexX, SubIndexX, SubX, ComponentIndexY, SubIndexY, SubY, LayerInfluenceCache, LayerDataPtrs);
+											if (ReplacementLayer)
+											{
+												const int32 ReplacementLayerIndex = ComponentWeightmapLayerAllocations.IndexOfByPredicate([&](const FWeightmapLayerAllocationInfo& AllocationInfo) { return AllocationInfo.LayerInfo == ReplacementLayer; });
+
+												// No need to check for nullptr here because ChooseReplacementLayer can't return an invalid layer.
+												LayerDataPtrs[ReplacementLayerIndex][TexDataIndex] = 255 - NewWeight;
+												LayerEditDataAllZero[ReplacementLayerIndex] = false;
+											}
+											else
+											{
+												// if we didn't find a suitable replacement we just have to leave it at 255, unfortunately
+												NewWeight = 255;
+											}
+										}
+										else if (NewWeight > CurrentWeight)
+										{
+											// if weight is increasing on a black spot then go straight to 255
+											NewWeight = 255;
+										}
+
+										// No need to check for nullptr here because UpdateLayerIdx is always valid
+										LayerDataPtrs[UpdateLayerIdx][TexDataIndex] = NewWeight;
+									}
+									else
+									{
+										for (int32 LayerIdx = 0; LayerIdx < ComponentWeightmapLayerAllocations.Num(); LayerIdx++)
+										{
+											if (LayerDataPtrs[LayerIdx] == nullptr)
+											{
+												continue;
+											}
+
+											uint8& Weight = LayerDataPtrs[LayerIdx][TexDataIndex];
+
+											if (LayerIdx == UpdateLayerIdx)
+											{
+												Weight = NewWeight;
+											}
+											else
+											{
+												// Exclude bNoWeightBlend layers
+												if (LayerNoWeightBlends[LayerIdx] == false)
+												{
+													Weight = static_cast<uint8>(FMath::Clamp(
+														FMath::RoundToInt((float)(255 - NewWeight) * (float)Weight / (float)OtherLayerWeightSum), 0, 255));
+												}
+											}
+
+											if (Weight != 0)
+											{
+												LayerEditDataAllZero[LayerIdx] = false;
+											}
+										}
+									}
+								}
+							}
+							else
+							{
+								// No need to check for nullptr here because UpdateLayerIdx is always valid
+								// Weight value set without adjusting other layers' weights
+								uint8& Weight = LayerDataPtrs[UpdateLayerIdx][TexDataIndex];
+								Weight = NewWeight;
+								if (Weight != 0)
+								{
+									LayerEditDataAllZero[UpdateLayerIdx] = false;
+								}
 							}
 						}
 					}
@@ -2378,10 +2894,11 @@ void FLandscapeEditDataInterface::SetAlphaData(ULandscapeLayerInfoObject* const
 	}
 }
 
-// Deprecated
-void FLandscapeEditDataInterface::SetAlphaData(ULandscapeLayerInfoObject* const LayerInfo, const int32 X1, const int32 Y1, const int32 X2, const int32 Y2, const uint8* Data, int32 Stride, ELandscapeLayerPaintingRestriction PaintingRestriction /*= None*/, bool bWeightAdjust /*= true*/, bool bTotalWeightAdjust /*= false*/)
+void FLandscapeEditDataInterface::SetAlphaData(ULandscapeLayerInfoObject* LayerInfo, const int32 X1, const int32 Y1, const int32 X2, const int32 Y2, const uint8* Data, int32 Stride, ELandscapeLayerPaintingRestriction PaintingRestriction)
 {
-	SetAlphaData(LayerInfo, X1, Y1, X2, Y2, Data, Stride, PaintingRestriction);
+PRAGMA_DISABLE_DEPRECATION_WARNINGS
+	SetAlphaData(LayerInfo, X1, Y1, X2, Y2, Data, Stride, PaintingRestriction, /*bWeightAdjust = */true, /*bTotalWeightAdjust = */false);
+PRAGMA_ENABLE_DEPRECATION_WARNINGS
 }
 
 void FLandscapeEditDataInterface::SetAlphaData(const TSet<ULandscapeLayerInfoObject*>& DirtyLayerInfos, const int32 X1, const int32 Y1, const int32 X2, const int32 Y2, const uint8* Data, int32 Stride, ELandscapeLayerPaintingRestriction PaintingRestriction /*= None*/)
@@ -2554,6 +3071,8 @@ void FLandscapeEditDataInterface::SetAlphaData(const TSet<ULandscapeLayerInfoObj
 
 							const int32 TexDataIndex = 4 * (TexX + TexY * TexSize);
 
+							int32 OtherLayerWeightSum = 0;
+
 							// Apply weights to all layers simultaneously.
 							for (int32 LayerIdx = 0; LayerIdx < ComponentWeightmapLayerAllocations.Num(); LayerIdx++)
 							{
diff --git a/Engine/Source/Runtime/Landscape/Private/LandscapeSplineRaster.cpp b/Engine/Source/Runtime/Landscape/Private/LandscapeSplineRaster.cpp
--- a/Engine/Source/Runtime/Landscape/Private/LandscapeSplineRaster.cpp
+++ b/Engine/Source/Runtime/Landscape/Private/LandscapeSplineRaster.cpp
@@ -448,7 +448,10 @@ void RasterizeControlPointAlpha(int32& MinX, int32& MinY, int32& MaxX, int32& Ma
 		Rasterizer.DrawTriangle(FalloffRight0, Left1, FalloffLeft1, FalloffRight0Pos, Left1Pos, FalloffLeft1Pos, false);
 	}
 
-	LandscapeEdit.SetAlphaData(LayerInfo, MinX, MinY, MaxX, MaxY, Data.GetData(), /*Stride = */0, ELandscapeLayerPaintingRestriction::None);
+PRAGMA_DISABLE_DEPRECATION_WARNINGS
+	LandscapeEdit.SetAlphaData(LayerInfo, MinX, MinY, MaxX, MaxY, Data.GetData(), /*Stride = */0, ELandscapeLayerPaintingRestriction::None, 
+		/*bWeightAdjust = */(LayerInfo->GetBlendMethod() == ELandscapeTargetLayerBlendMethod::FinalWeightBlending), /*bTotalWeightAdjust = */false);
+PRAGMA_ENABLE_DEPRECATION_WARNINGS
 
 	LandscapeEdit.GetComponentsInRegion(MinX, MinY, MaxX, MaxY, &ModifiedComponents);
 }
@@ -571,8 +574,10 @@ void RasterizeSegmentAlpha(int32& MinX, int32& MinY, int32& MaxX, int32& MaxY, F
 		Rasterizer.DrawTriangle(FalloffRight0, Right1, FalloffRight1, FalloffRight0Pos, Right1Pos, FalloffRight1Pos, false);
 	}
 
-	LandscapeEdit.SetAlphaData(LayerInfo, MinX, MinY, MaxX, MaxY, Data.GetData(), /*Stride = */0, ELandscapeLayerPaintingRestriction::None);
-
+PRAGMA_DISABLE_DEPRECATION_WARNINGS
+	LandscapeEdit.SetAlphaData(LayerInfo, MinX, MinY, MaxX, MaxY, Data.GetData(), /*Stride = */0, ELandscapeLayerPaintingRestriction::None, 
+		/*bWeightAdjust = */(LayerInfo->GetBlendMethod() == ELandscapeTargetLayerBlendMethod::FinalWeightBlending), /*bTotalWeightAdjust = */false);
+PRAGMA_ENABLE_DEPRECATION_WARNINGS
 	LandscapeEdit.GetComponentsInRegion(MinX, MinY, MaxX, MaxY, &ModifiedComponents);
 }
 
diff --git a/Engine/Source/Runtime/Landscape/Public/LandscapeEdit.h b/Engine/Source/Runtime/Landscape/Public/LandscapeEdit.h
--- a/Engine/Source/Runtime/Landscape/Public/LandscapeEdit.h
+++ b/Engine/Source/Runtime/Landscape/Public/LandscapeEdit.h
@@ -336,6 +336,14 @@ private:
 	inline bool IsLayerAllowed(const ULandscapeLayerInfoObject* LayerInfo,
 	                          int32 ComponentIndexX, int32 SubIndexX, int32 SubX,
 	                          int32 ComponentIndexY, int32 SubIndexY, int32 SubY);
+
+	// counts the total influence of each weight-blended layer on this component
+	// TODO [jonathan.bard] : Remove when we get rid of non-edit layers (this is only used by ChooseReplacementLayer, which is only used by SetAlphaData, when bWeightAdjust == true)
+	inline TMap<const ULandscapeLayerInfoObject*, uint32> CountWeightBlendedLayerInfluence(int32 ComponentIndexX, int32 ComponentIndexY, TOptional<TArrayView<const uint8* const>> LayerDataPtrs);
+
+	// chooses a replacement layer to use when erasing from 100% influence on a texel
+	// TODO [jonathan.bard] : Remove when we get rid of non-edit layers (this is only used by SetAlphaData, when bWeightAdjust == true)
+	const ULandscapeLayerInfoObject* ChooseReplacementLayer(const ULandscapeLayerInfoObject* LayerInfo, int32 ComponentIndexX, int32 SubIndexX, int32 SubX, int32 ComponentIndexY, int32 SubIndexY, int32 SubY, TMap<FIntPoint, TMap<const ULandscapeLayerInfoObject*, uint32>>& LayerInfluenceCache, TArrayView<const uint8* const> LayerDataPtrs);
 };
 
 struct FLandscapeDoNotDirtyScope
@@ -495,7 +503,20 @@ struct TAlphamapAccessor
 		: LandscapeInfo(InLandscapeInfo)
 		, LandscapeEdit(InLandscapeInfo)
 		, LayerInfo(InLayerInfo)
-	{}
+	{
+		// should be no Layer change during FAlphamapAccessor lifetime...
+		if (InLandscapeInfo && InLayerInfo)
+		{
+			if (LayerInfo == ALandscapeProxy::VisibilityLayer)
+			{
+				bPerformLegacyWeightBalancing = false;
+			}
+			else
+			{
+				bPerformLegacyWeightBalancing = (LayerInfo->GetBlendMethod() == ELandscapeTargetLayerBlendMethod::FinalWeightBlending);
+			}
+		}
+	}
 
 	void GetData(int32& X1, int32& Y1, int32& X2, int32& Y2, TMap<FIntPoint, uint8>& Data)
 	{
@@ -534,7 +555,9 @@ struct TAlphamapAccessor
 				LandscapeComponent->RequestWeightmapUpdate();
 			}
 			
-			LandscapeEdit.SetAlphaData(LayerInfo, X1, Y1, X2, Y2, Data, /*Stride = */0, PaintingRestriction);
+PRAGMA_DISABLE_DEPRECATION_WARNINGS
+			LandscapeEdit.SetAlphaData(LayerInfo, X1, Y1, X2, Y2, Data, /*Stride = */0, PaintingRestriction, /*bWeightAdjust = */bPerformLegacyWeightBalancing, /*bTotalWeightAdjust = */false);
+PRAGMA_ENABLE_DEPRECATION_WARNINGS
 			//LayerInfo->IsReferencedFromLoadedData = true;
 		}
 	}
@@ -559,7 +582,7 @@ private:
 	FLandscapeEditDataInterface LandscapeEdit;
 	TSet<ULandscapeComponent*> ModifiedComponents;
 	ULandscapeLayerInfoObject* LayerInfo;
-	bool bPerformLegacyWeightBalancing;
+	bool bPerformLegacyWeightBalancing = false;
 };
 
 template<bool bInUseInterp, bool bInUseTotalNormalize>
