* Fixed landscape spline modulation texture rasterization crashes (out of bounds sampling, wrong channel count)
* Added support for all known texture format types

#Synced-CL 48524909

diff --git a/Engine/Source/Runtime/Landscape/Private/LandscapeSplineRaster.cpp b/Engine/Source/Runtime/Landscape/Private/LandscapeSplineRaster.cpp
--- a/Engine/Source/Runtime/Landscape/Private/LandscapeSplineRaster.cpp
+++ b/Engine/Source/Runtime/Landscape/Private/LandscapeSplineRaster.cpp
@@ -22,6 +22,7 @@
 #include "ScopedTransaction.h"
 #include "Raster.h"
 #include "Landscape.h"
+#include "ImageCore.h"
 #endif
 
 #define LOCTEXT_NAMESPACE "Landscape"
@@ -95,7 +96,7 @@ protected:
 				uint8& DestHeightFlagsValue = (*HeightFlagsData)[DataIndex];
 				uint8 Flags = DestHeightFlagsValue;
 				if (bLowerTerrain)
-					Flags |= 1; 
+					Flags |= 1;
 				if (bRaiseTerrain)
 					Flags |= 2;
 				DestHeightFlagsValue = Flags;
@@ -128,43 +129,44 @@ class FModulateAlpha
 public:
 	static TSharedPtr<FModulateAlpha> CreateFromLayerInfo(ULandscapeLayerInfoObject* InLayerInfo, int32 InLandscapeMinX, int32 InLandscapeMinY)
 	{
-		if (CVarLandscapeSplineFalloffModulation.GetValueOnAnyThread() == 0 
-			|| (InLayerInfo == nullptr) 
+		if (CVarLandscapeSplineFalloffModulation.GetValueOnAnyThread() == 0
+			|| (InLayerInfo == nullptr)
 			|| (InLayerInfo->GetSplineFalloffModulationTexture() == nullptr))
 		{
 			return nullptr;
 		}
 
-		if (!InLayerInfo->GetSplineFalloffModulationTexture()->Source.IsValid())
-		{
-			UE_LOG(LogLandscape, Error, TEXT("Invalid source data for spline falloff modulation texture (%s). Alpha modulation will be disabled."), *InLayerInfo->GetSplineFalloffModulationTexture()->GetPathName());
-			return nullptr;
-		}
-
 		return MakeShareable(new FModulateAlpha(InLayerInfo, InLandscapeMinX, InLandscapeMinY));
 	}
 
 	float Modulate(float InValue, int32 InX, int32 InY) const
 	{
+		if (!bIsValid)
+		{
+			return InValue;
+		}
+
+		const float TextureWidth = Image.GetWidth();
+		const float TextureHeight = Image.GetHeight();
+
 		float X = FMath::Frac((float)(InX - LandscapeMinX) / (TextureWidth * Tiling)) * TextureWidth;
 		float Y = FMath::Frac((float)(InY - LandscapeMinY) / (TextureHeight * Tiling)) * TextureHeight;
 
 		int32 X0 = FMath::FloorToInt(X);
-		int32 X1 = FMath::CeilToInt(X);
+		int32 X1 = FMath::Min(FMath::CeilToInt(X), TextureWidth - 1);
 		int32 Y0 = FMath::FloorToInt(Y);
-		int32 Y1 = FMath::CeilToInt(Y);
+		int32 Y1 = FMath::Min(FMath::CeilToInt(Y), TextureHeight - 1);
 
-		const uint8* Data = MipData.GetData() + ChannelOffset;
-		uint8 SampleX0Y0 = *(Data + ((Y0*TextureWidth + X0) * sizeof(FColor)));
-		uint8 SampleX1Y0 = *(Data + ((Y0*TextureWidth + X1) * sizeof(FColor)));
-		uint8 SampleX0Y1 = *(Data + ((Y1*TextureWidth + X0) * sizeof(FColor)));
-		uint8 SampleX1Y1 = *(Data + ((Y1*TextureWidth + X1) * sizeof(FColor)));
+		float SampleX0Y0 = Image.GetOnePixelLinear(X0, Y0).Component(ChannelIndex);
+		float SampleX1Y0 = Image.GetOnePixelLinear(X1, Y0).Component(ChannelIndex); 
+		float SampleX0Y1 = Image.GetOnePixelLinear(X0, Y1).Component(ChannelIndex);
+		float SampleX1Y1 = Image.GetOnePixelLinear(X1, Y1).Component(ChannelIndex);
 
-		uint8 LerpY0 = FMath::Lerp(SampleX0Y0, SampleX1Y0, X - X0);
-		uint8 LerpY1 = FMath::Lerp(SampleX0Y1, SampleX1Y1, X - X0);
-		uint8 FinalLerp = FMath::Lerp(LerpY0, LerpY1, Y - Y0);
+		float LerpY0 = FMath::Lerp(SampleX0Y0, SampleX1Y0, X - X0);
+		float LerpY1 = FMath::Lerp(SampleX0Y1, SampleX1Y1, X - X0);
+		float FinalLerp = FMath::Lerp(LerpY0, LerpY1, Y - Y0);
 
-		InValue = InValue + (((FinalLerp / 255.0f) - Bias) * Scale);
+		InValue = InValue + ((FinalLerp - Bias) * Scale);
 
 		return FMath::Clamp(InValue, 0.0f, 1.0f);
 	}
@@ -174,30 +176,65 @@ private:
 	{
 		check(InLayerInfo && InLayerInfo->GetSplineFalloffModulationTexture());
 
-		verify(InLayerInfo->GetSplineFalloffModulationTexture()->Source.GetMipData(MipData, 0));
-		TextureWidth = InLayerInfo->GetSplineFalloffModulationTexture()->Source.GetSizeX();
-		TextureHeight = InLayerInfo->GetSplineFalloffModulationTexture()->Source.GetSizeY();
+		bIsValid = Initialize(InLayerInfo, InLandscapeMinX, InLandscapeMinY);
+	}
+
+	bool Initialize(ULandscapeLayerInfoObject* InLayerInfo, int32 InLandscapeMinX, int32 InLandscapeMinY)
+	{
+		FTextureSource& TextureSource = InLayerInfo->GetSplineFalloffModulationTexture()->Source;
+		if (!TextureSource.IsValid() || !TextureSource.GetMipImage(Image, /*MipIndex = */0))
+		{
+			UE_LOG(LogLandscape, Error, TEXT("Invalid source data for spline falloff modulation texture (%s) of layer %s. Alpha modulation will be disabled."),
+				*InLayerInfo->GetSplineFalloffModulationTexture()->GetPathName(), *InLayerInfo->GetLayerName().ToString());
+			return false;
+		}
+
+		if ((TextureSource.GetSizeX() * TextureSource.GetSizeY()) <= 0)
+		{
+			UE_LOG(LogLandscape, Error, TEXT("Invalid source data for spline falloff modulation texture (%s) of layer %s: incorrect texture size (%i x %i). Alpha modulation will be disabled."),
+				*InLayerInfo->GetSplineFalloffModulationTexture()->GetPathName(), *InLayerInfo->GetLayerName().ToString(), TextureSource.GetSizeX(), TextureSource.GetSizeY());
+			return false;
+		}
+		
+		const ETextureSourceFormat Format = TextureSource.GetFormat();
+		const FTextureSourceFormatInfo& FormatInfo = GTextureSourceFormats[Format];
+		if ((int32)InLayerInfo->GetSplineFalloffModulationColorMask() >= FormatInfo.NumComponents)
+		{
+			UE_LOG(LogLandscape, Error, TEXT("Invalid color mask (%s) for spline falloff modulation texture (%s) of layer %s: the texture only has %i channel(s) available. Alpha modulation will be disabled."),
+				*UEnum::GetValueAsString(InLayerInfo->GetSplineFalloffModulationColorMask()), *InLayerInfo->GetSplineFalloffModulationTexture()->GetPathName(), *InLayerInfo->GetLayerName().ToString(), FormatInfo.NumComponents);
+			return false;
+		}
+		ChannelIndex = (int32)InLayerInfo->GetSplineFalloffModulationColorMask();
+
+		if (FMath::Abs(InLayerInfo->GetSplineFalloffModulationTiling()) > UE_SMALL_NUMBER)
+		{
+			Tiling = 1.0f / InLayerInfo->GetSplineFalloffModulationTiling();
+		}
+		else
+		{
+			UE_LOG(LogLandscape, Warning, TEXT("Invalid tiling for spline falloff modulation texture of layer %s. Alpha modulation will be disabled."),
+				*InLayerInfo->GetLayerName().ToString(), Tiling);
+			return false;
+		}
 
-		Tiling = 1.0f / InLayerInfo->GetSplineFalloffModulationTiling();
 		Bias = InLayerInfo->GetSplineFalloffModulationBias();
 		Scale = InLayerInfo->GetSplineFalloffModulationScale();
 
 		LandscapeMinX = InLandscapeMinX;
 		LandscapeMinY = InLandscapeMinY;
 
-		ChannelOffset = ChannelOffsets[(int32)InLayerInfo->GetSplineFalloffModulationColorMask()];
+		return true;
 	}
 
 private:
-	float Tiling;
-	float Bias;
-	float Scale;
-	TArray64<uint8> MipData;
-	int32 TextureWidth;
-	int32 TextureHeight;
-	int32 LandscapeMinX;
-	int32 LandscapeMinY;
-	int32 ChannelOffset;
+	bool bIsValid = false;
+	float Tiling = 1.0f;
+	float Bias = 0.0f;
+	float Scale = 1.0f;
+	FImage Image;
+	int32 LandscapeMinX = 0;
+	int32 LandscapeMinY = 0;
+	int32 ChannelIndex = 0;
 };
 
 class FLandscapeSplineBlendmaskRasterPolicy
