Landscape Editor - Removed unneeded landscape subsystem checks from obsolete edit layer functions
* Fixes crash in registration flow while loading landscape levels in the actor pallete where the subsystem is null

#Synced-CL 48916257

diff --git a/Engine/Source/Runtime/Landscape/Private/Landscape.cpp b/Engine/Source/Runtime/Landscape/Private/Landscape.cpp
--- a/Engine/Source/Runtime/Landscape/Private/Landscape.cpp
+++ b/Engine/Source/Runtime/Landscape/Private/Landscape.cpp
@@ -5283,8 +5283,7 @@ void ALandscapeProxy::UpgradeSharedProperties(ALandscape* InParentLandscape)
 void ALandscapeProxy::DisplayObsoleteEditLayerMapCheck(ALandscape* InParentLandscape)
 {
 	ULandscapeInfo* LandscapeInfo = GetLandscapeInfo();
-	ULandscapeSubsystem* LandscapeSubsystem = GetWorld() ? GetWorld()->GetSubsystem<ULandscapeSubsystem>() : nullptr;
-	checkf(LandscapeInfo != nullptr && LandscapeSubsystem != nullptr, TEXT("DisplayObsoleteEditLayerMapCheck can only be called after the proxies are registered to ULandscapeInfo"));
+	checkf(LandscapeInfo != nullptr, TEXT("DisplayObsoleteEditLayerMapCheck can only be called after the proxies are registered to ULandscapeInfo"));
 
 	TMap<FGuid, FName> ObsoleteEditLayers = GetObsoleteEditLayersInComponents();
 
diff --git a/Engine/Source/Runtime/Landscape/Private/LandscapeEditLayers.cpp b/Engine/Source/Runtime/Landscape/Private/LandscapeEditLayers.cpp
--- a/Engine/Source/Runtime/Landscape/Private/LandscapeEditLayers.cpp
+++ b/Engine/Source/Runtime/Landscape/Private/LandscapeEditLayers.cpp
@@ -7867,8 +7867,7 @@ TMap<FGuid, FName> ALandscapeProxy::GetObsoleteEditLayersInComponents()
 {
 	TMap<FGuid, FName> ObsoleteEditLayers;
 
-	ULandscapeSubsystem* LandscapeSubsystem = GetWorld()->GetSubsystem<ULandscapeSubsystem>();
-	if (!LandscapeSubsystem || !GetLandscapeActor())
+	if (!GetLandscapeActor())
 	{
 		UE_LOG(LogLandscape, Error, TEXT("GetObsoleteEditLayersInComponents: Cannot get obsolete edit layer data until proxy is registered with parent ALandscape."));
 		return ObsoleteEditLayers;
