Fix for crash when resolving merged landscape height or weight textures

diff --git a/Engine/Source/Runtime/Landscape/Private/LandscapeTextureHash.cpp b/Engine/Source/Runtime/Landscape/Private/LandscapeTextureHash.cpp
--- a/Engine/Source/Runtime/Landscape/Private/LandscapeTextureHash.cpp
+++ b/Engine/Source/Runtime/Landscape/Private/LandscapeTextureHash.cpp
@@ -231,6 +231,9 @@ bool ULandscapeTextureHash::DoesTextureDataChangeExceedThreshold(
 
 void ULandscapeTextureHash::SetHash64(UTexture2D* LandscapeTexture, uint64 NewHash64, ELandscapeTextureUsage TextureUsage, ELandscapeTextureType TextureType)
 {
+	check(TextureUsage != ELandscapeTextureUsage::Unknown);
+	check(TextureType != ELandscapeTextureType::Unknown);
+
 	ULandscapeTextureHash* TextureHash = LandscapeTexture->GetAssetUserData<ULandscapeTextureHash>();
 	if (TextureHash == nullptr)
 	{
@@ -240,9 +243,9 @@ void ULandscapeTextureHash::SetHash64(UTexture2D* LandscapeTexture, uint64 NewHa
 	}
 	else
 	{
-		// pre-existing -- should have the same type
-		check(TextureHash->TextureType == TextureType);
-		check(TextureHash->TextureUsage == TextureUsage);
+		// pre-existing -- should have the same type (or unknown)
+		check(TextureHash->TextureUsage == TextureUsage || TextureHash->TextureUsage == ELandscapeTextureUsage::Unknown);
+		check(TextureHash->TextureType == TextureType || TextureHash->TextureType == ELandscapeTextureType::Unknown);
 	}
 	FGuid LandscapeTextureSourceId = LandscapeTexture->Source.GetId();
 
