Fix for editor crash when updating landscape grass maps with unregistered landscape components

We now skip any unregistered components in the BuildGrassMapsNowForComponents() function.
(Previously it would attempt to start grass map generation on the unregistered components, and sometimes assert).

diff --git a/Engine/Source/Runtime/Landscape/Private/LandscapeGrassMapsBuilder.cpp b/Engine/Source/Runtime/Landscape/Private/LandscapeGrassMapsBuilder.cpp
--- a/Engine/Source/Runtime/Landscape/Private/LandscapeGrassMapsBuilder.cpp
+++ b/Engine/Source/Runtime/Landscape/Private/LandscapeGrassMapsBuilder.cpp
@@ -1147,8 +1147,9 @@ bool FLandscapeGrassMapsBuilder::BuildGrassMapsNowForComponents(
 		for (ULandscapeComponent* Component : LandscapeComponents)
 		{
 			FComponentState* State = ComponentStates.FindRef(Component);
-			if (State == nullptr)
+			if ((State == nullptr) || (State->Component == nullptr))
 			{
+				// untracked or unregistered components will not be built
 				FailedComponents.Add(Component);
 				continue;
 			}
@@ -1184,8 +1185,8 @@ bool FLandscapeGrassMapsBuilder::BuildGrassMapsNowForComponents(
 				uint32 CurrentGrassMapGenHash = Component->ComputeGrassMapGenerationHash();
 				if (CurrentGrassMapGenHash != Component->GrassData->GenerationHash)
 				{
-					UE_LOG(LogGrass, Warning, TEXT("Unexpected content change while generating grass maps synchronously (Component:%s Ticks:%d WorldType:%d OldHash:%x NewHash:%x StateHash:%x NumElems:%d)"),
-						*Component->GetPathName(),
+					UE_LOG(LogGrass, Warning, TEXT("Unexpected content change while generating grass maps synchronously (Component:%p %p %s Ticks:%d WorldType:%d OldHash:%x NewHash:%x StateHash:%x NumElems:%d)"),
+						Component, State->Component, *Component->GetPathName(),
 						State->TickCount,
 						World->WorldType, Component->GrassData->GenerationHash, CurrentGrassMapGenHash, State->GrassMapGenerationHash,
 						Component->GrassData->NumElements);
@@ -1481,6 +1482,7 @@ bool FLandscapeGrassMapsBuilder::StartGrassMapGeneration(FComponentState& State,
 	check(State.Stage == EComponentStage::Pending);
 
 	ULandscapeComponent* Component = State.Component;
+	check(Component); // must be registered before starting grass map generation
 
 	if (bTryFastpaths && TryFastpathsFromPending(State, /* bRecalculateHashes = */ true))
 	{
