Landscape layer update should always bypass the rate limit on preparing textures.  Possible resolution for assertion failures.

#Synced-CL 48853956

diff --git a/Engine/Source/Runtime/Landscape/Classes/Landscape.h b/Engine/Source/Runtime/Landscape/Classes/Landscape.h
--- a/Engine/Source/Runtime/Landscape/Classes/Landscape.h
+++ b/Engine/Source/Runtime/Landscape/Classes/Landscape.h
@@ -562,7 +562,15 @@ PRAGMA_ENABLE_DEPRECATION_WARNINGS
 	LANDSCAPE_API void SetUseGeneratedLandscapeSplineMeshesActors(bool bInEnabled);
 	LANDSCAPE_API bool GetUseGeneratedLandscapeSplineMeshesActors() const;
 	LANDSCAPE_API bool PrepareTextureResources(bool bInWaitForStreaming);
-	bool PrepareTextureResourcesLimited(bool bInWaitForStreaming);
+
+	enum class EPrepareTextureResourcesLimitedMode : uint8
+	{
+		Default = 0,				// Default operation
+		WaitForStreaming = 1 << 0,	// Block until streaming completes
+		BypassRateLimit = 1 << 1	// Bypass rate limit and always call PrepareTextureResources
+	};
+	FRIEND_ENUM_CLASS_FLAGS(EPrepareTextureResourcesLimitedMode);
+	bool PrepareTextureResourcesLimited(EPrepareTextureResourcesLimitedMode Mode = EPrepareTextureResourcesLimitedMode::Default);
 
 	bool GetVisibilityLayerAllocationIndex() const { return 0; }
 	
diff --git a/Engine/Source/Runtime/Landscape/Private/Landscape.cpp b/Engine/Source/Runtime/Landscape/Private/Landscape.cpp
--- a/Engine/Source/Runtime/Landscape/Private/Landscape.cpp
+++ b/Engine/Source/Runtime/Landscape/Private/Landscape.cpp
@@ -7324,7 +7324,7 @@ bool ALandscapeProxy::UpdatePhysicalMaterialTasks(bool bInShouldMarkDirty)
 	if (ALandscape* Landscape = GetLandscapeActor())
 	{
 		// Physical material depends on the weightmaps, so make sure they are controlled by the texture streaming manager.
-		Landscape->PrepareTextureResourcesLimited(false);
+		Landscape->PrepareTextureResourcesLimited();
 	}
 
 	TSet<ULandscapeComponent*> OutdatedComponents;
diff --git a/Engine/Source/Runtime/Landscape/Private/LandscapeEditLayers.cpp b/Engine/Source/Runtime/Landscape/Private/LandscapeEditLayers.cpp
--- a/Engine/Source/Runtime/Landscape/Private/LandscapeEditLayers.cpp
+++ b/Engine/Source/Runtime/Landscape/Private/LandscapeEditLayers.cpp
@@ -342,6 +342,9 @@ FAutoConsoleCommand CmdLandscapeForceLayersFullUpdate(
 	TEXT("Trigger an edit layers update for all (no argument) or a specific set of landscapes (partial name match based on the strings passed in argument).\n"),
 	FConsoleCommandWithArgsDelegate::CreateStatic(&LandscapeForceLayersFullUpdate));
 
+
+ENUM_CLASS_FLAGS(ALandscape::EPrepareTextureResourcesLimitedMode);
+
 struct FLandscapeDirty
 {
 	FLandscapeDirty()
@@ -1840,15 +1843,17 @@ TArray<FLandscapeLayersCopyReadbackTextureParams> PrepareLandscapeLayersCopyRead
 	return Result;
 }
 
-bool ALandscape::PrepareTextureResourcesLimited(bool bInWaitForStreaming)
+bool ALandscape::PrepareTextureResourcesLimited(EPrepareTextureResourcesLimitedMode Mode)
 {
 	// Rate-limit to once per frame. Note that commandlets don't advance the frame counter.
-	//  Also, force it when bInWaitForStreaming is true, so that we don't prevent new textures from being prepared in case PrepareTextureResourcesLimited is called 
-	//  multiple times in the same frame with some new textures being created in-between (e.g. when creating a large WP landscape via multiple successive regions)
-	if (bInWaitForStreaming || (GFrameNumber != LastPrepareTextureResourcesCalled) || IsRunningCommandlet())
+	// Also, force it when bInWaitForStreaming is true, so that we don't prevent new textures from being prepared in case PrepareTextureResourcesLimited is called 
+	// multiple times in the same frame with some new textures being created in-between (e.g. when creating a large WP landscape via multiple successive regions)
+	if (EnumHasAnyFlags(Mode, EPrepareTextureResourcesLimitedMode::BypassRateLimit | EPrepareTextureResourcesLimitedMode::WaitForStreaming)
+		|| (GFrameNumber != LastPrepareTextureResourcesCalled)
+		|| IsRunningCommandlet())
 	{
 		LastPrepareTextureResourcesCalled = GFrameNumber;
-		bLastPrepareTextureResourcesResult = PrepareTextureResources(bInWaitForStreaming);
+		bLastPrepareTextureResourcesResult = PrepareTextureResources(EnumHasAnyFlags(Mode, EPrepareTextureResourcesLimitedMode::WaitForStreaming));
 	}
 	return bLastPrepareTextureResourcesResult;
 }
@@ -7253,8 +7258,14 @@ void ALandscape::UpdateLayersContent(bool bInWaitForStreaming, bool bInSkipMonit
 		LayerUpdateCount--;
 	};
 
-	// Note : no early-out allowed before this : even if not actually updating edit layers, we need to poll our resources in order to make sure we register to streaming events when needed: 
-	bool bResourcesReady = PrepareTextureResourcesLimited(bInWaitForStreaming);
+	// Note : no early-out allowed before this : even if not actually updating edit layers, we need to poll our resources in order to make sure we register to streaming events when needed:
+	// Always bypass the rate limit on PrepareTextureResources from here.  The rate limit is to prevent hammering many times per frame from physical material updates.
+	EPrepareTextureResourcesLimitedMode PrepareFlags = EPrepareTextureResourcesLimitedMode::BypassRateLimit;
+	if (bInWaitForStreaming)
+	{
+		PrepareFlags |= EPrepareTextureResourcesLimitedMode::WaitForStreaming;
+	}
+	bool bResourcesReady = PrepareTextureResourcesLimited(PrepareFlags);
 
 	ULandscapeInfo* LandscapeInfo = GetLandscapeInfo();
 	if (!CanUpdateLayersContent())
@@ -7660,10 +7671,6 @@ void ALandscape::ForceLayersFullUpdate()
 	StreamingManagers.UpdateResourceStreaming(GetWorld()->GetDeltaSeconds(), /* bProcessEverything */ true);
 	StreamingManagers.BlockTillAllRequestsFinished();
 
-	// Reset the frame number for limiting PrepareTextureResourcesLimited in case ForceLayersFullUpdate is called 
-	//  multiple times in a single frame (for safety, mostly, because the frame number is bypassed when bWaitForStreaming is true)
-	LastPrepareTextureResourcesCalled = 0;
-
 	RequestSplineLayerUpdate();
 	RequestLayersContentUpdateForceAll();
 	ForceUpdateLayersContent();
