Fix for a landscape crash in GC during cook when all landscape streaming proxies are moved during the load/cook process.

Moved landscapes are added to a list to remap their grid registration, and also check for edge patching again.
However, if ALL of the mapped streaming proxies are moved, it allows the grid to be redefined (handling editor cases where the global landscape transform is changed).
During this process, we temporarily unmap (unregister from the grid) all of the mapped landscape components prior to redefining the grid, without removing them from the patching update lists.
In normal editor usage, the remap of these landscape components is handled by the editor tick.  However in cook there is no tick, and so the components are left in an unexpected state of being unmapped but having edge patch requests active.  These edge patch requests are then not cleaned up when the components are unregistered, as that is part of unmap. And this triggers the assert sanity check in the GC add references callback.

The fix here is to always immediately re-map any components that were unmapped prior to the grid redefinition.

diff --git a/Engine/Source/Runtime/Landscape/Private/LandscapeEdgeFixup.cpp b/Engine/Source/Runtime/Landscape/Private/LandscapeEdgeFixup.cpp
--- a/Engine/Source/Runtime/Landscape/Private/LandscapeEdgeFixup.cpp
+++ b/Engine/Source/Runtime/Landscape/Private/LandscapeEdgeFixup.cpp
@@ -1389,8 +1389,8 @@ UE::Landscape::ENeighborFlags ULandscapeHeightmapTextureEdgeFixup::UpdateEdgeSna
 	FGuid HeightmapSourceID = ULandscapeTextureHash::GetHash(HeightmapTexture);
 	if (bForceUpdate || (HeightmapSourceID != EdgeSnapshot->TextureSourceID))
 	{
-		FIXUP_DEBUG_LOG_DETAIL(TEXT("  UpdateEdgeSnapshotFromHeightmapSource %p HeightMap: %p -- UPDATED (ID: %s --> %s)"), ActiveComponent.Get(), HeightmapTexture.Get(),
-			*EdgeSnapshot->TextureSourceID.ToString(), *HeightmapSourceID.ToString());
+		FIXUP_DEBUG_LOG_DETAIL(TEXT("  UpdateEdgeSnapshotFromHeightmapSource %p HeightMap: %p -- UPDATED (ID: %s --> %s) force:%d"), ActiveComponent.Get(), HeightmapTexture.Get(),
+			*EdgeSnapshot->TextureSourceID.ToString(), *HeightmapSourceID.ToString(), bForceUpdate);
 			
 		// create a new edge data
 		TSharedPtr<FHeightmapTextureEdgeSnapshot> NewEdgeSnapshot = FHeightmapTextureEdgeSnapshot::CreateEdgeSnapshotFromHeightmapSource(HeightmapTexture, LandscapeGridScale);
@@ -1428,7 +1428,7 @@ UE::Landscape::ENeighborFlags ULandscapeHeightmapTextureEdgeFixup::UpdateEdgeSna
 	}
 	else
 	{
-		FIXUP_DEBUG_LOG_DETAIL(TEXT("  UpdateEdgeSnapshotFromHeightmapSource %p HeightMap: %p -- SourceID matches (%s)"), ActiveComponent.Get(), HeightmapTexture, *EdgeSnapshot->TextureSourceID.ToString());
+		FIXUP_DEBUG_LOG_DETAIL(TEXT("  UpdateEdgeSnapshotFromHeightmapSource %p HeightMap: %p -- SourceID matches (%s)"), ActiveComponent.Get(), HeightmapTexture.Get(), *EdgeSnapshot->TextureSourceID.ToString());
 		if (bOutSuccess)
 		{
 			*bOutSuccess = true;
diff --git a/Engine/Source/Runtime/Landscape/Private/LandscapeGroup.cpp b/Engine/Source/Runtime/Landscape/Private/LandscapeGroup.cpp
--- a/Engine/Source/Runtime/Landscape/Private/LandscapeGroup.cpp
+++ b/Engine/Source/Runtime/Landscape/Private/LandscapeGroup.cpp
@@ -103,7 +103,8 @@ void FLandscapeGroup::RegisterComponent(ULandscapeComponent* Component)
 		// and we can't register until the landscape actor has been registered and shared settings are fixed
 		// (we will re-register all streaming proxy components when the landscape actor gets registered,
 		// to handle any components that were skipped here)
-		GROUP_DEBUG_LOG_DETAIL(TEXT("Group %p RegisterComponent %p -- WaitForActor"), this, Component);
+		GROUP_DEBUG_LOG_DETAIL(TEXT("Group %p RegisterComponent %p -- WaitForActor (%d RegisteredFixups, %d NeedPatching, %d Mapped, %d Moved)"),
+			this, Component, AllRegisteredFixups.Num(), HeightmapsNeedingEdgeTexturePatching.Num(), XYToEdgeFixupMap.Num(), HeightmapsMoved.Num());
 		return;
 	}
 
@@ -150,7 +151,8 @@ void FLandscapeGroup::RegisterComponent(ULandscapeComponent* Component)
 		else
 		{
 			// ensure texture user datas are installed
-			GROUP_DEBUG_LOG_DETAIL(TEXT("Group %p RegisterComponent %p -- Activate EdgeFixup"), this, Component);
+			GROUP_DEBUG_LOG_DETAIL(TEXT("Group %p RegisterComponent %p -- Activate EdgeFixup (%d RegisteredFixups, %d NeedPatching, %d Mapped, %d Moved)"),
+				this, Component, AllRegisteredFixups.Num(), HeightmapsNeedingEdgeTexturePatching.Num(), XYToEdgeFixupMap.Num(), HeightmapsMoved.Num());
 
 			const bool bShouldCompressHeightmap = false; // compress only happens at cook time, and once compressed it remains compressed
 			const bool bUseEdgeFixup = true;
@@ -190,7 +192,7 @@ void FLandscapeGroup::RegisterComponent(ULandscapeComponent* Component)
 		EdgeFixup->SetActiveComponent(Component, this);
 		HeightmapTextureToActiveComponent.FindOrAdd(Heightmap) = Component;
 
-		GROUP_DEBUG_LOG_DETAIL(TEXT("  -- EdgeFixup Set Active At Coord (%d, %d)"), EdgeFixup->GroupCoord.X, EdgeFixup->GroupCoord.Y);
+		GROUP_DEBUG_LOG_DETAIL(TEXT("  -- EdgeFixup Set %x Active At Coord (%d, %d)"), Component, EdgeFixup->GroupCoord.X, EdgeFixup->GroupCoord.Y);
 
 #if WITH_EDITOR
 		// Request an edge snapshot update, and initialize the GPU edge hashes from that.
@@ -277,7 +279,9 @@ void FLandscapeGroup::UnregisterComponent(ULandscapeComponent* Component)
 
 		if (!bReactivatedDisabled)
 		{
-			GROUP_DEBUG_LOG_DETAIL(TEXT("Group %p UnregisterComponent %p %d (%d,%d) -- Deactivating"), this, Component, EdgeFixup->bMapped, EdgeFixup->GroupCoord.X, EdgeFixup->GroupCoord.Y);
+			GROUP_DEBUG_LOG_DETAIL(TEXT("Group %p UnregisterComponent %p %d (%d,%d) -- Deactivating (%d RegisteredFixups, %d NeedSnapshot, %d Mapped, %d Moved)"),
+				this, Component, EdgeFixup->bMapped, EdgeFixup->GroupCoord.X, EdgeFixup->GroupCoord.Y,
+				AllRegisteredFixups.Num(), HeightmapsNeedingEdgeTexturePatching.Num(), XYToEdgeFixupMap.Num(), HeightmapsMoved.Num());
 
 			// no disabled available to replace the active component -- just set NO active component
 			const bool bDisableCurrentActive = false; // just unmap it
@@ -301,6 +305,7 @@ void FLandscapeGroup::OnTransformUpdated(ULandscapeComponent* Component)
 	{
 		// flag it as moved, so we can fix it on the next tick
 		// (this must be deferred so we can move all components simultaneously, and redefine the grid if necessary)
+		GROUP_DEBUG_LOG_DETAIL(TEXT("Group %x MOVE component %x"), this, Component);
 		HeightmapsMoved.Add(EdgeFixup);
 	}
 }
@@ -387,9 +392,11 @@ FIntPoint FLandscapeGroup::Map(ULandscapeHeightmapTextureEdgeFixup* Fixup, ULand
 	// meaning there's no time when there is nothing in the fixup map.  Instead, we check if all remaining items in the fixup map are moved.
 	// in this case we can assume it's ok to redefine the grid.
 	bool bCanRedefineGrid = false;
-	if (HeightmapsMoved.Num() >= XYToEdgeFixupMap.Num())
+	if ((XYToEdgeFixupMap.Num() == 0) || HeightmapsMoved.Num() >= XYToEdgeFixupMap.Num())
 	{
 		bCanRedefineGrid = true;
+
+		// check if every mapped fixup is moved
 		for (auto MappedFixupPair : XYToEdgeFixupMap)
 		{
 			ULandscapeHeightmapTextureEdgeFixup* MappedFixup = MappedFixupPair.Value;
@@ -400,29 +407,49 @@ FIntPoint FLandscapeGroup::Map(ULandscapeHeightmapTextureEdgeFixup* Fixup, ULand
 				break;
 			}
 		}
-
 	}
 
 	if (bCanRedefineGrid)
 	{
-		if (XYToEdgeFixupMap.Num() > 0)
+		GROUP_DEBUG_LOG_DETAIL(TEXT("Map REDEFINE GRID: %d existing, %d moved"), XYToEdgeFixupMap.Num(), HeightmapsMoved.Num());
+
+		// we will remap everything below, which will handle the moves.  Clear this early so recursive calls to Map don't trigger grid redefinition again
+		HeightmapsMoved.Empty();
+
+		// first unmap ALL mapped fixups prior to changing the rendercoord grid
+		TArray<TObjectPtr<ULandscapeHeightmapTextureEdgeFixup>> AllMappedFixups;
+		XYToEdgeFixupMap.GenerateValueArray(AllMappedFixups);
+		for (ULandscapeHeightmapTextureEdgeFixup* MappedFixup : AllMappedFixups)
 		{
-			// unmap all existing fixups prior to changing the rendercoord grid
-			// (we can only get here if all of these fixups are "moved", and will be remapped anyways)
-			TArray<TObjectPtr<ULandscapeHeightmapTextureEdgeFixup>> Fixups;
-			XYToEdgeFixupMap.GenerateValueArray(Fixups);
-			for (ULandscapeHeightmapTextureEdgeFixup* ExistingFixup : Fixups)
-			{
-				const bool bRemoveFromLists = false;
-				Unmap(ExistingFixup, bRemoveFromLists);
-			}
+			const bool bRemoveFromLists = false; // don't bother to remove from the lists, as we will just add them again below
+			Unmap(MappedFixup, bRemoveFromLists);
 		}
 
+		// redefine the grid
 		ComponentResolution = Component->ComponentSizeQuads;
 		GroupCoordOrigin = ComponentCenterWorldSpace;
 		GroupCoordXVector = ComponentXVectorWorldSpace;
 		GroupCoordYVector = ComponentYVectorWorldSpace;
 		LandscapeGridScale = ComponentLandscapeGridScale;
+
+		// remap all of the fixups that used to be mapped
+		for (ULandscapeHeightmapTextureEdgeFixup* MappedFixup : AllMappedFixups)
+		{
+			check(MappedFixup != Fixup); // we shouldn't ever be remapping ourselves recursively
+			Map(MappedFixup, MappedFixup->ActiveComponent);
+
+			// Mark it, and all neighbors, as needing edge texture patching after the remap.
+			// (may or may not need it, depending on what kind of move was done)
+			HeightmapsNeedingEdgeTexturePatching.Add(MappedFixup);
+			for (UE::Landscape::ENeighborIndex NeighborIndex : TEnumRange<UE::Landscape::ENeighborIndex>())
+			{
+				FIntPoint NeighborCoord = MappedFixup->GroupCoord + GetNeighborRelativePosition(NeighborIndex);
+				if (TObjectPtr<ULandscapeHeightmapTextureEdgeFixup>* NeighborPtrPtr = XYToEdgeFixupMap.Find(NeighborCoord))
+				{
+					HeightmapsNeedingEdgeTexturePatching.Add(NeighborPtrPtr->Get());
+				}
+			}
+		}
 	}
 	else
 	{
