Landscape Editor - Disabling heightmap compression cvar to temporarily fix broken normals along subsection/multi-component boundaries in cooked game

diff --git a/Engine/Source/Runtime/Landscape/Private/Landscape.cpp b/Engine/Source/Runtime/Landscape/Private/Landscape.cpp
--- a/Engine/Source/Runtime/Landscape/Private/Landscape.cpp
+++ b/Engine/Source/Runtime/Landscape/Private/Landscape.cpp
@@ -237,7 +237,7 @@ static FAutoConsoleVariable CVarAllowGrassStripping(
 	TEXT("Enables the conditional stripping of grass data during cook.  Disabling this means the bStripGrassWhenCooked* will be ignored."));
 
 static const TCHAR* CVarLandscapeHeightmapCompressionModeName = TEXT("landscape.HeightmapCompressionMode");
-int32 GLandscapeHeightmapCompressionMode = 1;
+int32 GLandscapeHeightmapCompressionMode = 0;
 static FAutoConsoleVariableRef CVarLandscapeHeightmapCompressionMode(
 	CVarLandscapeHeightmapCompressionModeName,
 	GLandscapeHeightmapCompressionMode,
