Landscape - Add defensive edit layer creation during non-edit layer auto conversion
* Fixed crash where proxies fail to create the parent Alandscape edit layers through FixupSharedData
* Creates matching edit layer before copying and warns users other properties may not have synced correctly

diff --git a/Engine/Source/Runtime/Landscape/Private/Landscape.cpp b/Engine/Source/Runtime/Landscape/Private/Landscape.cpp
--- a/Engine/Source/Runtime/Landscape/Private/Landscape.cpp
+++ b/Engine/Source/Runtime/Landscape/Private/Landscape.cpp
@@ -2587,9 +2587,15 @@ void ULandscapeComponent::CopyComponentDataToEditLayer(const FGuid& InSourceEdit
 	const ALandscapeProxy* LandscapeProxy = GetLandscapeProxy();
 	FLandscapeLayerComponentData* DestLayerComponentData = GetLayerData(InDestEditLayerGuid);
 
-	if (!InDestEditLayerGuid.IsValid() || DestLayerComponentData == nullptr || !DestLayerComponentData->IsInitialized() || !LandscapeProxy)
+	if (!InDestEditLayerGuid.IsValid() || !LandscapeProxy)
 	{
-		UE_LOG(LogLandscape, Error, TEXT("CopyDataToEditLayer: Failed to copy data to component: %s. Destination edit layer is invalid or Proxy is not yet registered with ALandscape"), *GetName());
+		UE_LOG(LogLandscape, Error, TEXT("CopyComponentDataToEditLayer: Failed to copy data to component: %s. Destination edit layer guid is invalid or Proxy is not yet registered with ALandscape"), *GetName());
+		return;
+	}
+
+	if (DestLayerComponentData == nullptr || !DestLayerComponentData->IsInitialized())
+	{
+		UE_LOG(LogLandscape, Error, TEXT("CopyComponentDataToEditLayer: Failed to copy data to component: %s. Destination edit layer data does not exist or is not initialized on component"), *GetName());
 		return;
 	}
 
@@ -2613,7 +2619,7 @@ void ULandscapeComponent::CopyComponentDataToEditLayer(const FGuid& InSourceEdit
 
 	if (DataToCopy == nullptr || !DataToCopy->IsInitialized())
 	{
-		UE_LOG(LogLandscape, Error, TEXT("CopyDataToEditLayer: Failed to copy data to destination layer for %s : Component: %s. Source layer GUID is invalid or data is not initialized properly"),
+		UE_LOG(LogLandscape, Error, TEXT("CopyComponentDataToEditLayer: Failed to copy data to destination layer for %s : Component: %s. Source layer GUID is invalid or data is not initialized properly"),
 			*LandscapeProxy->GetName(), *GetName());
 		return;
 	}
@@ -5344,7 +5350,13 @@ void ALandscapeProxy::DisplayObsoleteEditLayerMapCheck(ALandscape* InParentLands
 
 void ALandscapeProxy::FixupSharedData(ALandscape* Landscape, const bool bMapCheck)
 {
-	if ((Landscape == nullptr) || (Landscape == this))
+	if (Landscape == nullptr)
+	{
+		UE_LOG(LogLandscape, Error, TEXT("FixupSharedData: Invalid parent ALandscape. Proxy failed to synchronize parent properties."));
+		return;
+	}
+
+	if (Landscape == this)
 	{
 		return;
 	}
diff --git a/Engine/Source/Runtime/Landscape/Private/LandscapeEditLayers.cpp b/Engine/Source/Runtime/Landscape/Private/LandscapeEditLayers.cpp
--- a/Engine/Source/Runtime/Landscape/Private/LandscapeEditLayers.cpp
+++ b/Engine/Source/Runtime/Landscape/Private/LandscapeEditLayers.cpp
@@ -1448,8 +1448,8 @@ void ALandscape::CopyDataToEditLayer(ALandscapeProxy* InProxy, const FGuid& InSo
 
 	InProxy->Modify();
 
-	const ULandscapeEditLayerBase* DestEditLayer = InDestEditLayerGuid.IsValid() ? GetEditLayerConst(InDestEditLayerGuid) : GetEditLayerConst(0);
-	if (DestEditLayer == nullptr || !Cast<ULandscapeEditLayerPersistent>(DestEditLayer))
+	const ULandscapeEditLayerBase* FinalDestEditLayer = InDestEditLayerGuid.IsValid() ? GetEditLayerConst(InDestEditLayerGuid) : GetEditLayerConst(0);
+	if (FinalDestEditLayer == nullptr || !Cast<ULandscapeEditLayerPersistent>(FinalDestEditLayer))
 	{
 		UE_LOG(LogLandscape, Error, TEXT("CopyDataToEditLayer: Failed to copy data to destination layer for proxy %s. Destination edit layer is not found in ALandscape or does not support persistent edit layer data"), *InProxy->GetName());
 		return;
@@ -1460,7 +1460,17 @@ void ALandscape::CopyDataToEditLayer(ALandscapeProxy* InProxy, const FGuid& InSo
 
 	for (ULandscapeComponent* Component : InProxy->LandscapeComponents)
 	{
-		Component->CopyComponentDataToEditLayer(InSourceEditLayerGuid, DestEditLayer->GetGuid(), bInUseObsoleteLayerData, ProcessedHeightmaps, ProcessedWeightmaps);
+		// During an auto-conversion to edit layers all registered components should have synchronized the ALandscape edit layers property by now
+		// It appears some legacy maps may load improperly (see MapCheck_Message_LandscapeLazyObjectPtrDeprecation_Warning) resulting in edit layers not synching before this point
+		// To avoid crashing, we defensively (and safely) create the edit layer here before copying data
+		FLandscapeLayerComponentData* DestLayerComponentData = Component->GetLayerData(FinalDestEditLayer->GetGuid());
+		if (DestLayerComponentData == nullptr || !DestLayerComponentData->IsInitialized())
+		{
+			InProxy->AddLayer(FinalDestEditLayer->GetGuid());
+			UE_LOG(LogLandscape, Warning, TEXT("CopyDataToEditLayer: Defensively creating edit layer before data copy. It's expected this layer should have already been created via FixupSharedData. Other shared properties may not have upgraded properly!"), *GetName());
+		}
+
+		Component->CopyComponentDataToEditLayer(InSourceEditLayerGuid, FinalDestEditLayer->GetGuid(), bInUseObsoleteLayerData, ProcessedHeightmaps, ProcessedWeightmaps);
 	}
 
 	RequestLayersContentUpdateForceAll();
