Landscape Editor - Delete invalid TargetLayers in FixupSharedData
* Batched merge expects that all target layers have a valid name.
* 5.6 fixed an issue with the Create Layers From materials button generating target layers with invalid names. Maps serialized with this data via global merge crash on first merge

diff --git a/Engine/Source/Runtime/Landscape/Private/Landscape.cpp b/Engine/Source/Runtime/Landscape/Private/Landscape.cpp
--- a/Engine/Source/Runtime/Landscape/Private/Landscape.cpp
+++ b/Engine/Source/Runtime/Landscape/Private/Landscape.cpp
@@ -5386,6 +5386,16 @@ void ALandscapeProxy::FixupSharedData(ALandscape* Landscape, const bool bMapChec
 		bUpdated |= AddLayer(EditLayer->GetGuid());
 	}
 
+	// Check for invalid target layers (batched merge checks all target layer names). Global merge did not validate this condition, so legacy maps can be serialized with bad data
+	// A bug in Create Layers From Materials allowed invalid layers to be generated unintentionally. Maps saved with invalid target layers crash on first merge
+	if (const FLandscapeTargetLayerSettings* InvalidTargetLayerSettings = TargetLayers.Find(FName()))
+	{
+		UE_LOG(LogLandscape, Warning, TEXT("FixupSharedData: Removed invalid target layer. Batched merge requires all target layers to have a valid FName. Actor: '%s' Package: '%s'"), *GetActorNameOrLabel(), *GetPackage()->GetName());
+		LandscapeInfo->DeleteLayer(InvalidTargetLayerSettings->LayerInfoObj, FName());
+		bUpdated |= true;
+		SynchronizedProperties.Add(GET_MEMBER_NAME_CHECKED(ALandscapeProxy, TargetLayers));
+	}
+
 	if (bUpdated)
 	{
 		// In cases where LandscapeInfo is not fully ready yet, we forward the provided ALandscape. If ALandscape is available in LandscapeInfo, we let the object function naturally.
