Landscape Editor - Fix nullptr crash when deleting invalid target layers during ALandscapeProxy::PostLoad
* The FixupSharedData function can run when the proxy's LandscapeActor is valid but the LandscapeInfo's LandscapeActor is not
* Delete the target layer directly on the proxy until the landscape info is ready

#Synced-CL 0

diff --git a/Engine/Source/Runtime/Landscape/Private/Landscape.cpp b/Engine/Source/Runtime/Landscape/Private/Landscape.cpp
--- a/Engine/Source/Runtime/Landscape/Private/Landscape.cpp
+++ b/Engine/Source/Runtime/Landscape/Private/Landscape.cpp
@@ -5386,12 +5386,22 @@ void ALandscapeProxy::FixupSharedData(ALandscape* Landscape, const bool bMapChec
 		bUpdated |= AddLayer(EditLayer->GetGuid());
 	}
 
+	const bool bIsLandscapeInfoReady = LandscapeInfo->LandscapeActor != nullptr;
+
 	// Check for invalid target layers (batched merge checks all target layer names). Global merge did not validate this condition, so legacy maps can be serialized with bad data
 	// A bug in Create Layers From Materials allowed invalid layers to be generated unintentionally. Maps saved with invalid target layers crash on first merge
 	if (const FLandscapeTargetLayerSettings* InvalidTargetLayerSettings = TargetLayers.Find(FName()))
 	{
 		UE_LOG(LogLandscape, Warning, TEXT("FixupSharedData: Removed invalid target layer. Batched merge requires all target layers to have a valid FName. Actor: '%s' Package: '%s'"), *GetActorNameOrLabel(), *GetPackage()->GetName());
-		LandscapeInfo->DeleteLayer(InvalidTargetLayerSettings->LayerInfoObj, FName());
+		if (bIsLandscapeInfoReady)
+		{
+			LandscapeInfo->DeleteLayer(InvalidTargetLayerSettings->LayerInfoObj, FName());
+		}
+		else
+		{
+			// Let each proxy handle target layer removal when the landscape info is not ready
+			RemoveTargetLayer(FName(), /*bPostEditChange = */false);
+		}
 		bUpdated |= true;
 		SynchronizedProperties.Add(GET_MEMBER_NAME_CHECKED(ALandscapeProxy, TargetLayers));
 	}
@@ -5399,10 +5409,10 @@ void ALandscapeProxy::FixupSharedData(ALandscape* Landscape, const bool bMapChec
 	if (bUpdated)
 	{
 		// In cases where LandscapeInfo is not fully ready yet, we forward the provided ALandscape. If ALandscape is available in LandscapeInfo, we let the object function naturally.
-		const ALandscape* LandscapeActor = (LandscapeInfo->LandscapeActor == nullptr) ? Landscape : nullptr;
+		const ALandscape* LandscapeActorOverride = bIsLandscapeInfoReady ? nullptr : Landscape;
 
 		// Force resave the proxy through the modified landscape system, so that the user can then use the Build > Save Modified Landscapes (or Build > Build Landscape) button and therefore manually trigger the re-save of all modified proxies. * /
-		bool bNeedsManualResave = LandscapeInfo->MarkObjectDirty(/*InObject = */this, /*bInForceResave = */true, LandscapeActor);
+		bool bNeedsManualResave = LandscapeInfo->MarkObjectDirty(/*InObject = */this, /*bInForceResave = */true, LandscapeActorOverride);
 
 		if (bMapCheck && bNeedsManualResave && !SynchronizedProperties.IsEmpty())
 		{
diff --git a/Engine/Source/Runtime/Landscape/Private/LandscapeEdit.cpp b/Engine/Source/Runtime/Landscape/Private/LandscapeEdit.cpp
--- a/Engine/Source/Runtime/Landscape/Private/LandscapeEdit.cpp
+++ b/Engine/Source/Runtime/Landscape/Private/LandscapeEdit.cpp
@@ -5058,22 +5058,20 @@ void ULandscapeInfo::ExportLayer(ULandscapeLayerInfoObject* LayerInfo, const FSt
 
 void ULandscapeInfo::DeleteLayer(ULandscapeLayerInfoObject* LayerInfo, const FName& LayerName)
 {
+	checkf(LandscapeActor != nullptr, TEXT("ULandscapeInfo::DeleteLayer: Failed to delete landscape layer. Main LandscapeActor needs to be loaded and registered to the landscape info!"));
+
 	GWarn->BeginSlowTask(LOCTEXT("BeginDeletingLayerTask", "Deleting Layer"), true);
 
 	LandscapeActor->Modify();
 	// Remove from layer settings array
+	int32 LayerIndex = Layers.IndexOfByPredicate([LayerInfo, LayerName](const FLandscapeInfoLayerSettings& LayerSettings) { return LayerSettings.LayerInfoObj == LayerInfo && LayerSettings.LayerName == LayerName; });
+	if (LayerIndex != INDEX_NONE)
 	{
-		int32 LayerIndex = Layers.IndexOfByPredicate([LayerInfo, LayerName](const FLandscapeInfoLayerSettings& LayerSettings) { return LayerSettings.LayerInfoObj == LayerInfo && LayerSettings.LayerName == LayerName; });
-		if (LayerIndex != INDEX_NONE)
-		{
-			Layers.RemoveAt(LayerIndex);
-			LandscapeActor->RemoveTargetLayer(LayerName);
-		}
+		Layers.RemoveAt(LayerIndex);
 	}
 	
 	LandscapeActor->RemoveTargetLayer(LayerName);
 	
-	
 	// Remove data from all components
 	FLandscapeEditDataInterface LandscapeEdit(this);
 	LandscapeEdit.DeleteLayer(LayerInfo);
