More logging and validation checks for landscape splines.

diff --git a/Engine/Source/Editor/LandscapeEditor/Private/LandscapeEdModeSplineTools.cpp b/Engine/Source/Editor/LandscapeEditor/Private/LandscapeEdModeSplineTools.cpp
--- a/Engine/Source/Editor/LandscapeEditor/Private/LandscapeEdModeSplineTools.cpp
+++ b/Engine/Source/Editor/LandscapeEditor/Private/LandscapeEdModeSplineTools.cpp
@@ -607,6 +607,7 @@ public:
 		SplineSelection->ClearSelection();
 
 		SplinesComponent->MarkRenderStateDirty();
+		SplinesComponent->CheckSplinesValid();
 	}
 
 	bool UpdateSplitSegment(ULandscapeSplineControlPoint* ControlPoint, const FVector& LocalLocation)
@@ -677,6 +678,7 @@ public:
 
 		ULandscapeSplinesComponent* SplinesComponent = ControlPoint->GetOuterULandscapeSplinesComponent();
 		SplinesComponent->MarkRenderStateDirty();
+		SplinesComponent->CheckSplinesValid();
 
 		return true;
 	}
@@ -794,6 +796,7 @@ public:
 			{
 				SplineComponent->MarkRenderStateDirty();
 			}
+			SplineComponent->CheckSplinesValid();
 		}
 
 		GUnrealEd->RedrawLevelEditingViewports();
@@ -1206,19 +1209,22 @@ public:
 									return true;
 								}
 
-								FScopedTransaction Transaction(LOCTEXT("LandscapeSpline_MergeSpline", "Merge Spline"));
-								
-								ULandscapeInfo* SourceLandscapeInfo = SourceSplineActor->GetLandscapeInfo();
-								check(SourceLandscapeInfo);
+								{
+									FScopedTransaction Transaction(LOCTEXT("LandscapeSpline_MergeSpline", "Merge Spline"));
 
-								SourceLandscapeInfo->MoveSpline(ClickedControlPoint, SourceSplineActor);
-								AddSegment(SourceControlPoint, ClickedControlPoint, true, true);
+									ULandscapeInfo* SourceLandscapeInfo = SourceSplineActor->GetLandscapeInfo();
+									check(SourceLandscapeInfo);
 
-								// Moving the spline should leave us with an empty actor that we can delete
-								if(ClickedComponent->GetControlPoints().Num() == 0)
-								{
-									ClickedSplineActor->GetWorld()->EditorDestroyActor(ClickedSplineActor, true);
+									SourceLandscapeInfo->MoveSpline(ClickedControlPoint, SourceSplineActor);
+									AddSegment(SourceControlPoint, ClickedControlPoint, true, true);
+
+									// Moving the spline should leave us with an empty actor that we can delete
+									if (ClickedComponent->GetControlPoints().Num() == 0)
+									{
+										ClickedSplineActor->GetWorld()->EditorDestroyActor(ClickedSplineActor, true);
+									}
 								}
+								SourceComponent->CheckSplinesValid();
 								
 								return true;
 							}
@@ -2143,6 +2149,8 @@ public:
 				}
 			}
 		}
+
+		SplineComponent->CheckSplinesValid();
 	}
 
 protected:
diff --git a/Engine/Source/Runtime/Landscape/Classes/LandscapeSplinesComponent.h b/Engine/Source/Runtime/Landscape/Classes/LandscapeSplinesComponent.h
--- a/Engine/Source/Runtime/Landscape/Classes/LandscapeSplinesComponent.h
+++ b/Engine/Source/Runtime/Landscape/Classes/LandscapeSplinesComponent.h
@@ -231,6 +231,7 @@ public:
 	virtual void PostEditChangeProperty(FPropertyChangedEvent& PropertyChangedEvent) override;
 	virtual void PostEditUndo() override;
 #endif // WITH_EDITOR
+	virtual void PreSave(FObjectPreSaveContext ObjectSaveContext) override;
 	//~ End UObject Interface
 
 	//~ Begin UActorComponent Interface
diff --git a/Engine/Source/Runtime/Landscape/Private/LandscapeSplines.cpp b/Engine/Source/Runtime/Landscape/Private/LandscapeSplines.cpp
--- a/Engine/Source/Runtime/Landscape/Private/LandscapeSplines.cpp
+++ b/Engine/Source/Runtime/Landscape/Private/LandscapeSplines.cpp
@@ -41,6 +41,7 @@
 #include "TextureResource.h"
 #include "UObject/FortniteReleaseBranchCustomObjectVersion.h"
 #include "UObject/UObjectIterator.h"
+#include "UObject/ObjectSaveContext.h"
 #if WITH_EDITOR
 #include "Logging/TokenizedMessage.h"
 #include "Logging/MessageLog.h"
@@ -531,33 +532,33 @@ ILandscapeSplineInterface* ULandscapeSplinesComponent::GetSplineOwner()
 void ULandscapeSplinesComponent::CheckSplinesValid()
 {
 #if DO_CHECK
-	// Remove all null control points/segments
-	// This shouldn't happen, but it has somehow (TTP #334549) so we have to fix it
-	ensure(!ControlPoints.Remove(nullptr));
-	ensure(!Segments.Remove(nullptr));
+	// Remove all null control points/segments.  Likely, some error occurred in a previous session and the package is saved with dud references.
+	ensureMsgf(!ControlPoints.Remove(nullptr), TEXT("Error in ULandscapeSplinesComponent in %s, actor: %s"), *GetPathName(), GetOwner() ? *GetOwner()->GetActorNameOrLabel() : TEXT("null"));
+	ensureMsgf(!Segments.Remove(nullptr), TEXT("Error in ULandscapeSplinesComponent in %s, actor: %s"), *GetPathName(), GetOwner() ? *GetOwner()->GetActorNameOrLabel() : TEXT("null"));
 
-	// Check for cross-spline connections, as this is a potential source of nulls
-	// this may be allowed in future, but is not currently
+	// Check for cross-spline connections, as this is a potential source of nulls.  This can also happen if something went wrong
+	// with a segment/point connectivity edit (while merging splines, undo/redo of deletions, etc).  These errors show up some times,
+	// but the trigger conditions are not known.
 	for (ULandscapeSplineControlPoint* ControlPoint : ControlPoints)
 	{
-		ensure(ControlPoint->GetOuterULandscapeSplinesComponent() == this);
+		ensureMsgf(ControlPoint->GetOuterULandscapeSplinesComponent() == this, TEXT("Error in ULandscapeSplinesComponent in %s, actor: %s"), *GetPathName(), GetOwner() ? *GetOwner()->GetActorNameOrLabel() : TEXT("null"));
 		for (const FLandscapeSplineConnection& Connection : ControlPoint->ConnectedSegments)
 		{
 			if (Connection.Segment)
 			{
-				ensure(Connection.Segment->GetOuterULandscapeSplinesComponent() == this);
+				ensureMsgf(Connection.Segment->GetOuterULandscapeSplinesComponent() == this, TEXT("Error in ULandscapeSplinesComponent in %s, actor: %s"), *GetPathName(), GetOwner() ? *GetOwner()->GetActorNameOrLabel() : TEXT("null"));
 			}
 		}
 	}
 
 	for (ULandscapeSplineSegment* Segment : Segments)
 	{
-		ensure(Segment->GetOuterULandscapeSplinesComponent() == this);
+		ensureMsgf(Segment->GetOuterULandscapeSplinesComponent() == this, TEXT("Error in ULandscapeSplinesComponent in %s, actor: %s"), *GetPathName(), GetOwner() ? *GetOwner()->GetActorNameOrLabel() : TEXT("null"));
 		for (const FLandscapeSplineSegmentConnection& Connection : Segment->Connections)
 		{
 			if (Connection.ControlPoint)
 			{
-				ensure(Connection.ControlPoint->GetOuterULandscapeSplinesComponent() == this);
+				ensureMsgf(Connection.ControlPoint->GetOuterULandscapeSplinesComponent() == this, TEXT("Error in ULandscapeSplinesComponent in %s, actor: %s"), *GetPathName(), GetOwner() ? *GetOwner()->GetActorNameOrLabel() : TEXT("null"));
 			}
 		}
 	}
@@ -972,6 +973,9 @@ void ULandscapeSplinesComponent::CheckForErrors()
 	TSet<UWorld*> OutdatedWorlds;
 	TMap<UWorld*, FForeignWorldSplineData*> ForeignWorldSplineDataMapCache;
 
+	// Check for internally-inconsistent data in ControlPoints and Segments, likely due to a code error.
+	CheckSplinesValid();
+
 	// Check control point meshes
 	for (ULandscapeSplineControlPoint* ControlPoint : ControlPoints)
 	{
@@ -1139,6 +1143,13 @@ void ULandscapeSplinesComponent::PostLoad()
 	CheckSplinesValid();
 }
 
+void ULandscapeSplinesComponent::PreSave(FObjectPreSaveContext ObjectSaveContext)
+{
+	Super::PreSave(ObjectSaveContext);
+
+	CheckSplinesValid();
+}
+
 #if WITH_EDITOR
 static bool bHackIsUndoingSplines = false;
 void ULandscapeSplinesComponent::PostEditChangeProperty(FPropertyChangedEvent& PropertyChangedEvent)
@@ -1160,6 +1171,8 @@ void ULandscapeSplinesComponent::PostEditUndo()
 	Super::PostEditUndo();
 	bHackIsUndoingSplines = false;
 
+	CheckSplinesValid();
+
 	MarkRenderStateDirty();
 }
 
@@ -2578,6 +2591,7 @@ void ULandscapeSplineControlPoint::DeleteSplinePoints()
 	Modify();
 
 	ULandscapeSplinesComponent* OuterSplines = GetOuterChecked();
+	UE_LOG(LogLandscape, Display, TEXT("ULandscapeSplineControlPoint::DeleteSplinePoints removing %s from spline actor %s"), *GetName(), *OuterSplines->GetOwner()->GetActorNameOrLabel());
 
 	Points.Reset();
 	Bounds = FBox(ForceInit);
@@ -3623,6 +3637,7 @@ void ULandscapeSplineSegment::DeleteSplinePoints()
 	Modify();
 
 	ULandscapeSplinesComponent* OuterSplines = GetOuterChecked();
+	UE_LOG(LogLandscape, Display, TEXT("ULandscapeSplineSegment::DeleteSplinePoints removing %s from spline actor %s"), *GetName(), *OuterSplines->GetOwner()->GetActorNameOrLabel());
 
 	SplineInfo.Reset();
 	Points.Reset();
@@ -3843,6 +3858,9 @@ void ULandscapeInfo::MoveControlPoint(ULandscapeSplineControlPoint* InControlPoi
 	}
 	check(InControlPoint->GetOuterULandscapeSplinesComponent() == FromSplineComponent);
 
+	UE_LOG(LogLandscape, Display, TEXT("ULandscapeInfo::MoveControlPoint:  move %s, from %s, to %s."),
+		*InControlPoint->GetName(), *FromSplineComponent->GetOwner()->GetActorNameOrLabel(), *ToSplineComponent->GetOwner()->GetActorNameOrLabel());
+
 	ToSplineComponent->Modify();
 	InControlPoint->Modify();
 
@@ -3915,6 +3933,9 @@ void ULandscapeInfo::MoveSegment(ULandscapeSplineSegment* InSegment, TScriptInte
 	}
 	check(InSegment->GetOuterULandscapeSplinesComponent() == FromSplineComponent);
 
+	UE_LOG(LogLandscape, Display, TEXT("ULandscapeInfo::MoveSegment:  move %s, from %s, to %s."),
+		*InSegment->GetName(), *FromSplineComponent->GetOwner()->GetActorNameOrLabel(), *ToSplineComponent->GetOwner()->GetActorNameOrLabel());
+
 	ToSplineComponent->Modify();
 	InSegment->Modify();
 		
@@ -3987,6 +4008,10 @@ void ULandscapeInfo::MoveSpline(ULandscapeSplineControlPoint* InControlPoint, TS
 	{
 		return;
 	}
+
+	UE_LOG(LogLandscape, Display, TEXT("ULandscapeInfo::MoveSpline - ControlPoint %s. Deleting actor %s. Merging spline data to %s."),
+		*InControlPoint->GetName(), *FromSplineComponent->GetOwner()->GetActorNameOrLabel(), *Cast<AActor>(InNewOwner.GetObject())->GetActorNameOrLabel());
+
 	bUpdateBounds = true;
 	From.GetObject()->Modify();
 	FromSplineComponent->Modify();
