Don't call RHIUpdateTexture2D, use UpdateTexture2D on a command list.

diff --git a/Engine/Plugins/Mutable/Source/CustomizableObject/Private/MuCO/CustomizableObjectInstance.cpp b/Engine/Plugins/Mutable/Source/CustomizableObject/Private/MuCO/CustomizableObjectInstance.cpp
--- a/Engine/Plugins/Mutable/Source/CustomizableObject/Private/MuCO/CustomizableObjectInstance.cpp
+++ b/Engine/Plugins/Mutable/Source/CustomizableObject/Private/MuCO/CustomizableObjectInstance.cpp
@@ -5635,7 +5635,7 @@ void UpdateTextureRegionsMutable(UTexture2D* Texture, int32 MipIndex, uint32 Num
 
 				if (RegionData->MipIndex >= CurrentFirstMip + MipDifference)
 				{
-					RHIUpdateTexture2D(
+					CmdList.UpdateTexture2D(
 						RegionData->Texture2DResource->GetTexture2DRHI(),
 						RegionData->MipIndex - CurrentFirstMip - MipDifference,
 						RegionData->Region,
diff --git a/Engine/Plugins/VirtualProduction/DMX/DMXFixtures/Source/DMXFixtures/Private/DMXFixtureActorMatrix.cpp b/Engine/Plugins/VirtualProduction/DMX/DMXFixtures/Source/DMXFixtures/Private/DMXFixtureActorMatrix.cpp
--- a/Engine/Plugins/VirtualProduction/DMX/DMXFixtures/Source/DMXFixtures/Private/DMXFixtureActorMatrix.cpp
+++ b/Engine/Plugins/VirtualProduction/DMX/DMXFixtures/Source/DMXFixtures/Private/DMXFixtureActorMatrix.cpp
@@ -30,7 +30,7 @@ namespace
 				[=](FRHICommandListImmediate& RHICmdList)
 				{
 					FTexture2DResource* Resource = (FTexture2DResource*)DynamicTexture->GetResource();
-					RHIUpdateTexture2D(
+					RHICmdList.UpdateTexture2D(
 						Resource->GetTexture2DRHI(),
 						MipIndex,
 						Region,
diff --git a/Engine/Source/Editor/WorldBrowser/Private/Tiles/WorldTileThumbnails.cpp b/Engine/Source/Editor/WorldBrowser/Private/Tiles/WorldTileThumbnails.cpp
--- a/Engine/Source/Editor/WorldBrowser/Private/Tiles/WorldTileThumbnails.cpp
+++ b/Engine/Source/Editor/WorldBrowser/Private/Tiles/WorldTileThumbnails.cpp
@@ -221,7 +221,7 @@ void FTileAtlasPage::UpdateSlotImageData(int32 SlotIdx, FSlateTextureDataPtr Ima
 			[Context](FRHICommandList& RHICmdList)
 			{
 				FRHITexture* RHITexture2D = (FRHITexture*)Context.TextureRHI.GetReference();
-				RHIUpdateTexture2D(RHITexture2D, 0, Context.Region, Context.SourcePitch, Context.ImageData->GetRawBytesPtr());
+				RHICmdList.UpdateTexture2D(RHITexture2D, 0, Context.Region, Context.SourcePitch, Context.ImageData->GetRawBytesPtr());
 			});
 	}
 }
diff --git a/Engine/Source/Runtime/Engine/Private/Rendering/Texture3DResource.cpp b/Engine/Source/Runtime/Engine/Private/Rendering/Texture3DResource.cpp
--- a/Engine/Source/Runtime/Engine/Private/Rendering/Texture3DResource.cpp
+++ b/Engine/Source/Runtime/Engine/Private/Rendering/Texture3DResource.cpp
@@ -162,7 +162,7 @@ void FTexture3DResource::CreateTexture(FRHICommandListBase& RHICmdList)
 			{
 				FUpdateTextureRegion3D UpdateRegion(0, 0, 0, 0, 0, 0, Mip.SizeX, Mip.SizeY, Mip.SizeZ);
 
-				// RHIUpdateTexture3D crashes on some platforms at engine initialization time.
+				// RHICmdList.UpdateTexture3D crashes on some platforms at engine initialization time.
 				// The default volume texture end up being loaded at that point, which is a problem.
 				// We check if this is really the rendering thread to find out if the engine is initializing.
 				const uint32 NumBlockX = (uint32)FMath::DivideAndRoundUp<int32>(Mip.SizeX, BlockSizeX);
@@ -173,7 +173,7 @@ void FTexture3DResource::CreateTexture(FRHICommandListBase& RHICmdList)
 				int64 UploadSize = NumBlockX * NumBlockY * BlockBytes * Mip.SizeZ;
 				check( MipBytes >= UploadSize );
 
-				UE_LOG(LogTexture,Verbose,TEXT("FTexture3DResource::CreateTexture:RHIUpdateTexture3D %d : %dx%dx%d : RowStride=%d, SliceStride=%d, MipBytes=%d"),
+				UE_LOG(LogTexture,Verbose,TEXT("FTexture3DResource::CreateTexture:UpdateTexture3D %d : %dx%dx%d : RowStride=%d, SliceStride=%d, MipBytes=%d"),
 					ResourceMipIdx,Mip.SizeX,Mip.SizeY,Mip.SizeZ,
 					NumBlockX * BlockBytes, NumBlockX * NumBlockY * BlockBytes,
 					(int)MipBytes);
diff --git a/Engine/Source/Runtime/Engine/Private/Texture2D.cpp b/Engine/Source/Runtime/Engine/Private/Texture2D.cpp
--- a/Engine/Source/Runtime/Engine/Private/Texture2D.cpp
+++ b/Engine/Source/Runtime/Engine/Private/Texture2D.cpp
@@ -1465,7 +1465,7 @@ void UTexture2D::UpdateTextureRegions(int32 MipIndex, uint32 NumRegions, const F
 						RegionCopy.SrcX = 0;
 						RegionCopy.SrcY = 0;
 
-						RHIUpdateTexture2D(
+						RHICmdList.UpdateTexture2D(
 							RegionData->Texture2DResource->TextureRHI->GetTexture2D(),
 							RegionData->MipIndex - CurrentFirstMip,
 							RegionCopy,
diff --git a/Engine/Source/Runtime/Landscape/Private/LandscapeEdgeFixup.cpp b/Engine/Source/Runtime/Landscape/Private/LandscapeEdgeFixup.cpp
--- a/Engine/Source/Runtime/Landscape/Private/LandscapeEdgeFixup.cpp
+++ b/Engine/Source/Runtime/Landscape/Private/LandscapeEdgeFixup.cpp
@@ -1173,10 +1173,10 @@ void ULandscapeHeightmapTextureEdgeFixup::PatchTextureEdge_Internal(UE::Landscap
 						const uint32 SourcePitch = 4 * SrcWidth;
 						FUpdateTextureRegion2D UpdateRegion(DstStartX, DstStartY, SrcStartX, SrcStartY, SrcWidth, SrcHeight);
 
-						FIXUP_DEBUG_LOG_RENDER(TEXT("    PatchTextureEdge_Internal [RHIUpdateTexture2D] res:%p e:%d mip:%d reg:[%d %d %d %d %d %d]"), HeightTextureResource, EdgeIndex, MipIndex,
+						FIXUP_DEBUG_LOG_RENDER(TEXT("    PatchTextureEdge_Internal [UpdateTexture2D] res:%p e:%d mip:%d reg:[%d %d %d %d %d %d]"), HeightTextureResource, EdgeIndex, MipIndex,
 							UpdateRegion.DestX, UpdateRegion.DestY, UpdateRegion.SrcX, UpdateRegion.SrcY, UpdateRegion.Width, UpdateRegion.Height);
 
-						RHIUpdateTexture2D(RHIHeightmapTexture, ResourceMipIndex, UpdateRegion, SourcePitch, Buffer);
+						RHICmdList.UpdateTexture2D(RHIHeightmapTexture, ResourceMipIndex, UpdateRegion, SourcePitch, Buffer);
 					}
 				}
 
@@ -1241,10 +1241,10 @@ void ULandscapeHeightmapTextureEdgeFixup::PatchTextureCorner_Internal(UE::Landsc
 
 				FUpdateTextureRegion2D UpdateRegion(DstStartX, DstStartY, SrcStartX, SrcStartY, SrcWidth, SrcHeight);
 
-				FIXUP_DEBUG_LOG_RENDER(TEXT("    PatchTextureCorner_Internal [RHIUpdateTexture2D] res:%p e:%d mip:%d reg:[%d %d %d %d %d %d]"), HeightTextureResource, CornerIndex, ResourceMipIndex,
+				FIXUP_DEBUG_LOG_RENDER(TEXT("    PatchTextureCorner_Internal [UpdateTexture2D] res:%p e:%d mip:%d reg:[%d %d %d %d %d %d]"), HeightTextureResource, CornerIndex, ResourceMipIndex,
 					UpdateRegion.DestX, UpdateRegion.DestY, UpdateRegion.SrcX, UpdateRegion.SrcY, UpdateRegion.Width, UpdateRegion.Height);
 
-				RHIUpdateTexture2D(RHIHeightmapTexture, ResourceMipIndex, UpdateRegion, SourcePitch, Texel.Data);
+				RHICmdList.UpdateTexture2D(RHIHeightmapTexture, ResourceMipIndex, UpdateRegion, SourcePitch, Texel.Data);
 			}
 		}
 	);
diff --git a/Engine/Source/Runtime/MediaAssets/Private/Misc/MediaTextureResource.cpp b/Engine/Source/Runtime/MediaAssets/Private/Misc/MediaTextureResource.cpp
--- a/Engine/Source/Runtime/MediaAssets/Private/Misc/MediaTextureResource.cpp
+++ b/Engine/Source/Runtime/MediaAssets/Private/Misc/MediaTextureResource.cpp
@@ -1500,7 +1500,7 @@ void FMediaTextureResource::CopySample(FRHICommandListImmediate& RHICmdList, con
 					uint32 Stride = Sample->GetStride() >> MipLevel;
 
 					FUpdateTextureRegion2D Region(0, 0, 0, 0, Width, Height);
-					RHIUpdateTexture2D(RenderTargetTextureRHI, MipLevel, Region, Stride, Data);
+					RHICmdList.UpdateTexture2D(RenderTargetTextureRHI, MipLevel, Region, Stride, Data);
 
 					uint32 HeightUnit = (Height + FormatInfo.BlockSizeY - 1) / FormatInfo.BlockSizeY;
 					Data += Stride * HeightUnit;
diff --git a/Engine/Source/Runtime/Renderer/Private/SceneOcclusion.cpp b/Engine/Source/Runtime/Renderer/Private/SceneOcclusion.cpp
--- a/Engine/Source/Runtime/Renderer/Private/SceneOcclusion.cpp
+++ b/Engine/Source/Runtime/Renderer/Private/SceneOcclusion.cpp
@@ -1019,8 +1019,8 @@ void FHZBOcclusionTester::Submit(FRDGBuilder& GraphBuilder, const FViewInfo& Vie
 				const int32 BlockY = BlockIndex / SizeInBlocksY;
 
 				FUpdateTextureRegion2D Region(BlockX * BlockSize, BlockY * BlockSize, 0, 0, BlockSize, BlockSize);
-				RHIUpdateTexture2D((FRHITexture*)BoundsCenterTexture->GetRHI(), 0, Region, BlockStride, (uint8*)CenterBuffer);
-				RHIUpdateTexture2D((FRHITexture*)BoundsExtentTexture->GetRHI(), 0, Region, BlockStride, (uint8*)ExtentBuffer);
+				RHICmdList.UpdateTexture2D((FRHITexture*)BoundsCenterTexture->GetRHI(), 0, Region, BlockStride, (uint8*)CenterBuffer);
+				RHICmdList.UpdateTexture2D((FRHITexture*)BoundsExtentTexture->GetRHI(), 0, Region, BlockStride, (uint8*)ExtentBuffer);
 			}
 		});
 	}
diff --git a/Engine/Source/Runtime/WebBrowserTexture/Private/WebBrowserTextureResource.cpp b/Engine/Source/Runtime/WebBrowserTexture/Private/WebBrowserTextureResource.cpp
--- a/Engine/Source/Runtime/WebBrowserTexture/Private/WebBrowserTextureResource.cpp
+++ b/Engine/Source/Runtime/WebBrowserTexture/Private/WebBrowserTextureResource.cpp
@@ -234,7 +234,7 @@ void FWebBrowserTextureResource::CopySample(FRHICommandListBase& RHICmdList, con
 
 		// copy sample data to output render target
 		FUpdateTextureRegion2D Region(0, 0, 0, 0, SampleDim.X, SampleDim.Y);
-		RHIUpdateTexture2D(RenderTargetTextureRHI.GetReference(), 0, Region, Sample->GetStride(), (uint8*)Sample->GetBuffer());
+		RHICmdList.UpdateTexture2D(RenderTargetTextureRHI.GetReference(), 0, Region, Sample->GetStride(), (uint8*)Sample->GetBuffer());
 	}
 	Cleared = false;
 }
