Increase error tolerance on landscape LODGroup scale to temporarily suppress bad warnings.

diff --git a/Engine/Source/Runtime/Landscape/Private/LandscapeRender.cpp b/Engine/Source/Runtime/Landscape/Private/LandscapeRender.cpp
--- a/Engine/Source/Runtime/Landscape/Private/LandscapeRender.cpp
+++ b/Engine/Source/Runtime/Landscape/Private/LandscapeRender.cpp
@@ -729,8 +729,9 @@ void FLandscapeRenderSystem::CreateResources_Internal(FRHICommandListBase& RHICm
 		{
 			// validate each section has a matching resolution, scale and orientation
 			bool bResolutionMatches = (ComponentResolution == SectionComponentResolution);
-			bool bXVectorMatches = (SectionXVector - ComponentXVector).IsNearlyZero();
-			bool bYVectorMatches = (SectionYVector - ComponentYVector).IsNearlyZero();
+			// TODO [alan.roach]:  The scale values are showing up slightly off from expected values. 32767.999268 vs 32768.0  Temporarily suppress the warnings with a larger tolerance.
+			bool bXVectorMatches = (SectionXVector - ComponentXVector).IsNearlyZero(UE_KINDA_SMALL_NUMBER * 100.0);
+			bool bYVectorMatches = (SectionYVector - ComponentYVector).IsNearlyZero(UE_KINDA_SMALL_NUMBER * 100.0);
 			if (!(bResolutionMatches && bXVectorMatches && bYVectorMatches))
 			{
 				UE_LOG(LogLandscape, Warning, TEXT("Landscapes in LOD Group %d do not have matching resolution (%d == %d), scale (%f == %f, %f == %f) and/or rotation; geometry seam artifacts may appear. If using HLOD, it may need to be updated to reflect changed resolutions or transforms."),
