Fix assignment of landscape scale X value to scale Y in section info.

diff --git a/Engine/Source/Runtime/Landscape/Private/LandscapeRender.cpp b/Engine/Source/Runtime/Landscape/Private/LandscapeRender.cpp
--- a/Engine/Source/Runtime/Landscape/Private/LandscapeRender.cpp
+++ b/Engine/Source/Runtime/Landscape/Private/LandscapeRender.cpp
@@ -4681,7 +4681,7 @@ public:
 		, ComponentResolution(SectionComponentResolution)
 		, CenterWorldSpace(LocalToWorld.TransformPosition(InComponentCenterLocalSpace))
 		, XVectorWorldSpace(LocalToWorld.TransformVector(InComponentXVectorLocalSpace))
-		, YVectorWorldSpace(LocalToWorld.TransformVector(InComponentXVectorLocalSpace))
+		, YVectorWorldSpace(LocalToWorld.TransformVector(InComponentYVectorLocalSpace))
 	{
 	}
 
