Landscape Editor - Fixed issue with Edit Layer Contribution mode not changing data when switching between edit layers

diff --git a/Engine/Source/Runtime/Landscape/Private/LandscapeEditLayers.cpp b/Engine/Source/Runtime/Landscape/Private/LandscapeEditLayers.cpp
--- a/Engine/Source/Runtime/Landscape/Private/LandscapeEditLayers.cpp
+++ b/Engine/Source/Runtime/Landscape/Private/LandscapeEditLayers.cpp
@@ -7024,7 +7024,9 @@ void ALandscape::MonitorLandscapeEdModeChanges()
 
 	if (bRequiredEditingClientFullUpdate && (LandscapeEdModeInfo.ViewMode == ELandscapeViewMode::LayerContribution))
 	{
-		RequestLayersContentUpdateForceAll(ELandscapeLayerUpdateMode::Update_Client_Editing);
+		// Update_Client_Editing does not trigger readback changes. As a result, UpdateCollisionAndClients will not be run
+		// Use Update_All_Editing to force an update, making sure the contribution mode visualization state is updated with the read back results
+		RequestLayersContentUpdateForceAll(ELandscapeLayerUpdateMode::Update_All_Editing);
 	}
 }
 
