Nullptr check for the landscape actor in ULandscapeInfo::IsLandscapeEditableWorld

diff --git a/Engine/Source/Runtime/Landscape/Private/LandscapeEdit.cpp b/Engine/Source/Runtime/Landscape/Private/LandscapeEdit.cpp
--- a/Engine/Source/Runtime/Landscape/Private/LandscapeEdit.cpp
+++ b/Engine/Source/Runtime/Landscape/Private/LandscapeEdit.cpp
@@ -2933,7 +2933,7 @@ bool ULandscapeInfo::SupportsLandscapeEditing() const
 
 bool ULandscapeInfo::IsLandscapeEditableWorld() const
 {
-	if (GIsEditor)
+	if (GIsEditor && LandscapeActor.IsValid())
 	{
 		UWorld* World = LandscapeActor->GetWorld();
 		return World && SupportsLandscapeEditing() && !IsRunningCommandlet() && !World->IsPlayInEditor() && GEditor->PlayWorld == nullptr;
