Landscape Editor - Added the edit layer event registration calls into ALandscape::PostRegisterAllComponents
- Events were only handled within the create/delete edit layer functions which do not run during Undo/Duplication
- Fixed bug where edit layers did not trigger landscape updates after using Undo (when restoring deleted layers)
- Fixed bug where edit layers did not trigger landscape updates after non-WP landscape duplication

diff --git a/Engine/Source/Runtime/Landscape/Classes/Landscape.h b/Engine/Source/Runtime/Landscape/Classes/Landscape.h
--- a/Engine/Source/Runtime/Landscape/Classes/Landscape.h
+++ b/Engine/Source/Runtime/Landscape/Classes/Landscape.h
@@ -315,6 +315,7 @@ public:
 	virtual void PostEditChangeProperty(FPropertyChangedEvent& PropertyChangedEvent) override;
 	virtual void PostEditMove(bool bFinished) override;
 	virtual void PostEditUndo() override;
+	virtual void UnregisterAllComponents(bool bForReregister = false) override;
 	virtual void PostRegisterAllComponents() override;
 	virtual void PostActorCreated() override;
 	virtual bool ShouldImport(FStringView ActorPropString, bool IsMovingLevel) override;
diff --git a/Engine/Source/Runtime/Landscape/Private/Landscape.cpp b/Engine/Source/Runtime/Landscape/Private/Landscape.cpp
--- a/Engine/Source/Runtime/Landscape/Private/Landscape.cpp
+++ b/Engine/Source/Runtime/Landscape/Private/Landscape.cpp
@@ -3234,9 +3234,6 @@ PRAGMA_DISABLE_DEPRECATION_WARNINGS
 PRAGMA_ENABLE_DEPRECATION_WARNINGS
 			}
 
-			// Register to data change events on the edit layer so that we can update the landscape accordingly : 
-			Layer.EditLayer->OnLayerDataChanged().AddUObject(this, &ALandscape::OnEditLayerDataChanged);
-
 			for (FLandscapeLayerBrush& Brush : Layer.Brushes)
 			{
 				Brush.SetOwner(this);
diff --git a/Engine/Source/Runtime/Landscape/Private/LandscapeEdit.cpp b/Engine/Source/Runtime/Landscape/Private/LandscapeEdit.cpp
--- a/Engine/Source/Runtime/Landscape/Private/LandscapeEdit.cpp
+++ b/Engine/Source/Runtime/Landscape/Private/LandscapeEdit.cpp
@@ -5231,6 +5231,17 @@ void ALandscape::PostEditUndo()
 	RequestLayersContentUpdate(ELandscapeLayerUpdateMode::Update_All);
 }
 
+void ALandscape::UnregisterAllComponents(const bool bForReregister)
+{
+	// Unregister from edit layer data changed events
+	for (FLandscapeLayer& Layer : LandscapeEditLayers)
+	{
+		Layer.EditLayer->OnLayerDataChanged().RemoveAll(this);
+	}
+
+	Super::UnregisterAllComponents(bForReregister);
+}
+
 void ALandscape::PostRegisterAllComponents()
 {
 	Super::PostRegisterAllComponents();
@@ -5267,6 +5278,12 @@ void ALandscape::PostRegisterAllComponents()
 			});
 		}
 	}
+
+	// Register to data change events on the edit layers so that we can update the landscape accordingly : 
+	for (FLandscapeLayer& Layer : LandscapeEditLayers)
+	{
+		Layer.EditLayer->OnLayerDataChanged().AddUObject(this, &ALandscape::OnEditLayerDataChanged);
+	}
 }
 
 void ALandscape::PostActorCreated()
