Landscape Editor - Added a deferred layer content update flag to the LandscapeProxy. Used to track requests from landscape components when the ALandscape is not reachable
- Fixed an crash when moving a river spline point after using Undo on the deletion of an edit layer (with the water brush). In ULandscapeComponent::PostEditUndo the components LayerUpdateFlagPerMode is updated but the request is not called on the main ALandscape (may be unregistered during Undo). As a result, the component is left with unprocessed flags in the next UpdateLayersContent IsUpToDate checks

diff --git a/Engine/Source/Runtime/Landscape/Classes/LandscapeStreamingProxy.h b/Engine/Source/Runtime/Landscape/Classes/LandscapeStreamingProxy.h
--- a/Engine/Source/Runtime/Landscape/Classes/LandscapeStreamingProxy.h
+++ b/Engine/Source/Runtime/Landscape/Classes/LandscapeStreamingProxy.h
@@ -13,6 +13,7 @@ class ALandscape;
 #if WITH_EDITOR
 class UMaterialInterface;
 struct FActorPartitionIdentifier;
+enum ELandscapeLayerUpdateMode : uint32;
 #endif
 
 UCLASS(MinimalAPI, notplaceable)
@@ -35,6 +36,12 @@ private:
 	UPROPERTY()
 	TSet<FName> OverriddenSharedProperties;
 
+#if WITH_EDITORONLY_DATA
+	/** Tracks layer update requests made by this proxy's components when the ALandscape is unreachable. Requests are processed on the next registration once the parent ALandscape is valid */
+	UPROPERTY(Transient)
+	uint32 DeferredLayerContentUpdateModes = 0;
+#endif 
+
 public:
 	//~ Begin UObject Interface
 #if WITH_EDITOR
@@ -68,6 +75,8 @@ public:
 	virtual void SetSharedPropertyOverride(const FName& InPropertyName, const bool bIsOverridden) override;
 	//~ End ALandscapeProxy Interface
 
+	void RequestDeferredLayerContentUpdateModes(ELandscapeLayerUpdateMode InLayerContentUpdateModes);
+	void PerformDeferredLayerContentUpdateModesRequest();
 protected:
 	virtual void FixupOverriddenSharedProperties() override;
 #endif // WITH_EDITOR
diff --git a/Engine/Source/Runtime/Landscape/Private/Landscape.cpp b/Engine/Source/Runtime/Landscape/Private/Landscape.cpp
--- a/Engine/Source/Runtime/Landscape/Private/Landscape.cpp
+++ b/Engine/Source/Runtime/Landscape/Private/Landscape.cpp
@@ -1859,10 +1859,11 @@ TArray<float> ALandscapeProxy::GetLODScreenSizeArray() const
 
 ALandscape* ULandscapeComponent::GetLandscapeActor() const
 {
-	ALandscapeProxy* Landscape = GetLandscapeProxy();
-	if (Landscape)
+	if (ALandscapeProxy* LandscapeProxy = GetLandscapeProxy())
 	{
-		return Landscape->GetLandscapeActor();
+		// When the proxy is a LandscapeStreamingProxy, the landscape actor may be invalid whenever ALandscape is not registered to the current LandscapeInfo
+		// This may happen during operations like Undo that trigger UnregisterActor events
+		return LandscapeProxy->GetLandscapeActor();
 	}
 	return nullptr;
 }
@@ -5196,6 +5197,23 @@ void ALandscapeStreamingProxy::FixupOverriddenSharedProperties()
 	}
 }
 
+void ALandscapeStreamingProxy::RequestDeferredLayerContentUpdateModes(ELandscapeLayerUpdateMode InLayerContentUpdateModes)
+{
+	DeferredLayerContentUpdateModes |= InLayerContentUpdateModes;
+}
+
+void ALandscapeStreamingProxy::PerformDeferredLayerContentUpdateModesRequest()
+{
+	if (DeferredLayerContentUpdateModes != 0)
+	{
+		check(LandscapeActorRef != nullptr);
+		LandscapeActorRef->RequestLayersContentUpdate((ELandscapeLayerUpdateMode)DeferredLayerContentUpdateModes);
+
+		// Reset deferred update modes
+		DeferredLayerContentUpdateModes = 0;
+	}
+}
+
 void ALandscapeProxy::UpgradeSharedProperties(ALandscape* InParentLandscape)
 {
 	TArray<FName> SynchronizedProperties;
@@ -6140,6 +6158,10 @@ void ULandscapeInfo::RegisterLandscapeActorWithProxyInternal(ALandscapeProxy* Pr
 		StreamingProxy->SetLandscapeActor(Landscape);
 
 #if WITH_EDITOR
+		// Proxy may have a deferred update request. This occurs if a component requests an update when the LandscapeActorRef is invalid.
+		// This can happen on Undo, when the component requests an update while ALandscape is unregistered 
+		StreamingProxy->PerformDeferredLayerContentUpdateModesRequest();
+
 		StreamingProxy->FixupSharedData(Landscape, bMapCheck);
 #endif // WITH_EDITOR
 
diff --git a/Engine/Source/Runtime/Landscape/Private/LandscapeEditLayers.cpp b/Engine/Source/Runtime/Landscape/Private/LandscapeEditLayers.cpp
--- a/Engine/Source/Runtime/Landscape/Private/LandscapeEditLayers.cpp
+++ b/Engine/Source/Runtime/Landscape/Private/LandscapeEditLayers.cpp
@@ -6911,45 +6911,56 @@ void ULandscapeComponent::RequestEditingClientUpdate(bool bInUserTriggered)
 void ULandscapeComponent::RequestHeightmapUpdate(bool bUpdateAll, bool bUpdateCollision, bool bInUserTriggered)
 {
 	bUserTriggeredChangeRequested = bInUserTriggered;
+	uint32 NewLayerUpdateFlags = 0;
 	if (bUpdateAll || bUpdateCollision)
 	{
-		LayerUpdateFlagPerMode |= ELandscapeLayerUpdateMode::Update_Heightmap_Editing;
-		LayerUpdateFlagPerMode |= ELandscapeLayerUpdateMode::Update_Heightmap_All;
+		NewLayerUpdateFlags |= ELandscapeLayerUpdateMode::Update_Heightmap_Editing;
+		NewLayerUpdateFlags |= ELandscapeLayerUpdateMode::Update_Heightmap_All;
 	}
 	else
 	{
-		LayerUpdateFlagPerMode |= ELandscapeLayerUpdateMode::Update_Heightmap_Editing_NoCollision;
+		NewLayerUpdateFlags |= ELandscapeLayerUpdateMode::Update_Heightmap_Editing_NoCollision;
 	}
+
+	LayerUpdateFlagPerMode |= NewLayerUpdateFlags;
+
 	if (ALandscape* LandscapeActor = GetLandscapeActor())
 	{
-		LandscapeActor->RequestLayersContentUpdate(bUpdateCollision ? ELandscapeLayerUpdateMode::Update_Heightmap_Editing : ELandscapeLayerUpdateMode::Update_Heightmap_Editing_NoCollision);
-		if (bUpdateAll)
-		{
-			LandscapeActor->RequestLayersContentUpdate(ELandscapeLayerUpdateMode::Update_Heightmap_All);
-		}
+		LandscapeActor->RequestLayersContentUpdate((ELandscapeLayerUpdateMode)NewLayerUpdateFlags);
+	}
+	else
+	{
+		// This condition is only valid with streaming proxies, since the landscape actor will always be reachable in the non-WP Case
+		ALandscapeStreamingProxy* LandscapeStreamingProxy = CastChecked<ALandscapeStreamingProxy>(GetLandscapeProxy());
+		LandscapeStreamingProxy->RequestDeferredLayerContentUpdateModes((ELandscapeLayerUpdateMode)NewLayerUpdateFlags);
 	}
 }
 
 void ULandscapeComponent::RequestWeightmapUpdate(bool bUpdateAll, bool bUpdateCollision, bool bInUserTriggered)
 {
 	bUserTriggeredChangeRequested = bInUserTriggered;
-	
+	uint32 NewLayerUpdateFlags = 0;
 	if (bUpdateAll || bUpdateCollision)
 	{
-		LayerUpdateFlagPerMode |= ELandscapeLayerUpdateMode::Update_Weightmap_Editing;
-		LayerUpdateFlagPerMode |= ELandscapeLayerUpdateMode::Update_Weightmap_All;
+		NewLayerUpdateFlags |= ELandscapeLayerUpdateMode::Update_Weightmap_Editing;
+		NewLayerUpdateFlags |= ELandscapeLayerUpdateMode::Update_Weightmap_All;
 	}
 	else
 	{
-		LayerUpdateFlagPerMode |= ELandscapeLayerUpdateMode::Update_Weightmap_Editing_NoCollision;
+		NewLayerUpdateFlags |= ELandscapeLayerUpdateMode::Update_Weightmap_Editing_NoCollision;
 	}
+
+	LayerUpdateFlagPerMode |= NewLayerUpdateFlags;
+
 	if (ALandscape* LandscapeActor = GetLandscapeActor())
 	{
-		LandscapeActor->RequestLayersContentUpdate(bUpdateCollision ? ELandscapeLayerUpdateMode::Update_Weightmap_Editing : ELandscapeLayerUpdateMode::Update_Weightmap_Editing_NoCollision);
-		if (bUpdateAll)
-		{
-			LandscapeActor->RequestLayersContentUpdate(ELandscapeLayerUpdateMode::Update_Weightmap_All);
-		}
+		LandscapeActor->RequestLayersContentUpdate((ELandscapeLayerUpdateMode)NewLayerUpdateFlags);
+	}
+	else
+	{
+		// This condition is only valid with streaming proxies, since the landscape actor will always be reachable in the non-WP Case
+		ALandscapeStreamingProxy* LandscapeStreamingProxy = CastChecked<ALandscapeStreamingProxy>(GetLandscapeProxy());
+		LandscapeStreamingProxy->RequestDeferredLayerContentUpdateModes((ELandscapeLayerUpdateMode) NewLayerUpdateFlags);
 	}
 }
 
