Landscape Editor - In a WP landscape, setting the shared landscape material on the main landscape is not applied to all proxies on reload
* SynchronizeSharedProperties correctly updates the property values on load but does not trigger any other landscape updates
* Adds optional parameter to call PostEditChangeProperty (handles dirtying of render state, nanite, etc.) after the sync of a property

diff --git a/Engine/Source/Runtime/Landscape/Classes/LandscapeProxy.h b/Engine/Source/Runtime/Landscape/Classes/LandscapeProxy.h
--- a/Engine/Source/Runtime/Landscape/Classes/LandscapeProxy.h
+++ b/Engine/Source/Runtime/Landscape/Classes/LandscapeProxy.h
@@ -1256,9 +1256,10 @@ PRAGMA_ENABLE_DEPRECATION_WARNINGS
 	/**
 	* Enforce property sharing and requirements related to the parent Landscape actor
 	* @param InLandscape the parent landscape to use as a template.
+	* @param bInRunPostEditChangeProperty When true, runs PostEditChangeProperty for each property synchronized
 	* @return An array containing all synchronized properties names (properties that were different from the provided landscape).
 	*/
-	LANDSCAPE_API TArray<FName> SynchronizeSharedProperties(ALandscapeProxy* InLandscape);
+	LANDSCAPE_API TArray<FName> SynchronizeSharedProperties(ALandscapeProxy* InLandscape, bool bInRunPostEditChangeProperty = false);
 
 	// Returns true if the property is shared.
 	LANDSCAPE_API bool IsSharedProperty(const FName& InPropertyName) const;
diff --git a/Engine/Source/Runtime/Landscape/Private/Landscape.cpp b/Engine/Source/Runtime/Landscape/Private/Landscape.cpp
--- a/Engine/Source/Runtime/Landscape/Private/Landscape.cpp
+++ b/Engine/Source/Runtime/Landscape/Private/Landscape.cpp
@@ -3986,28 +3986,27 @@ bool ALandscape::GetUseGeneratedLandscapeSplineMeshesActors() const
 
 void ALandscape::EnableNaniteSkirts(bool bInEnable, float InSkirtDepth, bool bInShouldDirtyPackage)
 {
-	bNaniteSkirtEnabled = bInEnable;
-	NaniteSkirtDepth = InSkirtDepth;
+	FProperty* bNaniteSkirtEnabledProperty = StaticClass()->FindPropertyByName(GET_MEMBER_NAME_CHECKED(ALandscapeProxy, bNaniteSkirtEnabled));
+	FProperty* NaniteSkirtDepthProperty = StaticClass()->FindPropertyByName(GET_MEMBER_NAME_CHECKED(ALandscapeProxy, NaniteSkirtDepth));
+	check(bNaniteSkirtEnabledProperty && NaniteSkirtDepthProperty);
 
-	InvalidateOrUpdateNaniteRepresentation(/*bInCheckContentId*/true, /*InTargetPlatform*/nullptr);
-	UpdateRenderingMethod();
-	MarkComponentsRenderStateDirty();
-	Modify(bInShouldDirtyPackage);
-	
-	if (ULandscapeInfo* LandscapeInfo = GetLandscapeInfo())
+	// PostEditChangeProperty propagates to proxies and invalidates nanite/render state
+	if (bNaniteSkirtEnabled != bInEnable)
 	{
-		LandscapeInfo->ForEachLandscapeProxy([&](ALandscapeProxy* Proxy)
-		{
-			if (Proxy != nullptr)
-			{
-				Proxy->Modify(bInShouldDirtyPackage);
-				Proxy->SynchronizeSharedProperties(this);
-				Proxy->InvalidateOrUpdateNaniteRepresentation(/*bInCheckContentId*/true, /*InTargetPlatform*/nullptr);
-				Proxy->UpdateRenderingMethod();
-				Proxy->MarkComponentsRenderStateDirty();
-			}
-			return true;
-		});
+		Modify(bInShouldDirtyPackage);
+		bNaniteSkirtEnabled = bInEnable;
+
+		FPropertyChangedEvent PropertyChangedEvent(bNaniteSkirtEnabledProperty);
+		PostEditChangeProperty(PropertyChangedEvent);
+	}
+
+	if (NaniteSkirtDepth != InSkirtDepth)
+	{
+		Modify(bInShouldDirtyPackage);
+		NaniteSkirtDepth = InSkirtDepth;
+
+		FPropertyChangedEvent PropertyChangedEvent(NaniteSkirtDepthProperty);
+		PostEditChangeProperty(PropertyChangedEvent);
 	}
 }
 
@@ -5021,7 +5020,7 @@ void ALandscapeProxy::CopySharedProperties(ALandscapeProxy* InLandscape)
 	}
 }
 
-TArray<FName> ALandscapeProxy::SynchronizeSharedProperties(ALandscapeProxy* InLandscape)
+TArray<FName> ALandscapeProxy::SynchronizeSharedProperties(ALandscapeProxy* InLandscape, bool bInRunPostEditChangeProperty)
 {
 	TArray<FName> SynchronizedProperties = SynchronizeUnmarkedSharedProperties(InLandscape);
 
@@ -5039,7 +5038,17 @@ TArray<FName> ALandscapeProxy::SynchronizeSharedProperties(ALandscapeProxy* InLa
 			!Property->Identical_InContainer(this, InLandscape))
 		{
 			SynchronizedProperties.Emplace(Property->GetFName());
-			UE::Landscape::Private::CopyProperty(Property, InLandscape, this);
+
+			// Caller may need a PostEditChangeProperty to trigger other landscape updates (Eg. Dirty Render state, nanite, materials)
+			// PostEditChangeProperty updates are unnecessary when syncing properties during new landscape creation
+			if (bInRunPostEditChangeProperty)
+			{
+				UE::Landscape::Private::CopyPostEditPropertyByName(this, GetLandscapeActor(), Property->GetFName());
+			}
+			else
+			{
+				UE::Landscape::Private::CopyProperty(Property, InLandscape, this);
+			}
 		}
 	}
 
@@ -5381,7 +5390,7 @@ void ALandscapeProxy::FixupSharedData(ALandscape* Landscape, const bool bMapChec
 	}
 	else
 	{
-		SynchronizedProperties = SynchronizeSharedProperties(Landscape);
+		SynchronizedProperties = SynchronizeSharedProperties(Landscape, /*bInRunPostEditChangeProperty =*/true);
 		bUpdated |= !SynchronizedProperties.IsEmpty();
 	}
 
diff --git a/Engine/Source/Runtime/Landscape/Private/LandscapeEdit.cpp b/Engine/Source/Runtime/Landscape/Private/LandscapeEdit.cpp
--- a/Engine/Source/Runtime/Landscape/Private/LandscapeEdit.cpp
+++ b/Engine/Source/Runtime/Landscape/Private/LandscapeEdit.cpp
@@ -4120,26 +4120,17 @@ TSharedRef<UE::Landscape::Nanite::FAsyncBuildData> ALandscapeProxy::MakeAsyncNan
 
 void ALandscape::SetNanitePositionPrecision(int32 InPrecision, bool bInShouldDirtyPackage)
 {
-	NanitePositionPrecision = InPrecision;
+	FProperty* NanitePositionPrecisionProperty = StaticClass()->FindPropertyByName(GET_MEMBER_NAME_CHECKED(ALandscapeProxy, NanitePositionPrecision));
+	check(NanitePositionPrecisionProperty);
 
-	// TODO [chris.tchou] : We should make a consolidated 'value changed' path, unifying this with PostEditChangeProperty
-	InvalidateOrUpdateNaniteRepresentation(/*bInCheckContentId*/true, /*InTargetPlatform*/nullptr);
-	MarkComponentsRenderStateDirty();
-	Modify(bInShouldDirtyPackage);
-
-	if (ULandscapeInfo* LandscapeInfo = GetLandscapeInfo())
+	// PostEditChangeProperty propagates to proxies and invalidates nanite/render state
+	if (NanitePositionPrecision != InPrecision)
 	{
-		LandscapeInfo->ForEachLandscapeProxy([&](ALandscapeProxy* Proxy)
-			{
-				if (Proxy != nullptr)
-				{
-					Proxy->Modify(bInShouldDirtyPackage);
-					Proxy->SynchronizeSharedProperties(this);
-					Proxy->InvalidateOrUpdateNaniteRepresentation(/*bInCheckContentId*/true, /*InTargetPlatform*/nullptr);
-					Proxy->MarkComponentsRenderStateDirty();
-				}
-				return true;
-			});
+		Modify(bInShouldDirtyPackage);
+		NanitePositionPrecision = InPrecision;
+
+		FPropertyChangedEvent PropertyChangedEvent(NanitePositionPrecisionProperty);
+		PostEditChangeProperty(PropertyChangedEvent);
 	}
 }
 
@@ -6744,8 +6735,7 @@ void ALandscape::PostEditChangeProperty(FPropertyChangedEvent& PropertyChangedEv
 				// The ALandscape is included in ForEachLandscapeProxy, avoid getting stuck in a PostEditChange loop
 				if (Proxy != this)
 				{
-					Proxy->SynchronizeSharedProperties(this);
-					Proxy->PostEditChangeProperty(PropertyChangedEvent);
+					Proxy->SynchronizeSharedProperties(this, /*bInRunPostEditChangeProperty =*/true);
 				}
 				return true;
 			});
