World Partition - HLOD: Allow spline components to override their HLOD layer during the partitioning performed by the landscape spline mesh builder

diff --git a/Engine/Source/Editor/UnrealEd/Private/WorldPartition/WorldPartitionLandscapeSplineMeshesBuilder.cpp b/Engine/Source/Editor/UnrealEd/Private/WorldPartition/WorldPartitionLandscapeSplineMeshesBuilder.cpp
--- a/Engine/Source/Editor/UnrealEd/Private/WorldPartition/WorldPartitionLandscapeSplineMeshesBuilder.cpp
+++ b/Engine/Source/Editor/UnrealEd/Private/WorldPartition/WorldPartitionLandscapeSplineMeshesBuilder.cpp
@@ -277,7 +277,21 @@ ALandscapeSplineMeshesActor* UWorldPartitionLandscapeSplineMeshesBuilder::GetOrC
 	FGuid GridGuid = InLandscapeGuid;
 
 	AActor* MeshOwner = InMeshComponent->GetOwner();
-	UHLODLayer* HLODLayer = MeshOwner->GetHLODLayer();
+	UHLODLayer* HLODLayer = nullptr;
+
+	if (const USplineMeshComponent* SplineMeshComponent = Cast<USplineMeshComponent>(InMeshComponent))
+	{
+		HLODLayer = SplineMeshComponent->GetHLODLayerForSplinePartitioning();
+	}
+	else if (const UControlPointMeshComponent* ControlPointMeshComponent = Cast<UControlPointMeshComponent>(InMeshComponent))
+	{
+		HLODLayer = ControlPointMeshComponent->GetHLODLayerForSplinePartitioning();
+	}
+	else
+	{
+		checkNoEntry();
+	}
+
 	if (InMeshComponent->IsHLODRelevant() && MeshOwner->IsHLODRelevant() && HLODLayer != nullptr)
 	{
 		// To get a new unique GUID, combine the original GUID with the HLOD Layer package guid
diff --git a/Engine/Source/Runtime/Engine/Classes/Components/SplineMeshComponent.h b/Engine/Source/Runtime/Engine/Classes/Components/SplineMeshComponent.h
--- a/Engine/Source/Runtime/Engine/Classes/Components/SplineMeshComponent.h
+++ b/Engine/Source/Runtime/Engine/Classes/Components/SplineMeshComponent.h
@@ -173,6 +173,9 @@ class USplineMeshComponent : public UStaticMeshComponent, public IInterface_Coll
 #if WITH_EDITORONLY_DATA
 	UPROPERTY(transient)
 	uint8 bSelected:1;
+
+	UPROPERTY(EditAnywhere, Category = HLOD, AdvancedDisplay)
+	TObjectPtr<UHLODLayer> HLODLayerOverrideForSplinePartitioning;
 #endif
 
 private:
@@ -191,6 +194,7 @@ public:
 	ENGINE_API virtual void Serialize(FArchive& Ar) override;
 #if WITH_EDITOR
 	ENGINE_API bool IsRelevantForSplinePartitioning() const;
+	ENGINE_API UHLODLayer* GetHLODLayerForSplinePartitioning() const;
 	ENGINE_API virtual bool IsEditorOnly() const override;
 	ENGINE_API virtual bool Modify(bool bAlwaysMarkDirty = true) override;
 	ENGINE_API virtual void PostEditChangeProperty(FPropertyChangedEvent& PropertyChangedEvent) override;
diff --git a/Engine/Source/Runtime/Engine/Private/Components/SplineMeshComponent.cpp b/Engine/Source/Runtime/Engine/Private/Components/SplineMeshComponent.cpp
--- a/Engine/Source/Runtime/Engine/Private/Components/SplineMeshComponent.cpp
+++ b/Engine/Source/Runtime/Engine/Private/Components/SplineMeshComponent.cpp
@@ -636,6 +636,21 @@ bool USplineMeshComponent::IsRelevantForSplinePartitioning() const
 	return true;
 }
 
+UHLODLayer* USplineMeshComponent::GetHLODLayerForSplinePartitioning() const
+{
+	if (HLODLayerOverrideForSplinePartitioning)
+	{
+		return HLODLayerOverrideForSplinePartitioning;
+	}
+	
+	if (GetOwner())
+	{
+		return GetOwner()->GetHLODLayer();
+	}
+
+	return nullptr;
+}
+
 bool USplineMeshComponent::IsEditorOnly() const
 {
 	if (Super::IsEditorOnly())
diff --git a/Engine/Source/Runtime/Landscape/Classes/ControlPointMeshComponent.h b/Engine/Source/Runtime/Landscape/Classes/ControlPointMeshComponent.h
--- a/Engine/Source/Runtime/Landscape/Classes/ControlPointMeshComponent.h
+++ b/Engine/Source/Runtime/Landscape/Classes/ControlPointMeshComponent.h
@@ -15,6 +15,9 @@ class UControlPointMeshComponent : public UStaticMeshComponent
 #if WITH_EDITORONLY_DATA
 	UPROPERTY(transient)
 	uint32 bSelected:1;
+
+	UPROPERTY(EditAnywhere, Category = HLOD, AdvancedDisplay)
+	TObjectPtr<UHLODLayer> HLODLayerOverrideForSplinePartitioning;
 #endif
 
 	/** 
@@ -37,5 +40,6 @@ class UControlPointMeshComponent : public UStaticMeshComponent
 
 #if WITH_EDITOR
 	LANDSCAPE_API bool IsRelevantForSplinePartitioning() const;
+	LANDSCAPE_API UHLODLayer* GetHLODLayerForSplinePartitioning() const;
 #endif
 };
diff --git a/Engine/Source/Runtime/Landscape/Classes/LandscapeSplineControlPoint.h b/Engine/Source/Runtime/Landscape/Classes/LandscapeSplineControlPoint.h
--- a/Engine/Source/Runtime/Landscape/Classes/LandscapeSplineControlPoint.h
+++ b/Engine/Source/Runtime/Landscape/Classes/LandscapeSplineControlPoint.h
@@ -196,6 +196,10 @@ class ULandscapeSplineControlPoint : public UObject
 	UPROPERTY(EditAnywhere, BlueprintReadOnly, Category = Collision, meta = (ShowOnlyInnerProperties))
 	FBodyInstance BodyInstance;
 
+	/** Custom HLOD layer to use for this spline control point. Only relevant when using World Partition with spline partitioning. */
+	UPROPERTY(EditAnywhere, Category = HLOD, AdvancedDisplay)
+	TObjectPtr<UHLODLayer> HLODLayerOverrideForSplinePartitioning;
+
 protected:
 	UPROPERTY(Transient)
 	uint32 bSelected : 1;
diff --git a/Engine/Source/Runtime/Landscape/Classes/LandscapeSplineSegment.h b/Engine/Source/Runtime/Landscape/Classes/LandscapeSplineSegment.h
--- a/Engine/Source/Runtime/Landscape/Classes/LandscapeSplineSegment.h
+++ b/Engine/Source/Runtime/Landscape/Classes/LandscapeSplineSegment.h
@@ -299,6 +299,10 @@ class ULandscapeSplineSegment : public UObject
 	UPROPERTY(EditAnywhere, BlueprintReadOnly, Category = Collision, meta = (ShowOnlyInnerProperties))
 	FBodyInstance BodyInstance;
 
+	/** Custom HLOD layer to use for this spline segment. Only relevant when using World Partition with spline partitioning. */
+	UPROPERTY(EditAnywhere, Category = HLOD, AdvancedDisplay)
+	TObjectPtr<UHLODLayer> HLODLayerOverrideForSplinePartitioning;
+
 protected:
 	UPROPERTY(Transient)
 	uint32 bSelected : 1;
diff --git a/Engine/Source/Runtime/Landscape/Private/LandscapeSplines.cpp b/Engine/Source/Runtime/Landscape/Private/LandscapeSplines.cpp
--- a/Engine/Source/Runtime/Landscape/Private/LandscapeSplines.cpp
+++ b/Engine/Source/Runtime/Landscape/Private/LandscapeSplines.cpp
@@ -1781,6 +1781,21 @@ bool UControlPointMeshComponent::IsRelevantForSplinePartitioning() const
 	return true;
 }
 
+UHLODLayer* UControlPointMeshComponent::GetHLODLayerForSplinePartitioning() const
+{
+	if (HLODLayerOverrideForSplinePartitioning)
+	{
+		return HLODLayerOverrideForSplinePartitioning;
+	}
+
+	if (GetOwner())
+	{
+		return GetOwner()->GetHLODLayer();
+	}
+
+	return nullptr;
+}
+
 bool UControlPointMeshComponent::IsEditorOnly() const
 {
 	if (Super::IsEditorOnly())
@@ -2447,6 +2462,12 @@ void ULandscapeSplineControlPoint::UpdateSplinePoints(bool bUpdateCollision, boo
 			MeshComponent->MarkRenderStateDirty();
 		}
 
+		if (MeshComponent->HLODLayerOverrideForSplinePartitioning != HLODLayerOverrideForSplinePartitioning)
+		{
+			MeshComponent->Modify();
+			MeshComponent->HLODLayerOverrideForSplinePartitioning = HLODLayerOverrideForSplinePartitioning;
+		}
+
 		if (bComponentNeedsRegistering)
 		{
 			MeshComponent->RegisterComponent();
@@ -3333,6 +3354,8 @@ void ULandscapeSplineSegment::UpdateSplinePoints(bool bUpdateCollision, bool bUp
 			MeshComponent->SetHiddenInGame(bUsingEditorMesh || bHiddenInGame);
 			MeshComponent->SetVisibility(!bUsingEditorMesh || OuterSplines->bShowSplineEditorMesh);
 
+			MeshComponent->HLODLayerOverrideForSplinePartitioning = HLODLayerOverrideForSplinePartitioning;
+
 			MeshSettings.Add(FMeshSettings(T, MeshEntry));
 			iMesh++;
 			T += MeshT;
