Fixed off-by-one bounds computation when rastering landscape splines, which could lead to inconsistent results when a spline segment was located exactly at the border of 2 landscape streaming proxies (indeterminism)

#Synced-CL 48090587

diff --git a/Engine/Source/Runtime/Landscape/Private/LandscapeSplineRaster.cpp b/Engine/Source/Runtime/Landscape/Private/LandscapeSplineRaster.cpp
--- a/Engine/Source/Runtime/Landscape/Private/LandscapeSplineRaster.cpp
+++ b/Engine/Source/Runtime/Landscape/Private/LandscapeSplineRaster.cpp
@@ -645,7 +645,6 @@ bool ULandscapeInfo::ApplySplinesInternal(TScriptInterface<ILandscapeSplineInter
 	FLandscapeDoNotDirtyScope DoNotDirtyScope(LandscapeEdit, !bMarkPackageDirty);
 	TSet<ULandscapeComponent*> ModifiedComponents;
 	
-
 	for (const ULandscapeSplineControlPoint* ControlPoint : SplineComponent->ControlPoints)
 	{
 		if (ControlPoint->GetPoints().Num() < 2)
@@ -656,10 +655,10 @@ bool ULandscapeInfo::ApplySplinesInternal(TScriptInterface<ILandscapeSplineInter
 		FBox ControlPointBounds = ControlPoint->GetBounds();
 		ControlPointBounds = ControlPointBounds.TransformBy(SplineToLandscape.ToMatrixWithScale());
 
-		int32 MinX = FMath::CeilToInt32(ControlPointBounds.Min.X);
-		int32 MinY = FMath::CeilToInt32(ControlPointBounds.Min.Y);
-		int32 MaxX = FMath::FloorToInt32(ControlPointBounds.Max.X);
-		int32 MaxY = FMath::FloorToInt32(ControlPointBounds.Max.Y);
+		int32 MinX = FMath::FloorToInt32(ControlPointBounds.Min.X);
+		int32 MinY = FMath::FloorToInt32(ControlPointBounds.Min.Y);
+		int32 MaxX = FMath::CeilToInt32(ControlPointBounds.Max.X);
+		int32 MaxY = FMath::CeilToInt32(ControlPointBounds.Max.Y);
 
 		MinX = FMath::Max(MinX, LandscapeMinX);
 		MinY = FMath::Max(MinY, LandscapeMinY);
@@ -722,10 +721,10 @@ bool ULandscapeInfo::ApplySplinesInternal(TScriptInterface<ILandscapeSplineInter
 		FBox SegmentBounds = Segment->GetBounds();
 		SegmentBounds = SegmentBounds.TransformBy(SplineToLandscape.ToMatrixWithScale());
 
-		int32 MinX = FMath::CeilToInt32(SegmentBounds.Min.X);
-		int32 MinY = FMath::CeilToInt32(SegmentBounds.Min.Y);
-		int32 MaxX = FMath::FloorToInt32(SegmentBounds.Max.X);
-		int32 MaxY = FMath::FloorToInt32(SegmentBounds.Max.Y);
+		int32 MinX = FMath::FloorToInt32(SegmentBounds.Min.X);
+		int32 MinY = FMath::FloorToInt32(SegmentBounds.Min.Y);
+		int32 MaxX = FMath::CeilToInt32(SegmentBounds.Max.X);
+		int32 MaxY = FMath::CeilToInt32(SegmentBounds.Max.Y);
 
 		MinX = FMath::Max(MinX, LandscapeMinX);
 		MinY = FMath::Max(MinY, LandscapeMinY);
@@ -828,10 +827,10 @@ namespace LandscapeSplineRaster
 
 		SegmentBounds = SegmentBounds.TransformBy(SplineToLandscape.ToMatrixWithScale());
 
-		int32 MinX = FMath::CeilToInt32(SegmentBounds.Min.X);
-		int32 MinY = FMath::CeilToInt32(SegmentBounds.Min.Y);
-		int32 MaxX = FMath::FloorToInt32(SegmentBounds.Max.X);
-		int32 MaxY = FMath::FloorToInt32(SegmentBounds.Max.Y);
+		int32 MinX = FMath::FloorToInt32(SegmentBounds.Min.X);
+		int32 MinY = FMath::FloorToInt32(SegmentBounds.Min.Y);
+		int32 MaxX = FMath::CeilToInt32(SegmentBounds.Max.X);
+		int32 MaxY = FMath::CeilToInt32(SegmentBounds.Max.Y);
 
 		MinX = FMath::Max(MinX, LandscapeMinX);
 		MinY = FMath::Max(MinY, LandscapeMinY);
