Add support for custom vertex factory types during landscape PSO precaching when custom proxies are used

#Synced-CL 48092843

diff --git a/Engine/Source/Runtime/Landscape/Private/Landscape.cpp b/Engine/Source/Runtime/Landscape/Private/Landscape.cpp
--- a/Engine/Source/Runtime/Landscape/Private/Landscape.cpp
+++ b/Engine/Source/Runtime/Landscape/Private/Landscape.cpp
@@ -1261,6 +1261,13 @@ void ULandscapeComponent::PostLoad()
 			VertexFactoryDataList.Add(FPSOPrecacheVertexFactoryData(Culling::GetTileVertexFactoryType()));
 		}
 
+		ILandscapeModule& LandscapeModule = FModuleManager::GetModuleChecked<ILandscapeModule>("Landscape");
+		const UE::Landscape::FSetupLandscapePSOPrecacheVertexFactoryDataListDelegate& SetupLandscapePSOPrecacheVertexFactoryDataListDelegate = LandscapeModule.GetSetupLandscapePSOPrecacheVertexFactoryDataListDelegate();
+		if (SetupLandscapePSOPrecacheVertexFactoryDataListDelegate.IsBound())
+		{
+			SetupLandscapePSOPrecacheVertexFactoryDataListDelegate.Execute(this, VertexFactoryDataList);
+		}
+
 		TArray<FMaterialPSOPrecacheRequestID> MaterialPrecacheRequestIDs;
 		for (UMaterialInterface* MaterialInterface : Materials)
 		{
diff --git a/Engine/Source/Runtime/Landscape/Private/LandscapeModule.cpp b/Engine/Source/Runtime/Landscape/Private/LandscapeModule.cpp
--- a/Engine/Source/Runtime/Landscape/Private/LandscapeModule.cpp
+++ b/Engine/Source/Runtime/Landscape/Private/LandscapeModule.cpp
@@ -51,6 +51,13 @@ public:
 		CreateLandscapeComponentSceneProxyDelegate = InDelegate;
 	}
 	virtual const UE::Landscape::FCreateLandscapeComponentSceneProxyDelegate& GetCreateLandscapeComponentSceneProxyDelegate() const override { return CreateLandscapeComponentSceneProxyDelegate; }
+	virtual void SetSetupLandscapePSOPrecacheVertexFactoryDataListDelegate(const UE::Landscape::FSetupLandscapePSOPrecacheVertexFactoryDataListDelegate& InDelegate) override
+	{
+		// Setting new delegate without reseting previous one is not allowed.
+		check(SetupLandscapePSOPrecacheVertexFactoryDataListDelegate.IsBound() != InDelegate.IsBound());
+		SetupLandscapePSOPrecacheVertexFactoryDataListDelegate = InDelegate;
+	}
+	virtual const UE::Landscape::FSetupLandscapePSOPrecacheVertexFactoryDataListDelegate& GetSetupLandscapePSOPrecacheVertexFactoryDataListDelegate() const override { return SetupLandscapePSOPrecacheVertexFactoryDataListDelegate; }
 
 private:
 	void OnPostEngineInit();
@@ -64,6 +71,7 @@ private:
 	TSharedPtr<FLandscapeSceneViewExtension, ESPMode::ThreadSafe> SceneViewExtension;
 	ILandscapeEditorServices* LandscapeEditorServices = nullptr;
 	UE::Landscape::FCreateLandscapeComponentSceneProxyDelegate CreateLandscapeComponentSceneProxyDelegate;
+	UE::Landscape::FSetupLandscapePSOPrecacheVertexFactoryDataListDelegate SetupLandscapePSOPrecacheVertexFactoryDataListDelegate;
 
 	/** Event binding handles */
 #if ALLOW_LANDSCAPE_COLLISION_COMPONENT_DEBUG_DISPLAY
diff --git a/Engine/Source/Runtime/Landscape/Public/LandscapeModule.h b/Engine/Source/Runtime/Landscape/Public/LandscapeModule.h
--- a/Engine/Source/Runtime/Landscape/Public/LandscapeModule.h
+++ b/Engine/Source/Runtime/Landscape/Public/LandscapeModule.h
@@ -13,6 +13,7 @@ class ULandscapeComponent;
 namespace UE::Landscape
 {
 	DECLARE_DELEGATE_RetVal_OneParam(FPrimitiveSceneProxy*, FCreateLandscapeComponentSceneProxyDelegate, ULandscapeComponent* /*InLandscapeComponent*/);
+	DECLARE_DELEGATE_TwoParams(FSetupLandscapePSOPrecacheVertexFactoryDataListDelegate, ULandscapeComponent* /*InLandscapeComponent*/, FPSOPrecacheVertexFactoryDataList& /*VertexFactoryDataList*/);
 } // namespace UE::Landscape
 
 /**
@@ -28,4 +29,7 @@ public:
 	virtual void SetCreateLandscapeComponentSceneProxyDelegate(const UE::Landscape::FCreateLandscapeComponentSceneProxyDelegate& InDelegate) PURE_VIRTUAL(ILandscapeModule::SetCreateLandscapeComponentSceneProxyDelegate, );
 	virtual const UE::Landscape::FCreateLandscapeComponentSceneProxyDelegate& GetCreateLandscapeComponentSceneProxyDelegate() const 
 		PURE_VIRTUAL(ILandscapeModule::GetCreateLandscapeComponentSceneProxyDelegate, static UE::Landscape::FCreateLandscapeComponentSceneProxyDelegate Dummy; return Dummy; );
+	virtual void SetSetupLandscapePSOPrecacheVertexFactoryDataListDelegate(const UE::Landscape::FSetupLandscapePSOPrecacheVertexFactoryDataListDelegate& InDelegate) PURE_VIRTUAL(ILandscapeModule::SetSetupLandscapePSOPrecacheVertexFactoryDataListDelegate, );
+	virtual const UE::Landscape::FSetupLandscapePSOPrecacheVertexFactoryDataListDelegate& GetSetupLandscapePSOPrecacheVertexFactoryDataListDelegate() const 
+		PURE_VIRTUAL(ILandscapeModule::GetSetupLandscapePSOPrecacheVertexFactoryDataListDelegate, static UE::Landscape::FSetupLandscapePSOPrecacheVertexFactoryDataListDelegate Dummy; return Dummy; );
 };
