Landscape: Reduce cook logvolume: Downgrade the DeleteSplinePoints logs to verbose because they are triggering many times during cook (44k/120k lines in a large test project).
#rnx
#Synced-CL Unknown

diff --git a/Engine/Source/Runtime/Landscape/Private/LandscapeSplines.cpp b/Engine/Source/Runtime/Landscape/Private/LandscapeSplines.cpp
--- a/Engine/Source/Runtime/Landscape/Private/LandscapeSplines.cpp
+++ b/Engine/Source/Runtime/Landscape/Private/LandscapeSplines.cpp
@@ -2612,7 +2612,7 @@ void ULandscapeSplineControlPoint::DeleteSplinePoints()
 	Modify();
 
 	ULandscapeSplinesComponent* OuterSplines = GetOuterChecked();
-	UE_LOG(LogLandscape, Display, TEXT("ULandscapeSplineControlPoint::DeleteSplinePoints removing %s from spline actor %s"), *GetName(), *OuterSplines->GetOwner()->GetActorNameOrLabel());
+	UE_LOG(LogLandscape, Verbose, TEXT("ULandscapeSplineControlPoint::DeleteSplinePoints removing %s from spline actor %s"), *GetName(), *OuterSplines->GetOwner()->GetActorNameOrLabel());
 
 	Points.Reset();
 	Bounds = FBox(ForceInit);
@@ -3660,7 +3660,7 @@ void ULandscapeSplineSegment::DeleteSplinePoints()
 	Modify();
 
 	ULandscapeSplinesComponent* OuterSplines = GetOuterChecked();
-	UE_LOG(LogLandscape, Display, TEXT("ULandscapeSplineSegment::DeleteSplinePoints removing %s from spline actor %s"), *GetName(), *OuterSplines->GetOwner()->GetActorNameOrLabel());
+	UE_LOG(LogLandscape, Verbose, TEXT("ULandscapeSplineSegment::DeleteSplinePoints removing %s from spline actor %s"), *GetName(), *OuterSplines->GetOwner()->GetActorNameOrLabel());
 
 	SplineInfo.Reset();
 	Points.Reset();
