Change landscape spline error checking code to purge nulls even when checks are disabled.

#Synced-CL 48369572

diff --git a/Engine/Source/Runtime/Landscape/Private/LandscapeSplines.cpp b/Engine/Source/Runtime/Landscape/Private/LandscapeSplines.cpp
--- a/Engine/Source/Runtime/Landscape/Private/LandscapeSplines.cpp
+++ b/Engine/Source/Runtime/Landscape/Private/LandscapeSplines.cpp
@@ -189,25 +189,31 @@ public:
 		Segments.Reserve(Component->Segments.Num());
 		for (ULandscapeSplineSegment* Segment : Component->Segments)
 		{
-			FSegmentProxy SegmentProxy;
-			SegmentProxy.Owner = Segment;
-			SegmentProxy.HitProxy = nullptr;
-			SegmentProxy.Points = Segment->GetPoints();
-			SegmentProxy.bSelected = Segment->IsSplineSelected();
-			Segments.Add(SegmentProxy);
+			if (Segment)
+			{
+				FSegmentProxy SegmentProxy;
+				SegmentProxy.Owner = Segment;
+				SegmentProxy.HitProxy = nullptr;
+				SegmentProxy.Points = Segment->GetPoints();
+				SegmentProxy.bSelected = Segment->IsSplineSelected();
+				Segments.Add(SegmentProxy);
+			}
 		}
 
 		ControlPoints.Reserve(Component->ControlPoints.Num());
 		for (ULandscapeSplineControlPoint* ControlPoint : Component->ControlPoints)
 		{
-			FControlPointProxy ControlPointProxy;
-			ControlPointProxy.Owner = ControlPoint;
-			ControlPointProxy.HitProxy = nullptr;
-			ControlPointProxy.Location = ControlPoint->Location;
-			ControlPointProxy.Points = ControlPoint->GetPoints();
-			ControlPointProxy.SpriteScale = FMath::Clamp<float>(ControlPoint->Width != 0 ? ControlPoint->Width / 8 : ControlPoint->SideFalloff / 4, 10, 1000);
-			ControlPointProxy.bSelected = ControlPoint->IsSplineSelected();
-			ControlPoints.Add(ControlPointProxy);
+			if (ControlPoint)
+			{
+				FControlPointProxy ControlPointProxy;
+				ControlPointProxy.Owner = ControlPoint;
+				ControlPointProxy.HitProxy = nullptr;
+				ControlPointProxy.Location = ControlPoint->Location;
+				ControlPointProxy.Points = ControlPoint->GetPoints();
+				ControlPointProxy.SpriteScale = FMath::Clamp<float>(ControlPoint->Width != 0 ? ControlPoint->Width / 8 : ControlPoint->SideFalloff / 4, 10, 1000);
+				ControlPointProxy.bSelected = ControlPoint->IsSplineSelected();
+				ControlPoints.Add(ControlPointProxy);
+			}
 		}
 	}
 
@@ -531,11 +537,12 @@ ILandscapeSplineInterface* ULandscapeSplinesComponent::GetSplineOwner()
 
 void ULandscapeSplinesComponent::CheckSplinesValid()
 {
-#if DO_CHECK
-	// Remove all null control points/segments.  Likely, some error occurred in a previous session and the package is saved with dud references.
-	ensureMsgf(!ControlPoints.Remove(nullptr), TEXT("Error in ULandscapeSplinesComponent in %s, actor: %s"), *GetPathName(), GetOwner() ? *GetOwner()->GetActorNameOrLabel() : TEXT("null"));
-	ensureMsgf(!Segments.Remove(nullptr), TEXT("Error in ULandscapeSplinesComponent in %s, actor: %s"), *GetPathName(), GetOwner() ? *GetOwner()->GetActorNameOrLabel() : TEXT("null"));
+	// Remove all null control points/segments.  Do this even if checks are disabled to protect against saved bad data.
+	// Likely, some error occurred in a previous session and the package is saved with dud references.
+	ensureMsgf(!ControlPoints.Remove(nullptr), TEXT("Error found and resolved in ULandscapeSplinesComponent in %s, actor: %s"), *GetPathName(), GetOwner() ? *GetOwner()->GetActorNameOrLabel() : TEXT("null"));
+	ensureMsgf(!Segments.Remove(nullptr), TEXT("Error found and resolved in ULandscapeSplinesComponent in %s, actor: %s"), *GetPathName(), GetOwner() ? *GetOwner()->GetActorNameOrLabel() : TEXT("null"));
 
+#if DO_CHECK
 	// Check for cross-spline connections, as this is a potential source of nulls.  This can also happen if something went wrong
 	// with a segment/point connectivity edit (while merging splines, undo/redo of deletions, etc).  These errors show up some times,
 	// but the trigger conditions are not known.
