Adjusted tooltips for some landscape spline-related settings

#Synced-CL 48524909

diff --git a/Engine/Source/Runtime/Landscape/Public/LandscapeSettings.h b/Engine/Source/Runtime/Landscape/Public/LandscapeSettings.h
--- a/Engine/Source/Runtime/Landscape/Public/LandscapeSettings.h
+++ b/Engine/Source/Runtime/Landscape/Public/LandscapeSettings.h
@@ -129,11 +129,12 @@ protected:
 	UPROPERTY(transient)
 	bool bRestrictiveMode = false;
 
-	/** Offset in Z for the landscape spline icon in world-space. */
+	/** Offset in Z for the landscape spline point icon in world-space. */
 	UPROPERTY(EditAnywhere, config, Category = "Spline", meta = (ClampMin = "0", UIMax = "4096"))
 	float SplineIconWorldZOffset = 100.0f;
 
 	UPROPERTY(EditAnywhere, config, Category = "Spline", meta = (ClampMin = "32", UIMax = "2048"))
+	/** Size of the landscape spline control point icon in the viewport. */
 	float SplineIconScale = 125.0f;
 
 	UPROPERTY(EditAnywhere, config, Category = "Target Layers", meta = (ToolTip = 
