Don't overwrite the landscape actor pointer for spline actors when updating shared properties.

#Synced-CL 48612369

diff --git a/Engine/Source/Runtime/Landscape/Private/LandscapeSplineActor.cpp b/Engine/Source/Runtime/Landscape/Private/LandscapeSplineActor.cpp
--- a/Engine/Source/Runtime/Landscape/Private/LandscapeSplineActor.cpp
+++ b/Engine/Source/Runtime/Landscape/Private/LandscapeSplineActor.cpp
@@ -46,7 +46,10 @@ void ALandscapeSplineActor::UpdateSharedProperties(ULandscapeInfo* InLandscapeIn
 	check(InLandscapeInfo);
 	check(InLandscapeInfo->LandscapeGuid == LandscapeGuid);
 #if WITH_EDITOR
-	LandscapeActor = InLandscapeInfo->LandscapeActor.Get();
+	if (ALandscape* LandscapeInfoActor = InLandscapeInfo->LandscapeActor.Get())
+	{
+		LandscapeActor = LandscapeInfoActor;
+	}
 #endif // WITH_EDITOR
 }
 
