Landscape streaming proxies : removed fallback to ALandscape when Landscape(Hole)Material are nullptr. This is a remnant from the time where there was no LandscapeOverridable property system. This is now reported as a MapCheck warning (or info, if landscape.SilenceInvalidPropertyOverridesMapCheckWarnings is 1) and the editor returns an invalid material in such a case. This avoids cook-only errors when a user overrides the property but sets it to null : now the editor shows the default material, which is much better than the relying on the parent ALandscape's material, since its presence is not guaranteed on cook

#Synced-CL 48911231

diff --git a/Engine/Source/Runtime/Core/Private/Misc/MapErrors.cpp b/Engine/Source/Runtime/Core/Private/Misc/MapErrors.cpp
--- a/Engine/Source/Runtime/Core/Private/Misc/MapErrors.cpp
+++ b/Engine/Source/Runtime/Core/Private/Misc/MapErrors.cpp
@@ -41,6 +41,7 @@ FLazyName FMapErrors::FixedUpSharedLayerWeightmap(TEXT("FixedUpSharedLayerWeight
 FLazyName FMapErrors::LandscapeComponentPostLoad_Warning(TEXT("LandscapeComponentPostLoad_Warning"));
 FLazyName FMapErrors::LandscapeLazyObjectPtrDeprecation_Warning(TEXT("LandscapeLazyObjectPtrDeprecation_Warning"));
 FLazyName FMapErrors::LandscapeProxyFixupSharedData_Warning(TEXT("LandscapeProxyFixupSharedData_Warning"));
+FLazyName FMapErrors::LandscapeProxyInvalidPropertyOverride_Warning(TEXT("LandscapeProxyInvalidPropertyOverride_Warning"));
 FLazyName FMapErrors::DuplicateLevelInfo(TEXT("DuplicateLevelInfo"));
 FLazyName FMapErrors::NoKillZ(TEXT("NoKillZ"));
 FLazyName FMapErrors::LightComponentNull(TEXT("LightComponentNull"));
diff --git a/Engine/Source/Runtime/Core/Public/Misc/MapErrors.h b/Engine/Source/Runtime/Core/Public/Misc/MapErrors.h
--- a/Engine/Source/Runtime/Core/Public/Misc/MapErrors.h
+++ b/Engine/Source/Runtime/Core/Public/Misc/MapErrors.h
@@ -132,9 +132,12 @@ struct FMapErrors
 	/** Landscape proxy {ProxyName} of {LevelName} points to a LandscapeActor that is not currently loaded. This will lose the property upon save.  */
 	static CORE_API FLazyName LandscapeLazyObjectPtrDeprecation_Warning;
 
-	/** Landscape proxy {Proxy} had one or more shared properties not in sync with its parent landscape actor {Parent}. This has been fixed but the proxy needs to be saved in order to ensure cooking behaves as expected..  */
+	/** Landscape proxy {Proxy} had one or more shared properties not in sync with its parent landscape actor {Parent}. This has been fixed but the proxy needs to be saved in order to ensure cooking behaves as expected. */
 	static CORE_API FLazyName LandscapeProxyFixupSharedData_Warning;
 
+	/** The landscape proxy property has been overridden but is invalid.  */
+	static CORE_API FLazyName LandscapeProxyInvalidPropertyOverride_Warning;
+
     /**  Level */
 
 	/** Duplicate level info: Two WorldInfos somehow exist... */
diff --git a/Engine/Source/Runtime/Landscape/Classes/LandscapeProxy.h b/Engine/Source/Runtime/Landscape/Classes/LandscapeProxy.h
--- a/Engine/Source/Runtime/Landscape/Classes/LandscapeProxy.h
+++ b/Engine/Source/Runtime/Landscape/Classes/LandscapeProxy.h
@@ -1649,6 +1649,9 @@ protected:
 	/** Fixup any internal representation for shared properties. Will be used when deprecating or renaming a property. */
 	virtual void FixupOverriddenSharedProperties() {}
 
+	/** Validates that overridden shared properties are properly setup. Only run on load and MapCheck since it's just to inform the user that something is not properly setup */
+	virtual void CheckOverriddenSharedProperties() {}
+
 	/**
 	* Synchronizes all shared properties that cannot be marked by regular meta tags.
 	* @param The landscape to use as a template.
@@ -1656,6 +1659,7 @@ protected:
 	*/
 	LANDSCAPE_API virtual TArray<FName> SynchronizeUnmarkedSharedProperties(ALandscapeProxy* InLandscape);
 
+
 protected:
 	UE_DEPRECATED(5.6, "Use ULandscapeSubsystem::OnLandscapeProxyComponentDataChanged")
 	FLandscapeMaterialChangedDelegate LandscapeMaterialChangedDelegate;
diff --git a/Engine/Source/Runtime/Landscape/Classes/LandscapeStreamingProxy.h b/Engine/Source/Runtime/Landscape/Classes/LandscapeStreamingProxy.h
--- a/Engine/Source/Runtime/Landscape/Classes/LandscapeStreamingProxy.h
+++ b/Engine/Source/Runtime/Landscape/Classes/LandscapeStreamingProxy.h
@@ -55,6 +55,7 @@ public:
 	virtual bool CanChangeIsSpatiallyLoadedFlag() const override { return AActor::CanChangeIsSpatiallyLoadedFlag(); }
 	virtual bool ShouldIncludeGridSizeInName(UWorld* InWorld, const FActorPartitionIdentifier& InIdentifier) const override;
 	virtual void GetActorDescProperties(FPropertyPairsMap& PropertyPairsMap) const override;
+	virtual void CheckForErrors() override;
 #endif
 	//~ End UObject Interface
 
@@ -78,6 +79,7 @@ public:
 	void RequestDeferredLayerContentUpdateModes(ELandscapeLayerUpdateMode InLayerContentUpdateModes);
 	void PerformDeferredLayerContentUpdateModesRequest();
 protected:
+	virtual void CheckOverriddenSharedProperties() override;
 	virtual void FixupOverriddenSharedProperties() override;
 #endif // WITH_EDITOR
 };
diff --git a/Engine/Source/Runtime/Landscape/Private/Landscape.cpp b/Engine/Source/Runtime/Landscape/Private/Landscape.cpp
--- a/Engine/Source/Runtime/Landscape/Private/Landscape.cpp
+++ b/Engine/Source/Runtime/Landscape/Private/Landscape.cpp
@@ -220,6 +220,11 @@ static FAutoConsoleVariable CVarSilenceSharedPropertyDeprecationFixup(
 	true,
 	TEXT("Silently performs the fixup of discrepancies in shared properties when handling data modified before the enforcement introduction."));
 
+FAutoConsoleVariable CVarLandscapeSilenceInvalidPropertyOverridesMapCheckWarnings(
+	TEXT("landscape.SilenceInvalidPropertyOverridesMapCheckWarnings"),
+	false,
+	TEXT("Issue MapCheck Info messages instead of warnings if some landscape property overrides are incorrectly set (e.g. overriding the LandscapeMaterial but leaving it empty)."));
+
 static FAutoConsoleVariable CVarLandscapeSilenceMapCheckWarnings_Nanite(
 	TEXT("landscape.Nanite.SilenceMapCheckWarnings"),
 	false,
@@ -5148,10 +5153,11 @@ UMaterialInterface* ALandscapeStreamingProxy::GetLandscapeMaterial(int8 InLODInd
 		return LandscapeMaterial;
 	}
 
-	if (const ALandscape* Landscape = GetLandscapeActor())
-	{
-		return Landscape->GetLandscapeMaterial(InLODIndex);
-	}
+	// Note : we don't fallback to the parent ALandscape anymore, since this is only reliable in the editor. The LandscapeMaterial property is now a LandscapeOverridable one,
+	//  which means that overriding the material is now done explicitly via OverriddenSharedProperties rather than by the LandscapeMaterial property value being not null. It means it's now a
+	//  "valid" thing to override the material property and set it to "None" (probably a user mistake, though, hence the MapCheck warning if data is setup this way). 
+	//  In order to make the error plain visible, we therefore return an invalid material now, so that the editor doesn't show the parent ALandscape's material, while cooking
+	//  potentially returns an invalid material depending on whether ALandscape is loaded. 
 
 	return UMaterial::GetDefaultMaterial(MD_Surface);
 }
@@ -5162,10 +5168,13 @@ UMaterialInterface* ALandscapeStreamingProxy::GetLandscapeHoleMaterial() const
 	{
 		return LandscapeHoleMaterial;
 	}
-	else if (const ALandscape* Landscape = GetLandscapeActor())
-	{
-		return Landscape->GetLandscapeHoleMaterial();
-	}
+
+	// Note : we don't fallback to the parent ALandscape anymore, since this is only reliable in the editor. The LandscapeHoleMaterial property is now a LandscapeOverridable one,
+	//  which means that overriding the material is now done explicitly via OverriddenSharedProperties rather than by the LandscapeHoleMaterial property value being not null. It means it's now a
+	//  "valid" thing to override the material property and set it to "None" (probably a user mistake, though, hence the MapCheck warning if data is setup this way). 
+	//  In order to make the error plain visible, we therefore return a null material now, so that the editor doesn't show the parent ALandscape's hole material, while cooking
+	//  potentially returns an invalid material depending on whether ALandscape is loaded. 
+
 	return nullptr;
 }
 
@@ -5432,6 +5441,9 @@ void ALandscapeProxy::FixupSharedData(ALandscape* Landscape, const bool bMapChec
 		}
 	}
 
+	// Validate that property overrides are properly setup
+	CheckOverriddenSharedProperties();
+
 	OnLandscapeProxyFixupSharedDataDelegate.Broadcast(/*Proxy = */this, FOnLandscapeProxyFixupSharedDataParams { .Landscape = Landscape, .bUpgradeSharedPropertiesPerformed = bUpgradeSharedPropertiesPerformedBefore });
 }
 
diff --git a/Engine/Source/Runtime/Landscape/Private/LandscapeEdit.cpp b/Engine/Source/Runtime/Landscape/Private/LandscapeEdit.cpp
--- a/Engine/Source/Runtime/Landscape/Private/LandscapeEdit.cpp
+++ b/Engine/Source/Runtime/Landscape/Private/LandscapeEdit.cpp
@@ -117,6 +117,7 @@ static TAutoConsoleVariable<int32> CVarLandscapeApplyPhysicalMaterialChangesImme
     TEXT("Applies physical material task changes immediately rather than during the next cook/PIE."));
 
 extern FAutoConsoleVariable CVarStripLayerTextureMipsOnLoad;
+extern FAutoConsoleVariable CVarLandscapeSilenceInvalidPropertyOverridesMapCheckWarnings;
 
 #if WITH_EDITOR
 
@@ -6274,6 +6275,58 @@ PRAGMA_ENABLE_DEPRECATION_WARNINGS
 	}
 }
 
+void ALandscapeStreamingProxy::CheckOverriddenSharedProperties()
+{
+	ALandscape* ParentLandscape = GetLandscapeActor();
+	if (ParentLandscape == nullptr)
+	{
+		// Cannot validate anything without the parent 
+		return;
+	}
+
+	TArray<FName> PropertiesToUnOverride;
+	auto AddSharedPropertyOverrideMapCheckMessage = [this, &PropertiesToUnOverride](const FName& InOverriddenProperty)
+		{
+			FFormatNamedArguments Arguments;
+			Arguments.Add(TEXT("PropertyName"), FText::FromString(InOverriddenProperty.ToString()));
+			TWeakObjectPtr<ALandscapeStreamingProxy> LandscapeStreamingProxyWeakPtr = this;
+
+			FMessageLog MessageLog("MapCheck");
+			TSharedRef<FTokenizedMessage> TokenizedMessage = CVarLandscapeSilenceInvalidPropertyOverridesMapCheckWarnings->GetBool() ? MessageLog.Info() : MessageLog.Warning();
+			TokenizedMessage->AddToken(FUObjectToken::Create(this, FText::FromString(GetActorNameOrLabel())))
+				->AddToken(FTextToken::Create(FText::Format(LOCTEXT("MapCheck_Message_InvalidLandscapePropertyOverride", "The Landscape property {PropertyName} has been overridden but points to an invalid asset. It's most likely a mistake. The engine default material will be used. "), Arguments)))
+				->AddToken(FMapErrorToken::Create(FMapErrors::LandscapeProxyInvalidPropertyOverride_Warning))
+				->AddToken(FActionToken::Create(LOCTEXT("MapCheck_UnOverrideProperty", "Un-override property"), FText::Format(LOCTEXT("MapCheck_UnOverrideProperty_Desc", "Un-override {PropertyName}."), Arguments),
+					FOnActionTokenExecuted::CreateLambda([LandscapeStreamingProxyWeakPtr, PropertyName = InOverriddenProperty]()
+						{
+							if (LandscapeStreamingProxyWeakPtr.IsValid())
+							{
+								LandscapeStreamingProxyWeakPtr->SetSharedPropertyOverride(PropertyName, false);
+							}
+						}),
+					/*bInSingleUse = */true));
+		};
+
+	// If the user has set a material override but left the material unset, this is most likely a mistake : 
+	if ((ParentLandscape->LandscapeMaterial != nullptr) && (LandscapeMaterial == nullptr) && IsSharedPropertyOverridden(GET_MEMBER_NAME_CHECKED(ALandscapeProxy, LandscapeMaterial)))
+	{
+		AddSharedPropertyOverrideMapCheckMessage(GET_MEMBER_NAME_CHECKED(ALandscapeProxy, LandscapeMaterial));
+	}
+
+	// If the user has set a material override but left the material unset, this is most likely a mistake : 
+	if ((ParentLandscape->LandscapeHoleMaterial != nullptr) && (LandscapeHoleMaterial == nullptr) && IsSharedPropertyOverridden(GET_MEMBER_NAME_CHECKED(ALandscapeProxy, LandscapeHoleMaterial)))
+	{
+		AddSharedPropertyOverrideMapCheckMessage(GET_MEMBER_NAME_CHECKED(ALandscapeProxy, LandscapeHoleMaterial));
+	}
+}
+
+void ALandscapeStreamingProxy::CheckForErrors()
+{
+	Super::CheckForErrors();
+
+	CheckOverriddenSharedProperties();
+}
+
 bool ALandscapeStreamingProxy::CanEditChange(const FProperty* InProperty) const
 {
 	if (!Super::CanEditChange(InProperty))
