Temp fix for a crash due to the re-activation of async texture cache cancellation when UTexture::Modify is called. In the landscape flatten tool case, we use FLandscapeTextureDataInfo as a data reader, to read the final weightmap. Upon construction, this calls UTexture::Modify, which might end up cancelling the async task if one is running, expecting the texture to actually be modified. But since it's for a readonly access, the next time we run a landscape merge and we wait for textures to be "ready" (fully built/streamed in), the texture is requested for streaming (because IsTextureFullyStreamedIn will return false) again by calling TemporarilyDisableStreaming and will assert that it  worked. However, since no async cache task has been re-kicked, IsDefaultTexture will keep on returning true and an assert would occur.

The temp fix is to call BlockOnAnyAsyncBuild() on the texture *after* TemporarilyDisableStreaming (which also calls BlockOnAnyAsyncBuild, but only at beginning, not after having updated its texture resource) and *before* Modify, such that Modify cannot cancel a potential async cache task.

The better fix would be not to call Modify on read-only accesses but that is riskier

#Synced-CL 48954958

diff --git a/Engine/Source/Runtime/Landscape/Private/LandscapeEditInterface.cpp b/Engine/Source/Runtime/Landscape/Private/LandscapeEditInterface.cpp
--- a/Engine/Source/Runtime/Landscape/Private/LandscapeEditInterface.cpp
+++ b/Engine/Source/Runtime/Landscape/Private/LandscapeEditInterface.cpp
@@ -4037,7 +4037,13 @@ FLandscapeTextureDataInfo::FLandscapeTextureDataInfo(UTexture2D* InTexture, bool
 	MipInfo.AddZeroed(Texture->Source.GetNumMips());
 	Texture->SetFlags(RF_Transactional);
 	Texture->TemporarilyDisableStreaming();
+
+	// HACK [jonathan.bard] : call BlockOnAnyAsyncBuild here to make sure that Modify cannot cancel the potentially on-going async cache
+	//  The better fix would be not to call Modify when we don't intend to actually modify the texture (e.g. when FLandscapeTextureDataInfo is used for reading the data only)
+	Texture->BlockOnAnyAsyncBuild();
+
 	Texture->Modify(bShouldDirtyPackage);
+	check(!Texture->IsDefaultTexture()); // If we're the default texture, then Modify likely has cancelled the async cache, which would be a disaster if nothing requests a new cache afterwards
 }
 
 bool FLandscapeTextureDataInfo::UpdateTextureData()
