Deprecated CVar landscape.ForceLayersUpdate : it's a debug feature that was not well supported and was very inefficient anyway. The landscape.ForceLayersFullUpdate console command should be used instead to trigger a single global update
Also used the new deprecated CVar feature to properly deprecate some other already deprecated landscape CVars

#Synced-CL 48954958

diff --git a/Engine/Source/Runtime/Landscape/Private/LandscapeEditLayers.cpp b/Engine/Source/Runtime/Landscape/Private/LandscapeEditLayers.cpp
--- a/Engine/Source/Runtime/Landscape/Private/LandscapeEditLayers.cpp
+++ b/Engine/Source/Runtime/Landscape/Private/LandscapeEditLayers.cpp
@@ -113,10 +113,13 @@ DECLARE_GPU_STAT_NAMED(LandscapeLayers_ExtractLayers, TEXT("Landscape Extract La
 DECLARE_GPU_STAT_NAMED(LandscapeLayers_PackLayers, TEXT("Landscape Pack Layers"));
 
 #if WITH_EDITOR
-static TAutoConsoleVariable<int32> CVarForceLayersUpdate(
-	TEXT("landscape.ForceLayersUpdate"),
-	0,
-	TEXT("This will force landscape edit layers to be update every frame, rather than when requested only."));
+UE_DEPRECATED(5.8, "The CVar has been fully deprecated.")
+TAutoConsoleVariableDeprecated_NoReplace<int32> CVarForceLayersUpdate_Deprecated(
+	/*Name = */TEXT("landscape.ForceLayersUpdate"),
+	/*DefaultValue = */0, 
+	/*DeprecatedAtVersion = */TEXT("5.8"), 
+	/*UsageBehavior = */EShadowCVarBehavior::Warn, 
+	/*DeprecationMessage = */TEXT("landscape.ForceLayersUpdate is deprecated. Use landscape.ForceLayersFullUpdate to trigger a single global edit layers update."));
 
 int32 RenderCaptureLayersNextHeightmapDraws = 0;
 static FAutoConsoleVariableRef CVarRenderCaptureLayersNextHeightmapDraws(
@@ -131,22 +134,20 @@ static FAutoConsoleVariableRef CVarRenderCaptureLayersNextWeightmapDraws(
 	TEXT("Trigger N render captures during the next weightmap draw calls."));
 
 UE_DEPRECATED(5.7, "The CVar has been fully deprecated.")
-FAutoConsoleVariableDeprecated CVarOutputLayersRTContent(
+TAutoConsoleVariableDeprecated_NoReplace<int32> CVarOutputLayersRTContent_Deprecated(
 	/*Name = */TEXT("landscape.OutputLayersRTContent"),
-	/*ShadowName = */TEXT("landscape.OutputLayersRTContent.Deprecated"),
+	/*DefaultValue = */0,
 	/*DeprecatedAtVersion = */TEXT("5.7"),
 	/*UsageBehavior = */EShadowCVarBehavior::Warn,
-	/*LookupFailureBehavior = */EShadowCVarBehavior::Warn,
-	TEXT("landscape.OutputLayersRTContent is deprecated. This was only used in Global merge. Batched merge (landscape.EditLayersLocalMerge.Enable 2) is the only merge technique available now."));
+	/*DeprecationMessage = */TEXT("landscape.OutputLayersRTContent is deprecated. This was only used in Global merge. Batched merge (landscape.EditLayersLocalMerge.Enable 2) is the only merge technique available now."));
 
 UE_DEPRECATED(5.7, "The CVar has been fully deprecated.")
-FAutoConsoleVariableDeprecated CVarOutputLayersWeightmapsRTContent(
+TAutoConsoleVariableDeprecated_NoReplace<int32> CVarOutputLayersWeightmapsRTContent_Deprecated(
 	/*Name = */TEXT("landscape.OutputLayersWeightmapsRTContent"),
-	/*ShadowName = */TEXT("landscape.OutputLayersWeightmapsRTContent.Deprecated"),
+	/*DefaultValue = */0,
 	/*DeprecatedAtVersion = */TEXT("5.7"),
 	/*UsageBehavior = */EShadowCVarBehavior::Warn,
-	/*LookupFailureBehavior = */EShadowCVarBehavior::Warn,
-	TEXT("landscape.OutputLayersWeightmapsRTContent is deprecated. This was only used in Global merge. Batched merge (landscape.EditLayersLocalMerge.Enable 2) is the only merge technique available now."));
+	/*DeprecationMessage = */TEXT("landscape.OutputLayersWeightmapsRTContent is deprecated. This was only used in Global merge. Batched merge (landscape.EditLayersLocalMerge.Enable 2) is the only merge technique available now."));
 
 static TAutoConsoleVariable<int32> CVarLandscapeSimulatePhysics(
 	TEXT("landscape.SimulatePhysics"),
@@ -276,31 +277,28 @@ TAutoConsoleVariable<bool> CVarSilenceMergeBatchResolutionWarning(
 	TEXT("When true, don't warn about about exceeding batch merge resolution from landscape.BatchedMerge.MaxResolutionPerRenderBatch"));
 
 UE_DEPRECATED(5.7, "The CVar has been fully deprecated.")
-FAutoConsoleVariableDeprecated CVarLandscapeEditLayersLocalMerge_Deprecated(
+TAutoConsoleVariableDeprecated_NoReplace CVarLandscapeEditLayersLocalMerge_Deprecated(
 	/*Name = */TEXT("landscape.EditLayersLocalMerge.Enable"), 
-	/*ShadowName = */TEXT("landscape.EditLayersLocalMerge.Enable.Deprecated"),
+	/*DefaultValue = */0,
 	/*DeprecatedAtVersion = */TEXT("5.7"), 
 	/*UsageBehavior = */EShadowCVarBehavior::Warn, 
-	/*LookupFailureBehavior = */EShadowCVarBehavior::Warn, 
-	TEXT("landscape.EditLayersLocalMerge.Enable is deprecated. Batched merge (landscape.EditLayersLocalMerge.Enable 2) is the only merge technique available now."));
+	/*DeprecationMessage = */TEXT("landscape.EditLayersLocalMerge.Enable is deprecated. Batched merge (landscape.EditLayersLocalMerge.Enable 2) is the only merge technique available now."));
 
 UE_DEPRECATED(5.7, "The CVar has been fully deprecated.")
-FAutoConsoleVariableDeprecated CVarLandscapeEditLayersMaxComponentsPerHeightmapResolveBatch_Deprecated(
+TAutoConsoleVariableDeprecated_NoReplace CVarLandscapeEditLayersMaxComponentsPerHeightmapResolveBatch_Deprecated(
 	/*Name = */TEXT("landscape.EditLayersLocalMerge.MaxComponentsPerHeightmapResolveBatch"),
-	/*ShadowName = */TEXT("landscape.EditLayersLocalMerge.MaxComponentsPerHeightmapResolveBatch.Deprecated"),
+	/*DefaultValue = */0,
 	/*DeprecatedAtVersion = */TEXT("5.7"),
 	/*UsageBehavior = */EShadowCVarBehavior::Warn,
-	/*LookupFailureBehavior = */EShadowCVarBehavior::Warn,
-	TEXT("landscape.EditLayersLocalMerge.MaxComponentsPerHeightmapResolveBatch is deprecated. This was only used in Local merge. Batched merge (landscape.EditLayersLocalMerge.Enable 2) is the only merge technique available now."));
+	/*DeprecationMessage = */TEXT("landscape.EditLayersLocalMerge.MaxComponentsPerHeightmapResolveBatch is deprecated. This was only used in Local merge. Batched merge (landscape.EditLayersLocalMerge.Enable 2) is the only merge technique available now."));
 
 UE_DEPRECATED(5.7, "The CVar has been fully deprecated.")
-FAutoConsoleVariableDeprecated CVarLandscapeEditLayersMaxComponentsPerWeightmapResolveBatch_Deprecated(
+TAutoConsoleVariableDeprecated_NoReplace CVarLandscapeEditLayersMaxComponentsPerWeightmapResolveBatch_Deprecated(
 	/*Name = */TEXT("landscape.EditLayersLocalMerge.MaxComponentsPerWeightmapResolveBatch"),
-	/*ShadowName = */TEXT("landscape.EditLayersLocalMerge.MaxComponentsPerWeightmapResolveBatch.Deprecated"),
+	/*DefaultValue = */0,
 	/*DeprecatedAtVersion = */TEXT("5.7"),
 	/*UsageBehavior = */EShadowCVarBehavior::Warn,
-	/*LookupFailureBehavior = */EShadowCVarBehavior::Warn,
-	TEXT("landscape.EditLayersLocalMerge.MaxComponentsPerWeightmapResolveBatch is deprecated. This was only used in Local merge. Batched merge (landscape.EditLayersLocalMerge.Enable 2) is the only merge technique available now."));
+	/*DeprecationMessage = */TEXT("landscape.EditLayersLocalMerge.MaxComponentsPerWeightmapResolveBatch is deprecated. This was only used in Local merge. Batched merge (landscape.EditLayersLocalMerge.Enable 2) is the only merge technique available now."));
 
 TAutoConsoleVariable<int32> CVarLandscapeBatchedMergeMaxResolutionPerRenderBatch(
 	TEXT("landscape.BatchedMerge.MaxResolutionPerRenderBatch"),
@@ -2450,7 +2448,6 @@ ENUM_CLASS_FLAGS(FUpdateLayersContentContext::ERefreshFlags);
 struct FEditLayersHeightmapMergeParams
 {
 	int32 HeightmapUpdateModes = 0;
-	bool bForceRender = false;
 	bool bSkipBrush = false;
 };
 
@@ -4974,10 +4971,9 @@ int32 ALandscape::RegenerateLayersHeightmaps(const FUpdateLayersContentContext&
 	ULandscapeInfo* Info = GetLandscapeInfo();
 
 	const int32 HeightmapUpdateModes = LayerContentUpdateModes & ELandscapeLayerUpdateMode::Update_Heightmap_Types;
-	const bool bForceRender = CVarForceLayersUpdate.GetValueOnAnyThread() != 0;
 	const bool bSkipBrush = CVarLandscapeLayerBrushOptim.GetValueOnAnyThread() == 1 && (HeightmapUpdateModes == ELandscapeLayerUpdateMode::Update_Heightmap_Editing);
 
-	if ((HeightmapUpdateModes == 0 && !bForceRender) || Info == nullptr)
+	if ((HeightmapUpdateModes == 0) || Info == nullptr)
 	{
 		return 0;
 	}
@@ -5008,14 +5004,13 @@ int32 ALandscape::RegenerateLayersHeightmaps(const FUpdateLayersContentContext&
 		}
 	}
 
-	if (HeightmapUpdateModes || bForceRender)
+	if (HeightmapUpdateModes)
 	{
 		RenderCaptureInterface::FScopedCapture RenderCapture((RenderCaptureLayersNextHeightmapDraws != 0), TEXT("LandscapeLayersHeightmapCapture"));
 		RenderCaptureLayersNextHeightmapDraws = FMath::Max(0, RenderCaptureLayersNextHeightmapDraws - 1);
 
 		FEditLayersHeightmapMergeParams MergeParams;
 		MergeParams.HeightmapUpdateModes = HeightmapUpdateModes;
-		MergeParams.bForceRender = bForceRender;
 		MergeParams.bSkipBrush = bSkipBrush;
 
 		return PerformLayersHeightmapsBatchedMerge(InUpdateLayersContentContext, MergeParams);
@@ -5917,7 +5912,6 @@ void ALandscape::ReallocateLayersWeightmaps(FUpdateLayersContentContext& InUpdat
 struct FEditLayersWeightmapMergeParams
 {
 	int32 WeightmapUpdateModes;
-	bool bForceRender;
 	bool bSkipBrush;
 };
 
@@ -6424,11 +6418,10 @@ int32 ALandscape::RegenerateLayersWeightmaps(FUpdateLayersContentContext& InUpda
 	TRACE_CPUPROFILER_EVENT_SCOPE(LandscapeLayers_RegenerateLayersWeightmaps);
 	const int32 WeightmapUpdateModes = LayerContentUpdateModes & ELandscapeLayerUpdateMode::Update_Weightmap_Types;
 	const bool bSkipBrush = CVarLandscapeLayerBrushOptim.GetValueOnAnyThread() == 1 && (WeightmapUpdateModes == ELandscapeLayerUpdateMode::Update_Weightmap_Editing);
-	const bool bForceRender = CVarForceLayersUpdate.GetValueOnAnyThread() != 0;
 
 	ULandscapeInfo* Info = GetLandscapeInfo();
 
-	if (WeightmapUpdateModes == 0 && !bForceRender)
+	if (WeightmapUpdateModes == 0)
 	{
 		return 0;
 	}
@@ -6438,14 +6431,13 @@ int32 ALandscape::RegenerateLayersWeightmaps(FUpdateLayersContentContext& InUpda
 		return WeightmapUpdateModes;
 	}
 
-	if (WeightmapUpdateModes || bForceRender)
+	if (WeightmapUpdateModes)
 	{
 		RenderCaptureInterface::FScopedCapture RenderCapture((RenderCaptureLayersNextWeightmapDraws != 0), TEXT("LandscapeLayersWeightmapCapture"));
 		RenderCaptureLayersNextWeightmapDraws = FMath::Max(0, RenderCaptureLayersNextWeightmapDraws - 1);
 
 		FEditLayersWeightmapMergeParams MergeParams;
 		MergeParams.WeightmapUpdateModes = WeightmapUpdateModes;
-		MergeParams.bForceRender = bForceRender;
 		MergeParams.bSkipBrush = bSkipBrush;
 
 		return PerformLayersWeightmapsBatchedMerge(InUpdateLayersContentContext, MergeParams);
@@ -7343,7 +7335,6 @@ void ALandscape::UpdateLayersContent(bool bInWaitForStreaming, bool bInSkipMonit
 	}
 
 	const bool bProcessReadbacks = FLandscapeEditLayerReadback::HasWork();
-	const bool bForceRender = CVarForceLayersUpdate.GetValueOnAnyThread() != 0;
 
 	// User triggered change has been completely processed, resetting user triggered flag on all components.
 	if(IsUpToDate())
@@ -7359,10 +7350,7 @@ void ALandscape::UpdateLayersContent(bool bInWaitForStreaming, bool bInSkipMonit
 			}
 		);
 
-		if (!bForceRender)
-		{
-			return;
-		}
+		return;
 	}
 
 	bResourcesReady &= PrepareLayersTextureResources(bInWaitForStreaming);
@@ -7429,13 +7417,13 @@ void ALandscape::UpdateLayersContent(bool bInWaitForStreaming, bool bInSkipMonit
 		return;
 	}
 
-	if (LayerContentUpdateModes == 0 && !bForceRender)
+	if (LayerContentUpdateModes == 0)
 	{
 		return;
 	}
 
 	bool bUpdateAll = LayerContentUpdateModes & Update_All;
-	bool bPartialUpdate = !bForceRender && !bUpdateAll && CVarLandscapeLayerOptim.GetValueOnAnyThread() == 1;
+	bool bPartialUpdate = !bUpdateAll && CVarLandscapeLayerOptim.GetValueOnAnyThread() == 1;
 
 	FUpdateLayersContentContext UpdateLayersContentContext(MapHelper, bPartialUpdate);
 
