Mobile grass runtime generation fixes

#rnx

diff --git a/Engine/Source/Runtime/Landscape/Private/LandscapeRender.cpp b/Engine/Source/Runtime/Landscape/Private/LandscapeRender.cpp
--- a/Engine/Source/Runtime/Landscape/Private/LandscapeRender.cpp
+++ b/Engine/Source/Runtime/Landscape/Private/LandscapeRender.cpp
@@ -359,6 +359,9 @@ namespace UE::Landscape
 		static const bool bWaterPluginLoaded = IConsoleManager::Get().FindConsoleVariable(TEXT("r.Water.WaterInfo.RenderMethod")) != nullptr;
 		bNeedsFixedGridVertexFactory |= bWaterPluginLoaded;
 
+		// We need the fixed grid vertex factory for runtime grass generation.
+		bNeedsFixedGridVertexFactory |= ShouldBuildGrassMapRenderingResources();
+
 		return bNeedsFixedGridVertexFactory;
 	}
 
@@ -1766,10 +1769,6 @@ void FLandscapeComponentSceneProxy::CreateRenderThreadResources(FRHICommandListB
 		SharedBuffers->VertexFactory = LandscapeVertexFactory;
 
 		bool bNeedsFixedGridVertexFactory = NeedsFixedGridVertexFactory(GetScene().GetShaderPlatform());
-
-		// We also need the fixed vertex factor for grass/physical materials : 
-		bNeedsFixedGridVertexFactory |= (SharedBuffers->GrassIndexBuffer != nullptr);
-
 		if (bNeedsFixedGridVertexFactory)
 		{
 			FLandscapeFixedGridVertexFactory* FixedLandscapeVertexFactory = new FLandscapeFixedGridVertexFactory(FeatureLevel);
