Little code cleanup for FLandscapeTextureDataInfo : the call to BlockOnAnyAsyncBuild prior to Modify has to remain, unfortunately, because TLandscapeEditCache is using FLandscapeTextureDataInfo in such a way that it's very hard in practice to avoid calling Modify on textures accessed in readonly mode only in landscape tools (most of those usages lead to a write anyway).

#Synced-CL 48954958

diff --git a/Engine/Source/Runtime/Landscape/Private/LandscapeEditInterface.cpp b/Engine/Source/Runtime/Landscape/Private/LandscapeEditInterface.cpp
--- a/Engine/Source/Runtime/Landscape/Private/LandscapeEditInterface.cpp
+++ b/Engine/Source/Runtime/Landscape/Private/LandscapeEditInterface.cpp
@@ -335,7 +335,7 @@ void FLandscapeEditDataInterface::SetHeightData(int32 X1, int32 Y1, int32 X2, in
 			Component->Modify(GetShouldDirtyPackage());
 
 			FLandscapeTextureDataInfo* TexDataInfo = GetTextureDataInfo(Heightmap);
-			FColor* HeightmapTextureData = (FColor*)TexDataInfo->GetMipData(0);
+			FColor* HeightmapTextureData = (FColor*)TexDataInfo->GetMipData(/*InMipIndex = */0);
 
 			// Find the texture data corresponding to this vertex
 			int32 SizeU = Heightmap->Source.GetSizeX();
@@ -454,7 +454,7 @@ void FLandscapeEditDataInterface::RecalculateNormals()
 }
 
 template<typename TStoreData>
-void FLandscapeEditDataInterface::GetHeightDataTemplFast(const int32 X1, const int32 Y1, const int32 X2, const int32 Y2, TStoreData& StoreData, UTexture2D* InHeightmap, TStoreData* NormalData /*= NULL*/)
+void FLandscapeEditDataInterface::GetHeightDataTemplFast(const int32 X1, const int32 Y1, const int32 X2, const int32 Y2, TStoreData& StoreData, UTexture2D* InHeightmap, TStoreData* NormalData /*= nullptr*/)
 {
 	if (!LandscapeInfo) return;
 	int32 ComponentIndexX1, ComponentIndexY1, ComponentIndexX2, ComponentIndexY2;
@@ -473,10 +473,10 @@ void FLandscapeEditDataInterface::GetHeightDataTemplFast(const int32 X1, const i
 
 			UTexture2D* Heightmap = InHeightmap != nullptr ? InHeightmap : Component->GetHeightmap(GetEditLayer());
 
-			FLandscapeTextureDataInfo* TexDataInfo = NULL;
-			FColor* HeightmapTextureData = NULL;
+			FLandscapeTextureDataInfo* TexDataInfo = nullptr;
+			FColor* HeightmapTextureData = nullptr;
 			TexDataInfo = GetTextureDataInfo(Heightmap);
-			HeightmapTextureData = (FColor*)TexDataInfo->GetMipData(0);
+			HeightmapTextureData = (FColor*)TexDataInfo->GetMipData(/*InMipIndex = */0);
 
 			// Find coordinates of box that lies inside component
 			int32 ComponentX1 = FMath::Clamp<int32>(X1-ComponentIndexX*ComponentSizeQuads, 0, ComponentSizeQuads);
@@ -813,14 +813,14 @@ void FLandscapeEditDataInterface::CalcMissingValues(const int32 X1, const int32
 	}
 }
 
-FColor& FLandscapeEditDataInterface::GetHeightMapColor(const ULandscapeComponent* Component, int32 TexU, int32 TexV, FColor* TextureData)
+const FColor& FLandscapeEditDataInterface::GetHeightMapColor(const ULandscapeComponent* Component, int32 TexU, int32 TexV, FColor* TextureData)
 {
 	UTexture2D* Heightmap = Component->GetHeightmap(GetEditLayer());
 	check(Component);
 	if (!TextureData)
 	{
 		FLandscapeTextureDataInfo* TexDataInfo = GetTextureDataInfo(Heightmap);
-		TextureData = (FColor*)TexDataInfo->GetMipData(0);	
+		TextureData = (FColor*)TexDataInfo->GetMipData(/*InMipIndex = */0);	
 	}
 
 	// All Heightmaps of component have the same texture size
@@ -833,23 +833,23 @@ FColor& FLandscapeEditDataInterface::GetHeightMapColor(const ULandscapeComponent
 
 	const int32 TexX = HeightmapOffsetX + TexU;
 	const int32 TexY = HeightmapOffsetY + TexV;
-	FColor& TexData = TextureData[ TexX + TexY * SizeU ];
+	const FColor& TexData = TextureData[ TexX + TexY * SizeU ];
 	return TexData;
 }
 
-uint16 FLandscapeEditDataInterface::GetHeightMapData(const ULandscapeComponent* Component, int32 TexU, int32 TexV, FColor* TextureData /*= NULL*/)
+uint16 FLandscapeEditDataInterface::GetHeightMapData(const ULandscapeComponent* Component, int32 TexU, int32 TexV, FColor* TextureData /*= nullptr*/)
 {
 	const FColor& TexData = GetHeightMapColor(Component, TexU, TexV, TextureData);
 	return ((((uint16)TexData.R) << 8) | TexData.G);
 }
 
-uint16 FLandscapeEditDataInterface::GetHeightMapAlphaBlendData(const ULandscapeComponent* Component, int32 TexU, int32 TexV, FColor* TextureData /*= NULL*/)
+uint16 FLandscapeEditDataInterface::GetHeightMapAlphaBlendData(const ULandscapeComponent* Component, int32 TexU, int32 TexV, FColor* TextureData /*= nullptr*/)
 {
 	const FColor& TexData = GetHeightMapColor(Component, TexU, TexV, TextureData);
 	return ((((uint16)TexData.B) << 8) | TexData.A) & 0xFFFC;
 }
 
-uint8 FLandscapeEditDataInterface::GetHeightMapFlagsData(const ULandscapeComponent* Component, int32 TexU, int32 TexV, FColor* TextureData /*= NULL*/)
+uint8 FLandscapeEditDataInterface::GetHeightMapFlagsData(const ULandscapeComponent* Component, int32 TexU, int32 TexV, FColor* TextureData /*= nullptr*/)
 {
 	const FColor& TexData = GetHeightMapColor(Component, TexU, TexV, TextureData);
 	return TexData.A & 0x3;
@@ -889,18 +889,18 @@ void FLandscapeEditDataInterface::GetHeightDataInternal(int32& ValidX1, int32& V
 	{
 		NoBorderX1 = false;
 		NoBorderX2 = false;
-		BorderComponent[0] = BorderComponent[1] = NULL;
+		BorderComponent[0] = BorderComponent[1] = nullptr;
 		for( int32 ComponentIndexX=ComponentIndexX1;ComponentIndexX<=ComponentIndexX2;ComponentIndexX++ )
 		{
-			BorderComponent[2] = BorderComponent[3] = NULL;
+			BorderComponent[2] = BorderComponent[3] = nullptr;
 			int32 ComponentIndexXY = ComponentSizeX*(ComponentIndexY-ComponentIndexY1) + ComponentIndexX-ComponentIndexX1;
 			int32 ComponentIndexXX = ComponentIndexX - ComponentIndexX1;
 			int32 ComponentIndexYY = ComponentIndexY - ComponentIndexY1;
 			ComponentDataExist[ComponentIndexXY] = false;
 			ULandscapeComponent* Component = LandscapeInfo->XYtoComponentMap.FindRef(FIntPoint(ComponentIndexX,ComponentIndexY));
 
-			FLandscapeTextureDataInfo* TexDataInfo = NULL;
-			FColor* HeightmapTextureData = NULL;
+			FLandscapeTextureDataInfo* TexDataInfo = nullptr;
+			FColor* HeightmapTextureData = nullptr;
 			uint8 CornerSet = 0;
 			bool ExistLeft = ComponentIndexXX > 0 && ComponentDataExist[ ComponentIndexXX-1 + ComponentIndexYY * ComponentSizeX ];
 			bool ExistUp = ComponentIndexYY > 0 && ComponentDataExist[ ComponentIndexXX + (ComponentIndexYY-1) * ComponentSizeX ];
@@ -908,7 +908,7 @@ void FLandscapeEditDataInterface::GetHeightDataInternal(int32& ValidX1, int32& V
 			if( Component )
 			{
 				TexDataInfo = GetTextureDataInfo(Component->GetHeightmap(GetEditLayer()));
-				HeightmapTextureData = (FColor*)TexDataInfo->GetMipData(0);
+				HeightmapTextureData = (FColor*)TexDataInfo->GetMipData(/*InMipIndex = */0);
 				ComponentDataExist[ComponentIndexXY] = true;
 				// Update valid region
 				ValidX1 = FMath::Min<int32>(Component->GetSectionBase().X, ValidX1);
@@ -945,7 +945,7 @@ void FLandscapeEditDataInterface::GetHeightDataInternal(int32& ValidX1, int32& V
 						{
 							NoBorderX1 = false;
 							NeighborTexDataInfo[0] = GetTextureDataInfo(BorderComponent[0]->GetHeightmap(GetEditLayer()));
-							NeighborHeightmapTextureData[0] = (FColor*)NeighborTexDataInfo[0]->GetMipData(0);
+							NeighborHeightmapTextureData[0] = (FColor*)NeighborTexDataInfo[0]->GetMipData(/*InMipIndex = */0);
 							break;
 						}
 					}
@@ -961,7 +961,7 @@ void FLandscapeEditDataInterface::GetHeightDataInternal(int32& ValidX1, int32& V
 						{
 							NoBorderX2 = false;
 							NeighborTexDataInfo[1] = GetTextureDataInfo(BorderComponent[1]->GetHeightmap(GetEditLayer()));
-							NeighborHeightmapTextureData[1] = (FColor*)NeighborTexDataInfo[1]->GetMipData(0);
+							NeighborHeightmapTextureData[1] = (FColor*)NeighborTexDataInfo[1]->GetMipData(/*InMipIndex = */0);
 							break;
 						}
 					}
@@ -977,7 +977,7 @@ void FLandscapeEditDataInterface::GetHeightDataInternal(int32& ValidX1, int32& V
 						{
 							NoBorderY1[ComponentIndexXX] = false;
 							NeighborTexDataInfo[2] = GetTextureDataInfo(BorderComponent[2]->GetHeightmap(GetEditLayer()));
-							NeighborHeightmapTextureData[2] = (FColor*)NeighborTexDataInfo[2]->GetMipData(0);
+							NeighborHeightmapTextureData[2] = (FColor*)NeighborTexDataInfo[2]->GetMipData(/*InMipIndex = */0);
 							break;
 						}
 					}
@@ -988,7 +988,7 @@ void FLandscapeEditDataInterface::GetHeightDataInternal(int32& ValidX1, int32& V
 					if (BorderComponent[2])
 					{
 						NeighborTexDataInfo[2] = GetTextureDataInfo(BorderComponent[2]->GetHeightmap(GetEditLayer()));
-						NeighborHeightmapTextureData[2] = (FColor*)NeighborTexDataInfo[2]->GetMipData(0);
+						NeighborHeightmapTextureData[2] = (FColor*)NeighborTexDataInfo[2]->GetMipData(/*InMipIndex = */0);
 					}
 				}
 				// Search for bottom-closest component
@@ -1002,7 +1002,7 @@ void FLandscapeEditDataInterface::GetHeightDataInternal(int32& ValidX1, int32& V
 						{
 							NoBorderY2[ComponentIndexXX] = false;
 							NeighborTexDataInfo[3] = GetTextureDataInfo(BorderComponent[3]->GetHeightmap(GetEditLayer()));
-							NeighborHeightmapTextureData[3] = (FColor*)NeighborTexDataInfo[3]->GetMipData(0);
+							NeighborHeightmapTextureData[3] = (FColor*)NeighborTexDataInfo[3]->GetMipData(/*InMipIndex = */0);
 							break;
 						}
 					}
@@ -1013,18 +1013,18 @@ void FLandscapeEditDataInterface::GetHeightDataInternal(int32& ValidX1, int32& V
 					if (BorderComponent[3])
 					{
 						NeighborTexDataInfo[3] = GetTextureDataInfo(BorderComponent[3]->GetHeightmap(GetEditLayer()));
-						NeighborHeightmapTextureData[3] = (FColor*)NeighborTexDataInfo[3]->GetMipData(0);
+						NeighborHeightmapTextureData[3] = (FColor*)NeighborTexDataInfo[3]->GetMipData(/*InMipIndex = */0);
 					}
 				}
 
 				CornerComponent[0] = ComponentIndexX >= ComponentIndexX1 && ComponentIndexY >= ComponentIndexY1 ?
-					LandscapeInfo->XYtoComponentMap.FindRef(FIntPoint((ComponentIndexX-1),(ComponentIndexY-1))) : NULL;
+					LandscapeInfo->XYtoComponentMap.FindRef(FIntPoint((ComponentIndexX-1),(ComponentIndexY-1))) : nullptr;
 				CornerComponent[1] = ComponentIndexX <= ComponentIndexX2 && ComponentIndexY >= ComponentIndexY1 ?
-					LandscapeInfo->XYtoComponentMap.FindRef(FIntPoint((ComponentIndexX+1),(ComponentIndexY-1))) : NULL;
+					LandscapeInfo->XYtoComponentMap.FindRef(FIntPoint((ComponentIndexX+1),(ComponentIndexY-1))) : nullptr;
 				CornerComponent[2] = ComponentIndexX >= ComponentIndexX1 && ComponentIndexY <= ComponentIndexY2 ?
-					LandscapeInfo->XYtoComponentMap.FindRef(FIntPoint((ComponentIndexX-1),(ComponentIndexY+1))) : NULL;
+					LandscapeInfo->XYtoComponentMap.FindRef(FIntPoint((ComponentIndexX-1),(ComponentIndexY+1))) : nullptr;
 				CornerComponent[3] = ComponentIndexX <= ComponentIndexX2 && ComponentIndexY <= ComponentIndexY2 ?
-					LandscapeInfo->XYtoComponentMap.FindRef(FIntPoint((ComponentIndexX+1),(ComponentIndexY+1))) : NULL;
+					LandscapeInfo->XYtoComponentMap.FindRef(FIntPoint((ComponentIndexX+1),(ComponentIndexY+1))) : nullptr;
 
 				if (CornerComponent[0])
 				{
@@ -1405,7 +1405,7 @@ void FLandscapeEditDataInterface::GetHeightFlagsData(int32& X1, int32& Y1, int32
 	GetHeightFlagsDataTempl(X1, Y1, X2, Y2, ArrayStoreData);
 }
 
-void FLandscapeEditDataInterface::GetHeightDataFast(const int32 X1, const int32 Y1, const int32 X2, const int32 Y2, uint16* Data, int32 Stride, uint16* NormalData /*= NULL*/, UTexture2D* InHeightmap)
+void FLandscapeEditDataInterface::GetHeightDataFast(const int32 X1, const int32 Y1, const int32 X2, const int32 Y2, uint16* Data, int32 Stride, uint16* NormalData /*= nullptr*/, UTexture2D* InHeightmap)
 {
 	if( Stride==0 )
 	{
@@ -1430,7 +1430,7 @@ void FLandscapeEditDataInterface::GetHeightData(int32& X1, int32& Y1, int32& X2,
 	GetHeightDataTempl(X1, Y1, X2, Y2, SparseStoreData);
 }
 
-void FLandscapeEditDataInterface::GetHeightDataFast(const int32 X1, const int32 Y1, const int32 X2, const int32 Y2, TMap<FIntPoint, uint16>& SparseData, TMap<FIntPoint, uint16>* NormalData /*= NULL*/, UTexture2D* InHeightmap)
+void FLandscapeEditDataInterface::GetHeightDataFast(const int32 X1, const int32 Y1, const int32 X2, const int32 Y2, TMap<FIntPoint, uint16>& SparseData, TMap<FIntPoint, uint16>* NormalData /*= nullptr*/, UTexture2D* InHeightmap)
 {
 	TSparseDataAccess<uint16> SparseStoreData(SparseData);
 	if (NormalData)
@@ -1572,7 +1572,7 @@ void ULandscapeComponent::FillLayer(ULandscapeLayerInfoObject* LayerInfo, FLands
 		const int32 WeightmapOffsetY = static_cast<int32>(Component->WeightmapScaleBias.W * SizeV);
 
 		FWeightmapLayerAllocationInfo& FillLayerAllocation = ComponentWeightmapLayerAllocations[FillLayerIdx];
-		uint8* MaterialLayerData = (uint8*)LandscapeEdit.GetTextureDataInfo(ComponentWeightmapTextures[FillLayerAllocation.WeightmapTextureIndex])->GetMipData(0);
+		uint8* MaterialLayerData = (uint8*)LandscapeEdit.GetTextureDataInfo(ComponentWeightmapTextures[FillLayerAllocation.WeightmapTextureIndex])->GetMipData(/*InMipIndex = */0);
 
 		for (int32 Y = 0; Y < SizeV; ++Y)
 		{
@@ -1585,10 +1585,6 @@ void ULandscapeComponent::FillLayer(ULandscapeLayerInfoObject* LayerInfo, FLands
 		}
 	}
 
-	// todo - normalize texture usage: it's possible to end up with two textures each using one channel at this point
-	// e.g. if you start with 4 blended layers (in one texture) and a non-blended layer (in a 2nd), and fill one weight-blended layer (deleting the other three)
-	// this can also happen with normal painting I believe
-
 	// update mips if this texture isn't used in the layer compositing 
 	if (bIsFinalWeightmap)
 	{
@@ -1761,8 +1757,8 @@ void ULandscapeComponent::ReplaceLayerInternal(ULandscapeLayerInfoObject* FromLa
 		check(FromTexDataInfo->GetMipSizeX(0) == ToTexDataInfo->GetMipSizeX(0));
 		const int32 MipSize = FromTexDataInfo->GetMipSizeX(0);
 
-		uint8* const SrcTextureData = (uint8*)FromTexDataInfo->GetMipData(0) + ChannelOffsets[FromLayerAllocation.WeightmapTextureChannel];
-		uint8* const DestTextureData = (uint8*)ToTexDataInfo->GetMipData(0) + ChannelOffsets[ToLayerAllocation.WeightmapTextureChannel];
+		uint8* const SrcTextureData = (uint8*)FromTexDataInfo->GetMipData(/*InMipIndex = */0) + ChannelOffsets[FromLayerAllocation.WeightmapTextureChannel];
+		uint8* const DestTextureData = (uint8*)ToTexDataInfo->GetMipData(/*InMipIndex = */0) + ChannelOffsets[ToLayerAllocation.WeightmapTextureChannel];
 
 		for (int32 i = 0; i < FMath::Square(MipSize); i++)
 		{
@@ -2020,7 +2016,7 @@ namespace
 	{
 		// Preserve old Z value
 		FVector* PrevValue = SparseData.Find(FIntPoint(LandscapeX,LandscapeY));
-		if (PrevValue != NULL)
+		if (PrevValue != nullptr)
 		{
 			PrevValue->X = Offset.X;
 			PrevValue->Y = Offset.Y;
@@ -2140,8 +2136,8 @@ void FLandscapeEditDataInterface::SetAlphaData(ULandscapeLayerInfoObject* const
 			FIntPoint ComponentKey(ComponentIndexX,ComponentIndexY);
 			ULandscapeComponent* Component = LandscapeInfo->XYtoComponentMap.FindRef(ComponentKey);
 
-			// if NULL, there is no component at this location
-			if (Component == NULL)
+			// if nullptr, there is no component at this location
+			if (Component == nullptr)
 			{
 				continue;
 			}
@@ -2256,7 +2252,7 @@ void FLandscapeEditDataInterface::SetAlphaData(ULandscapeLayerInfoObject* const
 
 				if (Allocation.LayerInfo != nullptr) // only take into account valid layer
 				{
-					LayerDataPtrs[LayerIdx] = (uint8*)TexDataInfos[Allocation.WeightmapTextureIndex]->GetMipData(0) + ChannelOffsets[Allocation.WeightmapTextureChannel];
+					LayerDataPtrs[LayerIdx] = (uint8*)TexDataInfos[Allocation.WeightmapTextureIndex]->GetMipData(/*InMipIndex = */0) + ChannelOffsets[Allocation.WeightmapTextureChannel];
 					LayerEditDataAllZero[LayerIdx] = true;
 				}
 				else
@@ -2341,7 +2337,7 @@ void FLandscapeEditDataInterface::SetAlphaData(ULandscapeLayerInfoObject* const
 					const int32 TexY2 = (SubsectionSizeQuads+1) * SubIndexY + SubY2;
 					for (int32 WeightmapIdx = 0; WeightmapIdx < ComponentWeightmapTextures.Num(); WeightmapIdx++)
 					{
-						if (TexDataInfos[WeightmapIdx] != NULL)
+						if (TexDataInfos[WeightmapIdx] != nullptr)
 						{
 							TexDataInfos[WeightmapIdx]->AddMipUpdateRegion(0,TexX1,TexY1,TexX2,TexY2);
 						}
@@ -2434,8 +2430,8 @@ void FLandscapeEditDataInterface::SetAlphaData(const TSet<ULandscapeLayerInfoObj
 			FIntPoint ComponentKey(ComponentIndexX,ComponentIndexY);
 			ULandscapeComponent* Component = LandscapeInfo->XYtoComponentMap.FindRef(ComponentKey);
 
-			// if NULL, there is no component at this location
-			if (Component == NULL)
+			// if nullptr, there is no component at this location
+			if (Component == nullptr)
 			{
 				continue;
 			}
@@ -2506,7 +2502,7 @@ void FLandscapeEditDataInterface::SetAlphaData(const TSet<ULandscapeLayerInfoObj
 				const int32 LayerDataIdx = LandscapeInfo->GetLayerInfoIndex(ComponentWeightmapLayerAllocations[LayerIdx].LayerInfo);
 				check(LayerDataIdx != INDEX_NONE);
 				LayerDataInfos[LayerIdx].InDataPtr = Data + LayerDataIdx;
-				LayerDataInfos[LayerIdx].TexDataPtr = (uint8*)TexDataInfos[Allocation.WeightmapTextureIndex]->GetMipData(0) + ChannelOffsets[Allocation.WeightmapTextureChannel];
+				LayerDataInfos[LayerIdx].TexDataPtr = (uint8*)TexDataInfos[Allocation.WeightmapTextureIndex]->GetMipData(/*InMipIndex = */0) + ChannelOffsets[Allocation.WeightmapTextureChannel];
 				LayerEditDataAllZero[LayerIdx] = true;
 			}
 
@@ -2560,7 +2556,7 @@ void FLandscapeEditDataInterface::SetAlphaData(const TSet<ULandscapeLayerInfoObj
 								// this is equivalent to saying if (DirtyLayerInfos.Contains(Allocation.LayerInfo))
 								// which is what we really mean here, but this is quicker
 								// and I've lost count of the depth we've nested for loops at this point
-								if (LayerDataInfos[LayerIdx].TexDataPtr != NULL)
+								if (LayerDataInfos[LayerIdx].TexDataPtr != nullptr)
 								{
 									uint8& Weight = LayerDataInfos[LayerIdx].TexDataPtr[TexDataIndex];
 
@@ -2581,7 +2577,7 @@ void FLandscapeEditDataInterface::SetAlphaData(const TSet<ULandscapeLayerInfoObj
 					const int32 TexY2 = (SubsectionSizeQuads+1) * SubIndexY + SubY2;
 					for (int32 WeightmapIdx = 0; WeightmapIdx < ComponentWeightmapTextures.Num(); WeightmapIdx++)
 					{
-						if (TexDataInfos[WeightmapIdx] != NULL)
+						if (TexDataInfos[WeightmapIdx] != nullptr)
 						{
 							TexDataInfos[WeightmapIdx]->AddMipUpdateRegion(0,TexX1,TexY1,TexX2,TexY2);
 						}
@@ -2635,16 +2631,16 @@ void FLandscapeEditDataInterface::GetWeightDataTemplFast(ULandscapeLayerInfoObje
 				continue;
 			}
 
-			UTexture2D* WeightmapTexture = NULL;
-			FLandscapeTextureDataInfo* TexDataInfo = NULL;
-			uint8* WeightmapTextureData = NULL;
+			UTexture2D* WeightmapTexture = nullptr;
+			FLandscapeTextureDataInfo* TexDataInfo = nullptr;
+			uint8* WeightmapTextureData = nullptr;
 			uint8 WeightmapChannelOffset = 0;
 			TArray<FLandscapeTextureDataInfo*> TexDataInfos; // added for whole weight case...
 
 			const TArray<FWeightmapLayerAllocationInfo>& ComponentWeightmapLayerAllocations = Component->GetWeightmapLayerAllocations(GetEditLayer());
 			const TArray<UTexture2D*>& ComponentWeightmapTextures = Component->GetWeightmapTextures(GetEditLayer());
 
-			if (LayerInfo != NULL)
+			if (LayerInfo != nullptr)
 			{
 				for( int32 LayerIdx=0;LayerIdx<ComponentWeightmapLayerAllocations.Num();LayerIdx++ )
 				{
@@ -2652,7 +2648,7 @@ void FLandscapeEditDataInterface::GetWeightDataTemplFast(ULandscapeLayerInfoObje
 					{
 						WeightmapTexture = ComponentWeightmapTextures[ComponentWeightmapLayerAllocations[LayerIdx].WeightmapTextureIndex];
 						TexDataInfo = GetTextureDataInfo(WeightmapTexture);
-						WeightmapTextureData = (uint8*)TexDataInfo->GetMipData(0);
+						WeightmapTextureData = (uint8*)TexDataInfo->GetMipData(/*InMipIndex = */0);
 						WeightmapChannelOffset = ChannelOffsets[ComponentWeightmapLayerAllocations[LayerIdx].WeightmapTextureChannel];
 						break;
 					}
@@ -2697,7 +2693,7 @@ void FLandscapeEditDataInterface::GetWeightDataTemplFast(ULandscapeLayerInfoObje
 							int32 LandscapeX = SubIndexX*SubsectionSizeQuads + ComponentIndexX*ComponentSizeQuads + SubX;
 							int32 LandscapeY = SubIndexY*SubsectionSizeQuads + ComponentIndexY*ComponentSizeQuads + SubY;
 
-							if (LayerInfo != NULL)
+							if (LayerInfo != nullptr)
 							{
 								// Find the input data corresponding to this vertex
 								uint8 Weight;
@@ -2728,7 +2724,7 @@ void FLandscapeEditDataInterface::GetWeightDataTemplFast(ULandscapeLayerInfoObje
 								{
 									int32 Idx = ComponentWeightmapLayerAllocations[LayerIdx].WeightmapTextureIndex;
 									UTexture2D* ComponentWeightmapTexture = ComponentWeightmapTextures[Idx];
-									uint8* ComponentWeightmapTextureData = (uint8*)TexDataInfos[Idx]->GetMipData(0);
+									uint8* ComponentWeightmapTextureData = (uint8*)TexDataInfos[Idx]->GetMipData(/*InMipIndex = */0);
 									uint8 ComponentWeightmapChannelOffset = ChannelOffsets[ComponentWeightmapLayerAllocations[LayerIdx].WeightmapTextureChannel];
 
 									// Find the texture data corresponding to this vertex
@@ -2760,12 +2756,12 @@ void FLandscapeEditDataInterface::GetWeightDataTemplFast(ULandscapeLayerInfoObje
 	}
 }
 
-uint8 FLandscapeEditDataInterface::GetWeightMapData(const ULandscapeComponent* Component, ULandscapeLayerInfoObject* LayerInfo, int32 TexU, int32 TexV, uint8 Offset /*= 0*/, UTexture2D* Texture /*= NULL*/, uint8* TextureData /*= NULL*/)
+uint8 FLandscapeEditDataInterface::GetWeightMapData(const ULandscapeComponent* Component, ULandscapeLayerInfoObject* LayerInfo, int32 TexU, int32 TexV, uint8 Offset /*= 0*/, UTexture2D* Texture /*= nullptr*/, uint8* TextureData /*= nullptr*/)
 {
 	check(Component);
 	if (!Texture || !TextureData)
 	{
-		if (LayerInfo != NULL)
+		if (LayerInfo != nullptr)
 		{
 			const TArray<FWeightmapLayerAllocationInfo>& ComponentWeightmapLayerAllocations = Component->GetWeightmapLayerAllocations(GetEditLayer());
 			const TArray<UTexture2D*>& ComponentWeightmapTextures = Component->GetWeightmapTextures(GetEditLayer());
@@ -2776,7 +2772,7 @@ uint8 FLandscapeEditDataInterface::GetWeightMapData(const ULandscapeComponent* C
 				{
 					Texture = ComponentWeightmapTextures[ComponentWeightmapLayerAllocations[LayerIdx].WeightmapTextureIndex];
 					FLandscapeTextureDataInfo* TexDataInfo = GetTextureDataInfo(Texture);
-					TextureData = (uint8*)TexDataInfo->GetMipData(0);
+					TextureData = (uint8*)TexDataInfo->GetMipData(/*InMipIndex = */0);
 					Offset = ChannelOffsets[ComponentWeightmapLayerAllocations[LayerIdx].WeightmapTextureChannel];
 					break;
 				}
@@ -2833,19 +2829,19 @@ void FLandscapeEditDataInterface::GetWeightDataTempl(ULandscapeLayerInfoObject*
 	{
 		NoBorderX1 = false;
 		NoBorderX2 = false;
-		BorderComponent[0] = BorderComponent[1] = NULL;
+		BorderComponent[0] = BorderComponent[1] = nullptr;
 		for( int32 ComponentIndexX=ComponentIndexX1;ComponentIndexX<=ComponentIndexX2;ComponentIndexX++ )
 		{	
-			BorderComponent[2] = BorderComponent[3] = NULL;
+			BorderComponent[2] = BorderComponent[3] = nullptr;
 			int32 ComponentIndexXY = ComponentSizeX*(ComponentIndexY-ComponentIndexY1) + ComponentIndexX-ComponentIndexX1;
 			int32 ComponentIndexXX = ComponentIndexX - ComponentIndexX1;
 			int32 ComponentIndexYY = ComponentIndexY - ComponentIndexY1;
 			ComponentDataExist[ComponentIndexXY] = false;
 			ULandscapeComponent* Component = LandscapeInfo->XYtoComponentMap.FindRef(FIntPoint(ComponentIndexX,ComponentIndexY));
 
-			UTexture2D* WeightmapTexture = NULL;
-			FLandscapeTextureDataInfo* TexDataInfo = NULL;
-			uint8* WeightmapTextureData = NULL;
+			UTexture2D* WeightmapTexture = nullptr;
+			FLandscapeTextureDataInfo* TexDataInfo = nullptr;
+			uint8* WeightmapTextureData = nullptr;
 			uint8 WeightmapChannelOffset = 0;
 			TArray<FLandscapeTextureDataInfo*> TexDataInfos; // added for whole weight case...
 			uint8 CornerSet = 0;
@@ -2857,7 +2853,7 @@ void FLandscapeEditDataInterface::GetWeightDataTempl(ULandscapeLayerInfoObject*
 			{
 				const TArray<UTexture2D*>& ComponentWeightmapTextures = Component->GetWeightmapTextures(GetEditLayer());
 
-				if (LayerInfo != NULL)
+				if (LayerInfo != nullptr)
 				{
 					const TArray<FWeightmapLayerAllocationInfo>& ComponentWeightmapLayerAllocations = Component->GetWeightmapLayerAllocations(GetEditLayer());
 
@@ -2867,7 +2863,7 @@ void FLandscapeEditDataInterface::GetWeightDataTempl(ULandscapeLayerInfoObject*
 						{
 							WeightmapTexture = ComponentWeightmapTextures[ComponentWeightmapLayerAllocations[LayerIdx].WeightmapTextureIndex];
 							TexDataInfo = GetTextureDataInfo(WeightmapTexture);
-							WeightmapTextureData = (uint8*)TexDataInfo->GetMipData(0);
+							WeightmapTextureData = (uint8*)TexDataInfo->GetMipData(/*InMipIndex = */0);
 							WeightmapChannelOffset = ChannelOffsets[ComponentWeightmapLayerAllocations[LayerIdx].WeightmapTextureChannel];
 							break;
 						}
@@ -2924,7 +2920,7 @@ void FLandscapeEditDataInterface::GetWeightDataTempl(ULandscapeLayerInfoObject*
 						if (BorderComponent[0])
 						{
 							NoBorderX1 = false;
-							if (LayerInfo != NULL)
+							if (LayerInfo != nullptr)
 							{
 								const TArray<FWeightmapLayerAllocationInfo>& BorderWeightmapLayerAllocations = BorderComponent[0]->GetWeightmapLayerAllocations(GetEditLayer());
 								const TArray<UTexture2D*>& BorderWeightmapTextures = BorderComponent[0]->GetWeightmapTextures(GetEditLayer());
@@ -2935,7 +2931,7 @@ void FLandscapeEditDataInterface::GetWeightDataTempl(ULandscapeLayerInfoObject*
 									{
 										NeighborWeightmapTexture[0] = BorderWeightmapTextures[BorderWeightmapLayerAllocations[LayerIdx].WeightmapTextureIndex];
 										NeighborTexDataInfo[0] = GetTextureDataInfo(NeighborWeightmapTexture[0]);
-										NeighborWeightmapTextureData[0] = (uint8*)NeighborTexDataInfo[0]->GetMipData(0);
+										NeighborWeightmapTextureData[0] = (uint8*)NeighborTexDataInfo[0]->GetMipData(/*InMipIndex = */0);
 										NeighborWeightmapChannelOffset[0] = ChannelOffsets[BorderWeightmapLayerAllocations[LayerIdx].WeightmapTextureChannel];
 										break;
 									}
@@ -2955,7 +2951,7 @@ void FLandscapeEditDataInterface::GetWeightDataTempl(ULandscapeLayerInfoObject*
 						if (BorderComponent[1])
 						{
 							NoBorderX2 = false;
-							if (LayerInfo != NULL)
+							if (LayerInfo != nullptr)
 							{
 								const TArray<FWeightmapLayerAllocationInfo>& BorderWeightmapLayerAllocations = BorderComponent[1]->GetWeightmapLayerAllocations(GetEditLayer());
 								const TArray<UTexture2D*>& BorderWeightmapTextures = BorderComponent[1]->GetWeightmapTextures(GetEditLayer());
@@ -2966,7 +2962,7 @@ void FLandscapeEditDataInterface::GetWeightDataTempl(ULandscapeLayerInfoObject*
 									{
 										NeighborWeightmapTexture[1] = BorderWeightmapTextures[BorderWeightmapLayerAllocations[LayerIdx].WeightmapTextureIndex];
 										NeighborTexDataInfo[1] = GetTextureDataInfo(NeighborWeightmapTexture[1]);
-										NeighborWeightmapTextureData[1] = (uint8*)NeighborTexDataInfo[1]->GetMipData(0);
+										NeighborWeightmapTextureData[1] = (uint8*)NeighborTexDataInfo[1]->GetMipData(/*InMipIndex = */0);
 										NeighborWeightmapChannelOffset[1] = ChannelOffsets[BorderWeightmapLayerAllocations[LayerIdx].WeightmapTextureChannel];
 										break;
 									}
@@ -2985,7 +2981,7 @@ void FLandscapeEditDataInterface::GetWeightDataTempl(ULandscapeLayerInfoObject*
 						if (BorderComponent[2])
 						{
 							NoBorderY1[ComponentIndexXX] = false;
-							if (LayerInfo != NULL)
+							if (LayerInfo != nullptr)
 							{
 								const TArray<FWeightmapLayerAllocationInfo>& BorderWeightmapLayerAllocations = BorderComponent[2]->GetWeightmapLayerAllocations(GetEditLayer());
 								const TArray<UTexture2D*>& BorderWeightmapTextures = BorderComponent[2]->GetWeightmapTextures(GetEditLayer());
@@ -2996,7 +2992,7 @@ void FLandscapeEditDataInterface::GetWeightDataTempl(ULandscapeLayerInfoObject*
 									{
 										NeighborWeightmapTexture[2] = BorderWeightmapTextures[BorderWeightmapLayerAllocations[LayerIdx].WeightmapTextureIndex];
 										NeighborTexDataInfo[2] = GetTextureDataInfo(NeighborWeightmapTexture[2]);
-										NeighborWeightmapTextureData[2] = (uint8*)NeighborTexDataInfo[2]->GetMipData(0);
+										NeighborWeightmapTextureData[2] = (uint8*)NeighborTexDataInfo[2]->GetMipData(/*InMipIndex = */0);
 										NeighborWeightmapChannelOffset[2] = ChannelOffsets[BorderWeightmapLayerAllocations[LayerIdx].WeightmapTextureChannel];
 										break;
 									}
@@ -3010,7 +3006,7 @@ void FLandscapeEditDataInterface::GetWeightDataTempl(ULandscapeLayerInfoObject*
 					BorderComponent[2] = BorderComponentY1[ComponentIndexXX];
 					if (BorderComponent[2])
 					{
-						if (LayerInfo != NULL)
+						if (LayerInfo != nullptr)
 						{
 							const TArray<FWeightmapLayerAllocationInfo>& BorderWeightmapLayerAllocations = BorderComponent[2]->GetWeightmapLayerAllocations(GetEditLayer());
 							const TArray<UTexture2D*>& BorderWeightmapTextures = BorderComponent[2]->GetWeightmapTextures(GetEditLayer());
@@ -3021,7 +3017,7 @@ void FLandscapeEditDataInterface::GetWeightDataTempl(ULandscapeLayerInfoObject*
 								{
 									NeighborWeightmapTexture[2] = BorderWeightmapTextures[BorderWeightmapLayerAllocations[LayerIdx].WeightmapTextureIndex];
 									NeighborTexDataInfo[2] = GetTextureDataInfo(NeighborWeightmapTexture[2]);
-									NeighborWeightmapTextureData[2] = (uint8*)NeighborTexDataInfo[2]->GetMipData(0);
+									NeighborWeightmapTextureData[2] = (uint8*)NeighborTexDataInfo[2]->GetMipData(/*InMipIndex = */0);
 									NeighborWeightmapChannelOffset[2] = ChannelOffsets[BorderWeightmapLayerAllocations[LayerIdx].WeightmapTextureChannel];
 									break;
 								}
@@ -3039,7 +3035,7 @@ void FLandscapeEditDataInterface::GetWeightDataTempl(ULandscapeLayerInfoObject*
 						if (BorderComponent[3])
 						{
 							NoBorderY2[ComponentIndexXX] = false;
-							if (LayerInfo != NULL)
+							if (LayerInfo != nullptr)
 							{
 								const TArray<FWeightmapLayerAllocationInfo>& BorderWeightmapLayerAllocations = BorderComponent[3]->GetWeightmapLayerAllocations(GetEditLayer());
 								const TArray<UTexture2D*>& BorderWeightmapTextures = BorderComponent[3]->GetWeightmapTextures(GetEditLayer());
@@ -3050,7 +3046,7 @@ void FLandscapeEditDataInterface::GetWeightDataTempl(ULandscapeLayerInfoObject*
 									{
 										NeighborWeightmapTexture[3] = BorderWeightmapTextures[BorderWeightmapLayerAllocations[LayerIdx].WeightmapTextureIndex];
 										NeighborTexDataInfo[3] = GetTextureDataInfo(NeighborWeightmapTexture[3]);
-										NeighborWeightmapTextureData[3] = (uint8*)NeighborTexDataInfo[3]->GetMipData(0);
+										NeighborWeightmapTextureData[3] = (uint8*)NeighborTexDataInfo[3]->GetMipData(/*InMipIndex = */0);
 										NeighborWeightmapChannelOffset[3] = ChannelOffsets[BorderWeightmapLayerAllocations[LayerIdx].WeightmapTextureChannel];
 										break;
 									}
@@ -3065,7 +3061,7 @@ void FLandscapeEditDataInterface::GetWeightDataTempl(ULandscapeLayerInfoObject*
 					BorderComponent[3] = BorderComponentY2[ComponentIndexXX];
 					if (BorderComponent[3])
 					{
-						if (LayerInfo != NULL)
+						if (LayerInfo != nullptr)
 						{
 							const TArray<FWeightmapLayerAllocationInfo>& BorderWeightmapLayerAllocations = BorderComponent[3]->GetWeightmapLayerAllocations(GetEditLayer());
 							const TArray<UTexture2D*>& BorderWeightmapTextures = BorderComponent[3]->GetWeightmapTextures(GetEditLayer());
@@ -3076,7 +3072,7 @@ void FLandscapeEditDataInterface::GetWeightDataTempl(ULandscapeLayerInfoObject*
 								{
 									NeighborWeightmapTexture[3] = BorderWeightmapTextures[BorderWeightmapLayerAllocations[LayerIdx].WeightmapTextureIndex];
 									NeighborTexDataInfo[3] = GetTextureDataInfo(NeighborWeightmapTexture[3]);
-									NeighborWeightmapTextureData[3] = (uint8*)NeighborTexDataInfo[3]->GetMipData(0);
+									NeighborWeightmapTextureData[3] = (uint8*)NeighborTexDataInfo[3]->GetMipData(/*InMipIndex = */0);
 									NeighborWeightmapChannelOffset[3] = ChannelOffsets[BorderWeightmapLayerAllocations[LayerIdx].WeightmapTextureChannel];
 									break;
 								}
@@ -3086,13 +3082,13 @@ void FLandscapeEditDataInterface::GetWeightDataTempl(ULandscapeLayerInfoObject*
 				}
 
 				CornerComponent[0] = ComponentIndexX >= ComponentIndexX1 && ComponentIndexY >= ComponentIndexY1 ?
-					LandscapeInfo->XYtoComponentMap.FindRef(FIntPoint((ComponentIndexX-1),(ComponentIndexY-1))) : NULL;
+					LandscapeInfo->XYtoComponentMap.FindRef(FIntPoint((ComponentIndexX-1),(ComponentIndexY-1))) : nullptr;
 				CornerComponent[1] = ComponentIndexX <= ComponentIndexX2 && ComponentIndexY >= ComponentIndexY1 ?
-					LandscapeInfo->XYtoComponentMap.FindRef(FIntPoint((ComponentIndexX+1),(ComponentIndexY-1))) : NULL;
+					LandscapeInfo->XYtoComponentMap.FindRef(FIntPoint((ComponentIndexX+1),(ComponentIndexY-1))) : nullptr;
 				CornerComponent[2] = ComponentIndexX >= ComponentIndexX1 && ComponentIndexY <= ComponentIndexY2 ?
-					LandscapeInfo->XYtoComponentMap.FindRef(FIntPoint((ComponentIndexX-1),(ComponentIndexY+1))) : NULL;
+					LandscapeInfo->XYtoComponentMap.FindRef(FIntPoint((ComponentIndexX-1),(ComponentIndexY+1))) : nullptr;
 				CornerComponent[3] = ComponentIndexX <= ComponentIndexX2 && ComponentIndexY <= ComponentIndexY2 ?
-					LandscapeInfo->XYtoComponentMap.FindRef(FIntPoint((ComponentIndexX+1),(ComponentIndexY+1))) : NULL;
+					LandscapeInfo->XYtoComponentMap.FindRef(FIntPoint((ComponentIndexX+1),(ComponentIndexY+1))) : nullptr;
 
 				if (CornerComponent[0])
 				{
@@ -3212,7 +3208,7 @@ void FLandscapeEditDataInterface::GetWeightDataTempl(ULandscapeLayerInfoObject*
 							int32 LandscapeX = SubIndexX*SubsectionSizeQuads + ComponentIndexX*ComponentSizeQuads + SubX;
 							int32 LandscapeY = SubIndexY*SubsectionSizeQuads + ComponentIndexY*ComponentSizeQuads + SubY;
 
-							if (LayerInfo != NULL)
+							if (LayerInfo != nullptr)
 							{
 								// Find the input data corresponding to this vertex
 								uint8 Weight = 0;
@@ -3347,7 +3343,7 @@ void FLandscapeEditDataInterface::GetWeightDataTempl(ULandscapeLayerInfoObject*
 								{
 									int32 Idx = ComponentWeightmapLayerAllocations[LayerIdx].WeightmapTextureIndex;
 									UTexture2D* ComponentWeightmapTexture = ComponentWeightmapTextures[Idx];
-									uint8* ComponentWeightmapTextureData = (uint8*)TexDataInfos[Idx]->GetMipData(0);
+									uint8* ComponentWeightmapTextureData = (uint8*)TexDataInfos[Idx]->GetMipData(/*InMipIndex = */0);
 									uint8 ComponentWeightmapChannelOffset = ChannelOffsets[ComponentWeightmapLayerAllocations[LayerIdx].WeightmapTextureChannel];
 
 									// Find the texture data corresponding to this vertex
@@ -3455,7 +3451,7 @@ void FLandscapeEditDataInterface::GetWeightDataFast(ULandscapeLayerInfoObject* L
 FLandscapeTextureDataInfo* FLandscapeTextureDataInterface::GetTextureDataInfo(UTexture2D* Texture)
 {
 	FLandscapeTextureDataInfo* Result = TextureDataMap.FindRef(Texture);
-	if( !Result )
+	if (Result == nullptr)
 	{
 		Result = TextureDataMap.Add(Texture, new FLandscapeTextureDataInfo(Texture, GetShouldDirtyPackage()));
 	}
@@ -3638,7 +3634,7 @@ template<typename TData>
 bool FLandscapeTextureDataInterface::EqualTextureValueTempl(UTexture2D* Src, TData Value)
 {
 	FLandscapeTextureDataInfo* DestDataInfo = GetTextureDataInfo(Src);
-	TData* DestTextureData = (TData*)DestDataInfo->GetMipData(0);
+	TData* DestTextureData = (TData*)DestDataInfo->GetMipData(/*InMipIndex = */0);
 	int32 Size = Src->Source.GetSizeX() * Src->Source.GetSizeY();
 
 	for( int32 i = 0 ;i < Size; ++i )
@@ -3782,8 +3778,8 @@ void FLandscapeEditDataInterface::GetEditToolTextureData(const int32 X1, const i
 		{
 			ULandscapeComponent* Component = LandscapeInfo->XYtoComponentMap.FindRef(FIntPoint(ComponentIndexX, ComponentIndexY));
 
-			FLandscapeTextureDataInfo* TexDataInfo = NULL;
-			uint8* SelectTextureData = NULL;
+			FLandscapeTextureDataInfo* TexDataInfo = nullptr;
+			uint8* SelectTextureData = nullptr;
 			UTexture2D* EditToolTexture = nullptr;
 			if (Component)
 			{
@@ -3792,7 +3788,7 @@ void FLandscapeEditDataInterface::GetEditToolTextureData(const int32 X1, const i
 			if (EditToolTexture)
 			{
 				TexDataInfo = GetTextureDataInfo(EditToolTexture);
-				SelectTextureData = (uint8*)TexDataInfo->GetMipData(0);
+				SelectTextureData = (uint8*)TexDataInfo->GetMipData(/*InMipIndex = */0);
 			}
 
 			// Find coordinates of box that lies inside component
@@ -3871,14 +3867,14 @@ void FLandscapeEditDataInterface::SetEditToolTextureData(int32 X1, int32 Y1, int
 		{	
 			FIntPoint ComponentKey(ComponentIndexX,ComponentIndexY);
 			ULandscapeComponent* Component = LandscapeInfo->XYtoComponentMap.FindRef(ComponentKey);
-			if (Component == NULL)
+			if (Component == nullptr)
 			{
 				continue;
 			}
 
 			UTexture2D*& EditToolTexture = GetComponentTexture(Component);
-			// if NULL, it was painted away
-			if (EditToolTexture == NULL)
+			// if nullptr, it was painted away
+			if (EditToolTexture == nullptr)
 			{
 				// Construct Texture...
 				int32 WeightmapSize = (Component->SubsectionSizeQuads+1) * Component->NumSubsections;
@@ -3889,7 +3885,7 @@ void FLandscapeEditDataInterface::SetEditToolTextureData(int32 X1, int32 Y1, int
 			}
 
 			FLandscapeTextureDataInfo* TexDataInfo = GetTextureDataInfo(EditToolTexture);
-			uint8* SelectTextureData = (uint8*)TexDataInfo->GetMipData(0);
+			uint8* SelectTextureData = (uint8*)TexDataInfo->GetMipData(/*InMipIndex = */0);
 
 			// Find the texture data corresponding to this vertex
 			int32 SizeU = EditToolTexture->Source.GetSizeX();
@@ -3964,7 +3960,7 @@ void FLandscapeEditDataInterface::SetXYOffsetData(int32 X1, int32 Y1, int32 X2,
 }
 
 // Deprecated
-FVector2D FLandscapeEditDataInterface::GetXYOffsetmapData(const ULandscapeComponent* Component, int32 TexU, int32 TexV, FColor* TextureData/* = NULL*/)
+FVector2D FLandscapeEditDataInterface::GetXYOffsetmapData(const ULandscapeComponent* Component, int32 TexU, int32 TexV, FColor* TextureData/* = nullptr*/)
 {
 	return FVector2D();
 }
@@ -4034,18 +4030,55 @@ FLandscapeTextureDataInfo::FLandscapeTextureDataInfo(UTexture2D* InTexture, bool
 	, TextureUsage(InTextureUsage)
 	, TextureType(InTextureType)
 {
-	MipInfo.AddZeroed(Texture->Source.GetNumMips());
+	MipInfo.AddDefaulted(Texture->Source.GetNumMips());
 	Texture->SetFlags(RF_Transactional);
 	Texture->TemporarilyDisableStreaming();
 
-	// HACK [jonathan.bard] : call BlockOnAnyAsyncBuild here to make sure that Modify cannot cancel the potentially on-going async cache
+	// Call BlockOnAnyAsyncBuild here to make sure that Modify cannot cancel the potentially on-going async cache
 	//  The better fix would be not to call Modify when we don't intend to actually modify the texture (e.g. when FLandscapeTextureDataInfo is used for reading the data only)
+	//  but in practice, that would require rebuilding the entire TLandscapeEditCache system, that is used both as a read-only and/ read-write facility
 	Texture->BlockOnAnyAsyncBuild();
-
 	Texture->Modify(bShouldDirtyPackage);
 	check(!Texture->IsDefaultTexture()); // If we're the default texture, then Modify likely has cancelled the async cache, which would be a disaster if nothing requests a new cache afterwards
 }
 
+void FLandscapeTextureDataInfo::AddMipUpdateRegion(int32 InMipIndex, int32 InX1, int32 InY1, int32 InX2, int32 InY2)
+{
+	if (MipInfo[InMipIndex].bFull)
+	{
+		return;
+	}
+
+	check(InMipIndex < MipInfo.Num());
+	uint32 Width = 1 + InX2 - InX1;
+	uint32 Height = 1 + InY2 - InY1;
+	// Catch situation where we are updating the whole texture to avoid adding redundant regions once the whole region as been included.
+	if (Width == GetMipSizeX(InMipIndex) && Height == GetMipSizeY(InMipIndex))
+	{
+		MipInfo[InMipIndex].bFull = true;
+		MipInfo[InMipIndex].MipUpdateRegions.Reset();
+		// Push a full region for UpdateTextureData() to process later
+		MipInfo[InMipIndex].MipUpdateRegions.Emplace(0, 0, 0, 0, Width, Height);
+		return;
+	}
+
+	MipInfo[InMipIndex].MipUpdateRegions.Emplace(InX1, InY1, InX1, InY1, Width, Height);
+}
+
+void* FLandscapeTextureDataInfo::GetMipData(int32 InMipIndex)
+{
+	check(InMipIndex < MipInfo.Num());
+
+	if (!MipInfo[InMipIndex].MipData)
+	{
+		// will Unlock in destructor
+		MipInfo[InMipIndex].MipData = Texture->Source.LockMip(InMipIndex);
+		// probably should check that we got the lock :
+		//check( MipInfo[InMipIndex].MipData != nullptr );
+	}
+	return MipInfo[InMipIndex].MipData;
+}
+
 bool FLandscapeTextureDataInfo::UpdateTextureData()
 {
 	const bool bCompressed = !Texture->CompressionNone;
@@ -4060,9 +4093,9 @@ bool FLandscapeTextureDataInfo::UpdateTextureData()
 	// Only wait once
 	bool bNeedToFinishCompilation = true;
 
-	for (int32 i = 0; i < MipInfo.Num(); i++)
+	for (int32 MipIndex = 0; MipIndex < MipInfo.Num(); MipIndex++)
 	{
-		if (MipInfo[i].MipData && MipInfo[i].MipUpdateRegions.Num() > 0)
+		if (MipInfo[MipIndex].MipData && MipInfo[MipIndex].MipUpdateRegions.Num() > 0)
 		{
 			if (bCompressed)
 			{
@@ -4078,18 +4111,18 @@ bool FLandscapeTextureDataInfo::UpdateTextureData()
 				bNeedToFinishCompilation = false;
 			}
 
-			const uint32 SrcSizeX = (Texture->Source.GetSizeX()) >> i;
-			const uint32 SrcSizeY = (Texture->Source.GetSizeY()) >> i;
+			const uint32 SrcSizeX = (Texture->Source.GetSizeX()) >> MipIndex;
+			const uint32 SrcSizeY = (Texture->Source.GetSizeY()) >> MipIndex;
 			const uint32 SrcPitch = (SrcSizeX * DataSize);
 			const uint32 BufferSize = SrcSizeX * SrcSizeY * DataSize;
 
 			// Copy Mip update data so we can avoid waiting for Render thread in calling method
 			FMipInfo* CopyMipInfo = new FMipInfo();
-			CopyMipInfo->MipUpdateRegions = MipInfo[i].MipUpdateRegions;
+			CopyMipInfo->MipUpdateRegions = MipInfo[MipIndex].MipUpdateRegions;
 			CopyMipInfo->MipData = FMemory::Malloc(BufferSize);
-			FMemory::Memcpy(CopyMipInfo->MipData, MipInfo[i].MipData, BufferSize);
+			FMemory::Memcpy(CopyMipInfo->MipData, MipInfo[MipIndex].MipData, BufferSize);
 
-			Texture->UpdateTextureRegions(i, CopyMipInfo->MipUpdateRegions.Num(), &CopyMipInfo->MipUpdateRegions[0], SrcPitch, DataSize, (uint8*)CopyMipInfo->MipData, 
+			Texture->UpdateTextureRegions(MipIndex, CopyMipInfo->MipUpdateRegions.Num(), &CopyMipInfo->MipUpdateRegions[0], SrcPitch, DataSize, (uint8*)CopyMipInfo->MipData,
 				[CopyMipInfo](uint8* SrcData, const FUpdateTextureRegion2D*) 
 				{ 
 					FMemory::Free(CopyMipInfo->MipData);
@@ -4108,12 +4141,12 @@ bool FLandscapeTextureDataInfo::UpdateTextureData()
 FLandscapeTextureDataInfo::~FLandscapeTextureDataInfo()
 {
 	// Unlock any mips still locked.
-	for( int32 i=0;i<MipInfo.Num();i++ )
+	for( int32 MipIndex = 0; MipIndex < MipInfo.Num(); MipIndex++ )
 	{
-		if( MipInfo[i].MipData )
+		if(MipInfo[MipIndex].MipData != nullptr)
 		{
-			Texture->Source.UnlockMip(i);
-			MipInfo[i].MipData = NULL;
+			Texture->Source.UnlockMip(MipIndex);
+			MipInfo[MipIndex].MipData = nullptr;
 		}
 	}
 	Texture->ClearFlags(RF_Transactional);
diff --git a/Engine/Source/Runtime/Landscape/Public/LandscapeEdit.h b/Engine/Source/Runtime/Landscape/Public/LandscapeEdit.h
--- a/Engine/Source/Runtime/Landscape/Public/LandscapeEdit.h
+++ b/Engine/Source/Runtime/Landscape/Public/LandscapeEdit.h
@@ -49,11 +49,11 @@ public:
 
 struct FLandscapeTextureDataInfo
 {
-	struct FMipInfo
+	struct UE_INTERNAL FMipInfo
 	{
-		void* MipData;
+		void* MipData = nullptr;
 		TArray<FUpdateTextureRegion2D> MipUpdateRegions;
-		bool bFull;
+		bool bFull = false;
 	};
 
 	FLandscapeTextureDataInfo(UTexture2D* InTexture, bool bShouldDirtyPackage, ELandscapeTextureUsage TextureUsage = ELandscapeTextureUsage::Unknown, ELandscapeTextureType TextureType = ELandscapeTextureType::Unknown);
@@ -64,56 +64,22 @@ struct FLandscapeTextureDataInfo
 
 	int32 NumMips() { return MipInfo.Num(); }
 
-	void AddMipUpdateRegion(int32 MipNum, int32 InX1, int32 InY1, int32 InX2, int32 InY2)
-	{
-		if (MipInfo[MipNum].bFull)
-		{
-			return;
-		}
-
-		check(MipNum < MipInfo.Num());
-		uint32 Width = 1 + InX2 - InX1;
-		uint32 Height = 1 + InY2 - InY1;
-		// Catch situation where we are updating the whole texture to avoid adding redundant regions once the whole region as been included.
-		if (Width == GetMipSizeX(MipNum) && Height == GetMipSizeY(MipNum))
-		{
-			MipInfo[MipNum].bFull = true;
-			MipInfo[MipNum].MipUpdateRegions.Reset();
-			// Push a full region for UpdateTextureData() to process later
-			MipInfo[MipNum].MipUpdateRegions.Emplace(0, 0, 0, 0, Width, Height);
-			return;
-		}
-
-		MipInfo[MipNum].MipUpdateRegions.Emplace(InX1, InY1, InX1, InY1, Width, Height);
-	}
-		
-	void* GetMipData(int32 MipNum)
-	{
-		check( MipNum < MipInfo.Num() );
-		if( !MipInfo[MipNum].MipData )
-		{
-			// will Unlock in destructor
-			MipInfo[MipNum].MipData = Texture->Source.LockMip(MipNum);
-			// probably should check that we got the lock :
-			//check( MipInfo[MipNum].MipData != nullptr );
-			// also this return value is usually cast to FColor *
-			// should instead have a function that returns FColor * and ensures the mip actually is BGRA8
-		}
-		return MipInfo[MipNum].MipData;
-	}
+	void AddMipUpdateRegion(int32 InMipIndex, int32 InX1, int32 InY1, int32 InX2, int32 InY2);
+	
+	void* GetMipData(int32 InMipIndex);
 
-	int32 GetMipSizeX(int32 MipNum) const
+	int32 GetMipSizeX(int32 InMipIndex) const
 	{
-		return FMath::Max(Texture->Source.GetSizeX() >> MipNum, 1);
+		return FMath::Max(Texture->Source.GetSizeX() >> InMipIndex, 1);
 	}
 
-	int32 GetMipSizeY(int32 MipNum) const
+	int32 GetMipSizeY(int32 InMipIndex) const
 	{
-		return FMath::Max(Texture->Source.GetSizeY() >> MipNum, 1);
+		return FMath::Max(Texture->Source.GetSizeY() >> InMipIndex, 1);
 	}
 
 private:
-	UTexture2D* Texture;
+	UTexture2D* Texture = nullptr;
 	TArray<FMipInfo> MipInfo;
 	ELandscapeTextureUsage TextureUsage = ELandscapeTextureUsage::Unknown;
 	ELandscapeTextureType TextureType = ELandscapeTextureType::Unknown;
@@ -330,7 +296,7 @@ private:
 	template<typename TStoreData, typename TGetHeightMapDataFunction>
 	void GetHeightDataInternal(int32& ValidX1, int32& ValidY1, int32& ValidX2, int32& ValidY2, TStoreData& StoreData, TGetHeightMapDataFunction GetHeightMapDataFunction);
 
-	FORCEINLINE FColor& GetHeightMapColor(const ULandscapeComponent* Component, int32 TexU, int32 TexV, FColor* TextureData);
+	FORCEINLINE const FColor& GetHeightMapColor(const ULandscapeComponent* Component, int32 TexU, int32 TexV, FColor* TextureData);
 
 	// test if layer is allowed for a given texel
 	inline bool IsLayerAllowed(const ULandscapeLayerInfoObject* LayerInfo,
