Build fix

#Synced-CL 49076111

diff --git a/Engine/Source/Runtime/Landscape/Public/LandscapeEdit.h b/Engine/Source/Runtime/Landscape/Public/LandscapeEdit.h
--- a/Engine/Source/Runtime/Landscape/Public/LandscapeEdit.h
+++ b/Engine/Source/Runtime/Landscape/Public/LandscapeEdit.h
@@ -49,7 +49,7 @@ public:
 
 struct FLandscapeTextureDataInfo
 {
-	struct UE_INTERNAL FMipInfo
+	struct FMipInfo
 	{
 		void* MipData = nullptr;
 		TArray<FUpdateTextureRegion2D> MipUpdateRegions;
