Fixed for crash when using ALandscape::RenderWeightmap on mobile

#tests Validated by QA with code-only PF on a version that still had the call to the problematic BP function

#Synced-CL 49356504

diff --git a/Engine/Source/Runtime/Landscape/Classes/LandscapeComponent.h b/Engine/Source/Runtime/Landscape/Classes/LandscapeComponent.h
--- a/Engine/Source/Runtime/Landscape/Classes/LandscapeComponent.h
+++ b/Engine/Source/Runtime/Landscape/Classes/LandscapeComponent.h
@@ -800,6 +800,7 @@ public:
 	LANDSCAPE_API TArray<FWeightmapLayerAllocationInfo>& GetWeightmapLayerAllocations(const FGuid& InLayerGuid);
 	LANDSCAPE_API const TArray<FWeightmapLayerAllocationInfo>& GetWeightmapLayerAllocations(const FGuid& InLayerGuid) const;
 
+	LANDSCAPE_API const TArray<UTexture2D*>& GetCurrentRuntimeWeightmapTextures() const;
 	LANDSCAPE_API TArray<FWeightmapLayerAllocationInfo>& GetCurrentRuntimeWeightmapLayerAllocations();
 	LANDSCAPE_API const TArray<FWeightmapLayerAllocationInfo>& GetCurrentRuntimeWeightmapLayerAllocations() const;
 
diff --git a/Engine/Source/Runtime/Landscape/Private/Landscape.cpp b/Engine/Source/Runtime/Landscape/Private/Landscape.cpp
--- a/Engine/Source/Runtime/Landscape/Private/Landscape.cpp
+++ b/Engine/Source/Runtime/Landscape/Private/Landscape.cpp
@@ -2540,6 +2540,13 @@ TArray<FWeightmapLayerAllocationInfo>& ULandscapeComponent::GetWeightmapLayerAll
 	return WeightmapLayerAllocations;
 }
 
+const TArray<UTexture2D*>& ULandscapeComponent::GetCurrentRuntimeWeightmapTextures() const
+{
+	UWorld* World = GetWorld();
+	bool bIsMobile = (World != nullptr) && (World->GetFeatureLevel() == ERHIFeatureLevel::ES3_1);
+	return bIsMobile ? MobileWeightmapTextures : WeightmapTextures;
+}
+
 const TArray<FWeightmapLayerAllocationInfo>& ULandscapeComponent::GetCurrentRuntimeWeightmapLayerAllocations() const
 {
 	UWorld* World = GetWorld();
@@ -3800,8 +3807,8 @@ bool ALandscape::RenderMergedTextureInternal(const FTransform& InRenderAreaWorld
 			}
 			else
 			{
-				const TArray<UTexture2D*>& WeightmapTextures = Component->GetWeightmapTextures();
-				const TArray<FWeightmapLayerAllocationInfo>& AllocInfos = Component->GetWeightmapLayerAllocations();
+				const TArray<UTexture2D*>& WeightmapTextures = Component->GetCurrentRuntimeWeightmapTextures();
+				const TArray<FWeightmapLayerAllocationInfo>& AllocInfos = Component->GetCurrentRuntimeWeightmapLayerAllocations();
 				const FWeightmapLayerAllocationInfo* AllocInfo = AllocInfos.FindByPredicate([TargetLayerName](const FWeightmapLayerAllocationInfo& InAllocInfo) { return InAllocInfo.IsAllocated() && (InAllocInfo.GetLayerName() == TargetLayerName); });
 				if (AllocInfo != nullptr)
 				{
