Minor preliminary landscape code refactor :
* Moved GRASS_DEBUG_LOG to LandscapePrivate.h
* DEBUG_TRANSITION is now GRASS_DEBUG_TRANSITION (specific to LandscapeGrassMapsBuilder.cpp)
* Deprecated UE::Landscape::CanRenderGrassMap(ULandscapeComponent*) in favor of ULandscapeComponent::CanRenderGrassMap() : the former was meant for the editor and was superseded by the latter. The implementations were otherwise similar

#Synced-CL 49356504

diff --git a/Engine/Plugins/PCG/Source/PCG/Private/Elements/PCGGenerateLandscapeTextures.cpp b/Engine/Plugins/PCG/Source/PCG/Private/Elements/PCGGenerateLandscapeTextures.cpp
--- a/Engine/Plugins/PCG/Source/PCG/Private/Elements/PCGGenerateLandscapeTextures.cpp
+++ b/Engine/Plugins/PCG/Source/PCG/Private/Elements/PCGGenerateLandscapeTextures.cpp
@@ -488,7 +488,7 @@ bool FPCGGenerateLandscapeTexturesElement::ExecuteInternal(FPCGContext* InContex
 				return true;
 			}
 
-			bAllReady &= UE::Landscape::CanRenderGrassMap(LandscapeComponent);
+			bAllReady &= LandscapeComponent->CanRenderGrassMap();
 		}
 
 		if (bAllReady)
diff --git a/Engine/Source/Runtime/Landscape/Internal/LandscapeGrassWeightExporter.h b/Engine/Source/Runtime/Landscape/Internal/LandscapeGrassWeightExporter.h
--- a/Engine/Source/Runtime/Landscape/Internal/LandscapeGrassWeightExporter.h
+++ b/Engine/Source/Runtime/Landscape/Internal/LandscapeGrassWeightExporter.h
@@ -15,6 +15,7 @@ using FRDGTextureRef = FRDGTexture*;
 
 namespace UE::Landscape
 {
+	UE_DEPRECATED(5.8, "Use ULandscapeComponent::CanRenderGrassMap instead")
 	LANDSCAPE_API bool CanRenderGrassMap(ULandscapeComponent* Component);
 	LANDSCAPE_API bool IsRuntimeGrassMapGenerationSupported();
 }
diff --git a/Engine/Source/Runtime/Landscape/Private/LandscapeGrass.cpp b/Engine/Source/Runtime/Landscape/Private/LandscapeGrass.cpp
--- a/Engine/Source/Runtime/Landscape/Private/LandscapeGrass.cpp
+++ b/Engine/Source/Runtime/Landscape/Private/LandscapeGrass.cpp
@@ -910,18 +910,38 @@ bool ULandscapeComponent::UpdateGrassTypes(bool bForceUpdate)
 bool ULandscapeComponent::CanRenderGrassMap() const
 {
 	// Check we can render
-	UWorld* ComponentWorld = GetWorld();
-	if (!GIsEditor || GUsingNullRHI || !ComponentWorld || ComponentWorld->IsGameWorld() || ComponentWorld->GetFeatureLevel() < ERHIFeatureLevel::SM5 || !SceneProxy)
+	const UWorld* World = GetWorld();
+	if (GUsingNullRHI || !World || !SceneProxy)
 	{
+		GRASS_DEBUG_LOG(TEXT("GrassMap No SceneProxy for %s"), *GetName());
 		return false;
 	}
 
-	UMaterialInstance* MaterialInstance = GetMaterialInstanceCount(false) > 0 ? GetMaterialInstance(0) : nullptr;
-	FMaterialResource* MaterialResource = MaterialInstance != nullptr ? MaterialInstance->GetMaterialResource(GetFeatureLevelShaderPlatform_Checked(ComponentWorld->GetFeatureLevel())) : nullptr;
-
 	// Check we can render the material
+	const int32 MatInstCount = GetCurrentRuntimeMaterialInstanceCount();
+	if (MatInstCount <= 0)
+	{
+		GRASS_DEBUG_LOG(TEXT("GrassMap MaterialInstanceCount %d <= 0 for %s (%d dyn:%d inst:%d mob:%d)"),
+			MatInstCount,
+			*GetName(),
+			GetLandscapeProxy()->bUseDynamicMaterialInstance ? 1 : 0,
+			MaterialInstancesDynamic.Num(),
+			MaterialInstances.Num(),
+			MobileMaterialInterfaces.Num());
+		return false;
+	}
+
+	UMaterialInterface* MaterialInterface = GetCurrentRuntimeMaterialInterface(0);
+	if (MaterialInterface == nullptr)
+	{
+		GRASS_DEBUG_LOG(TEXT("GrassMap MaterialInterface NULL for %s"), *GetName());
+		return false;
+	}
+
+	const FMaterialResource* MaterialResource = MaterialInterface->GetMaterialResource(GetFeatureLevelShaderPlatform_Checked(World->GetFeatureLevel()));
 	if (MaterialResource == nullptr)
 	{
+		GRASS_DEBUG_LOG(TEXT("GrassMap MaterialResource NULL for %s (%s feature level %d)"), *GetName(), *MaterialInterface->GetName(), World->GetFeatureLevel());
 		return false;
 	}
 
@@ -929,12 +949,12 @@ bool ULandscapeComponent::CanRenderGrassMap() const
 	FMaterialShaderTypes ShaderTypes;
 	UE::Landscape::Grass::AddGrassWeightShaderTypes(ShaderTypes);
 
-	FVertexFactoryType* LandscapeGrassVF = FindVertexFactoryType(FName(TEXT("FLandscapeFixedGridVertexFactory"), FNAME_Find));
+	const FVertexFactoryType* LandscapeGrassVF = FindVertexFactoryType(FName(TEXT("FLandscapeFixedGridVertexFactory"), FNAME_Find));
 	if (!MaterialResource->HasShaders(ShaderTypes, LandscapeGrassVF))
 	{
+		GRASS_DEBUG_LOG(TEXT("GrassMap MaterialResource does not have FixedGridVF for %s"), *GetName());
 		return false;
 	}
-
 	return true;
 }
 
diff --git a/Engine/Source/Runtime/Landscape/Private/LandscapeGrassMapsBuilder.cpp b/Engine/Source/Runtime/Landscape/Private/LandscapeGrassMapsBuilder.cpp
--- a/Engine/Source/Runtime/Landscape/Private/LandscapeGrassMapsBuilder.cpp
+++ b/Engine/Source/Runtime/Landscape/Private/LandscapeGrassMapsBuilder.cpp
@@ -21,9 +21,7 @@
 
 #define LOCTEXT_NAMESPACE "Landscape"
 
-#define GRASS_DEBUG_LOG(...) UE_LOG(LogGrass, Verbose, __VA_ARGS__)
-
-#define DEBUG_TRANSITION(StateRef, StageBefore, StageAfter) \
+#define GRASS_DEBUG_TRANSITION(StateRef, StageBefore, StageAfter) \
 	GRASS_DEBUG_LOG(TEXT("%s %s -> %s (after %d ticks) Pend:%d Strm:%d Rend:%d Fetch:%d Pop:%d NR:%d Total:%d"), \
 		StateRef.Component ? *StateRef.Component->GetName() : TEXT("<REMOVED>"), \
 		TEXT(#StageBefore), \
@@ -211,55 +209,10 @@ namespace UE::Landscape
 	}
 #endif // WITH_EDITOR
 
+	// Deprecated
 	bool CanRenderGrassMap(ULandscapeComponent *Component)
 	{
-		// Check we can render
-		const UWorld* ComponentWorld = Component->GetWorld();
-		if (GUsingNullRHI || !ComponentWorld || !Component->SceneProxy)
-		{
-			GRASS_DEBUG_LOG(TEXT("GrassMap No SceneProxy for %s"), *Component->GetName());
-			return false;
-		}
-
-		// Check we can render the material
-		const int32 MatInstCount = Component->GetCurrentRuntimeMaterialInstanceCount();
-		if (MatInstCount <= 0)
-		{
-			GRASS_DEBUG_LOG(TEXT("GrassMap MaterialInstanceCount %d <= 0 for %s (%d dyn:%d inst:%d mob:%d)"),
-				MatInstCount,
-				*Component->GetName(),
-				Component->GetLandscapeProxy()->bUseDynamicMaterialInstance ? 1 : 0,
-				Component->MaterialInstancesDynamic.Num(),
-				Component->MaterialInstances.Num(),
-				Component->MobileMaterialInterfaces.Num());
-			return false;
-		}
-
-		UMaterialInterface* MaterialInterface = Component->GetCurrentRuntimeMaterialInterface(0);
-		if (MaterialInterface == nullptr)
-		{
-			GRASS_DEBUG_LOG(TEXT("GrassMap MaterialInterface NULL for %s"), *Component->GetName());
-			return false;
-		}
-
-		const FMaterialResource* MaterialResource = MaterialInterface->GetMaterialResource(GetFeatureLevelShaderPlatform_Checked(ComponentWorld->GetFeatureLevel()));
-		if (MaterialResource == nullptr)
-		{
-			GRASS_DEBUG_LOG(TEXT("GrassMap MaterialResource NULL for %s (%s feature level %d)"), *Component->GetName(), *MaterialInterface->GetName(), ComponentWorld->GetFeatureLevel());
-			return false;
-		}
-
-		// We only need the GrassWeight shaders on the fixed grid vertex factory to render grass maps : 
-		FMaterialShaderTypes ShaderTypes;
-		UE::Landscape::Grass::AddGrassWeightShaderTypes(ShaderTypes);
-
-		const FVertexFactoryType* LandscapeGrassVF = FindVertexFactoryType(FName(TEXT("FLandscapeFixedGridVertexFactory"), FNAME_Find));
-		if (!MaterialResource->HasShaders(ShaderTypes, LandscapeGrassVF))
-		{
-			GRASS_DEBUG_LOG(TEXT("GrassMap MaterialResource does not have FixedGridVF for %s"), *Component->GetName());
-			return false;
-		}
-		return true;
+		return Component->CanRenderGrassMap();
 	}
 
 	bool IsRuntimeGrassMapGenerationSupported()
@@ -472,7 +425,7 @@ bool FLandscapeGrassMapsBuilder::UpdateTrackedComponents(const TArray<FVector>&
 #if WITH_EDITOR
 				if (!bIsGameWorld && AmortizedUpdate.ShouldUpdate(ComponentStateIndex))
 				{
-					if (UE::Landscape::CanRenderGrassMap(Component))
+					if (Component->CanRenderGrassMap())
 					{
 						break; // cancel and evict to restart build process
 					}
@@ -614,7 +567,7 @@ bool FLandscapeGrassMapsBuilder::UpdateTrackedComponents(const TArray<FVector>&
 #if WITH_EDITOR
 				// in editor the renderability can change unexpectedly -- so check one last time just before we actually render
 				// (in game, we only check at the beginning of kicking off the generation process)
-				if (!UE::Landscape::CanRenderGrassMap(State->Component))
+				if (!State->Component->CanRenderGrassMap())
 				{
 					// can't render, move to NotReady state, which will monitor until it is renderable
 					StreamingToNotReady(*State);
@@ -1016,12 +969,12 @@ void FLandscapeGrassMapsBuilder::DebugRenderComponents(
 				case EComponentStage::NotReady:
 					{
 						// if it's not ready because of shader reasons
-						if (!FailedComponents.Contains(Component) && !UE::Landscape::CanRenderGrassMap(Component))
+						if (!FailedComponents.Contains(Component) && !Component->CanRenderGrassMap())
 						{
 	#if WITH_EDITOR
 							// in editor, try to force compilation to complete
 							UE::Landscape::CompileGrassMapShader(State->Component);
-							if (!UE::Landscape::CanRenderGrassMap(Component))
+							if (!Component->CanRenderGrassMap())
 	#endif // WITH_EDITOR
 							{
 								// failed to compile shaders... we won't be able to build the grass map for this component
@@ -1214,12 +1167,12 @@ bool FLandscapeGrassMapsBuilder::BuildGrassMapsNowForComponents(
 			if (State->Stage == EComponentStage::NotReady)
 			{
 				// if it's not ready because of shader reasons
-				if (!FailedComponents.Contains(Component) && !UE::Landscape::CanRenderGrassMap(Component))
+				if (!FailedComponents.Contains(Component) && !Component->CanRenderGrassMap())
 				{
 #if WITH_EDITOR
 					// in editor, try to force compilation to complete
 					UE::Landscape::CompileGrassMapShader(State->Component);
-					if (!UE::Landscape::CanRenderGrassMap(Component))
+					if (!Component->CanRenderGrassMap())
 #endif // WITH_EDITOR
 					{
 						// failed to compile shaders... we won't be able to build the grass map for this component
@@ -1304,7 +1257,7 @@ bool FLandscapeGrassMapsBuilder::CancelAndEvict(FComponentState& State, bool bCa
 			check(NotReadyCount > 0);
 			NotReadyCount--;
 			PendingCount++;
-			DEBUG_TRANSITION(State, NotReady, Pending);
+			GRASS_DEBUG_TRANSITION(State, NotReady, Pending);
 		}
 		break;
 	case EComponentStage::TextureStreaming:
@@ -1313,7 +1266,7 @@ bool FLandscapeGrassMapsBuilder::CancelAndEvict(FComponentState& State, bool bCa
 			check(StreamingCount > 0);
 			StreamingCount--;
 			PendingCount++;
-			DEBUG_TRANSITION(State, TextureStreaming, Pending);
+			GRASS_DEBUG_TRANSITION(State, TextureStreaming, Pending);
 		}
 		break;
 	case EComponentStage::Rendering:
@@ -1353,7 +1306,7 @@ bool FLandscapeGrassMapsBuilder::CancelAndEvict(FComponentState& State, bool bCa
 			check(RenderingCount > 0);
 			RenderingCount--;
 			PendingCount++;
-			DEBUG_TRANSITION(State, Rendering, Pending);
+			GRASS_DEBUG_TRANSITION(State, Rendering, Pending);
 		}
 		break;
 	case EComponentStage::AsyncFetch:
@@ -1378,7 +1331,7 @@ bool FLandscapeGrassMapsBuilder::CancelAndEvict(FComponentState& State, bool bCa
 			check(AsyncFetchCount > 0);
 			AsyncFetchCount--;
 			PendingCount++;
-			DEBUG_TRANSITION(State, AsyncFetch, Pending);
+			GRASS_DEBUG_TRANSITION(State, AsyncFetch, Pending);
 		}
 		break;
 	case EComponentStage::GrassMapsPopulated:
@@ -1402,7 +1355,7 @@ bool FLandscapeGrassMapsBuilder::CancelAndEvict(FComponentState& State, bool bCa
 			check(PopulatedCount > 0);
 			PopulatedCount--;
 			PendingCount++;
-			DEBUG_TRANSITION(State, Populated, Pending);
+			GRASS_DEBUG_TRANSITION(State, Populated, Pending);
 		}
 		break;
 	default:
@@ -1491,7 +1444,7 @@ bool FLandscapeGrassMapsBuilder::StartGrassMapGeneration(FComponentState& State,
 	}
 
 	// if we can't currently render, it's not ready
-	if (!UE::Landscape::CanRenderGrassMap(Component))
+	if (!Component->CanRenderGrassMap())
 	{
 #if WITH_EDITOR
 		if (bForceCompileShaders)
@@ -1499,7 +1452,7 @@ bool FLandscapeGrassMapsBuilder::StartGrassMapGeneration(FComponentState& State,
 			UE::Landscape::CompileGrassMapShader(Component);
 		}
 
-		if (!UE::Landscape::CanRenderGrassMap(Component))
+		if (!Component->CanRenderGrassMap())
 #endif // WITH_EDITOR
 		{
 			PendingToNotReady(State);
@@ -1519,7 +1472,7 @@ void FLandscapeGrassMapsBuilder::StreamingToNotReady(FComponentState& State)
 	State.Stage = EComponentStage::NotReady;
 	NotReadyCount++;
 
-	DEBUG_TRANSITION(State, Streaming, NotReady);
+	GRASS_DEBUG_TRANSITION(State, Streaming, NotReady);
 	State.TickCount = 0;
 }
 
@@ -1532,7 +1485,7 @@ void FLandscapeGrassMapsBuilder::PendingToNotReady(FComponentState& State)
 
 	RemoveFromPendingComponentHeap(&State);
 
-	DEBUG_TRANSITION(State, Pending, NotReady);
+	GRASS_DEBUG_TRANSITION(State, Pending, NotReady);
 	State.TickCount = 0;
 }
 
@@ -1545,7 +1498,7 @@ void FLandscapeGrassMapsBuilder::PendingToPopulatedFastPathAlreadyHasData(FCompo
 
 	RemoveFromPendingComponentHeap(&State);
 
-	DEBUG_TRANSITION(State, Pending, Populated_Existing);
+	GRASS_DEBUG_TRANSITION(State, Pending, Populated_Existing);
 	State.TickCount = 0;
 }
 
@@ -1567,7 +1520,7 @@ void FLandscapeGrassMapsBuilder::PendingToPopulatedFastPathNoGrass(FComponentSta
 
 	RemoveFromPendingComponentHeap(&State);
 
-	DEBUG_TRANSITION(State, Pending, Populated_Empty);
+	GRASS_DEBUG_TRANSITION(State, Pending, Populated_Empty);
 	State.TickCount = 0;
 }
 
@@ -1602,7 +1555,7 @@ void FLandscapeGrassMapsBuilder::PendingToStreaming(FComponentState& State)
 
 	RemoveFromPendingComponentHeap(&State);
 
-	DEBUG_TRANSITION(State, Pending, Streaming);
+	GRASS_DEBUG_TRANSITION(State, Pending, Streaming);
 	State.TickCount = 0;
 }
 
@@ -1675,7 +1628,7 @@ void FLandscapeGrassMapsBuilder::KickOffRenderAndReadback(FComponentState& State
 	State.Stage = EComponentStage::Rendering;
 	RenderingCount++;
 
-	DEBUG_TRANSITION(State, Streaming, Rendering);
+	GRASS_DEBUG_TRANSITION(State, Streaming, Rendering);
 	State.TickCount = 0;
 }
 
@@ -1697,7 +1650,7 @@ void FLandscapeGrassMapsBuilder::LaunchAsyncFetchTask(FComponentState& State)
 	State.Stage = EComponentStage::AsyncFetch;
 	AsyncFetchCount++;
 
-	DEBUG_TRANSITION(State, Rendering, AsyncFetch);
+	GRASS_DEBUG_TRANSITION(State, Rendering, AsyncFetch);
 	State.TickCount = 0;
 }
 
@@ -1717,7 +1670,7 @@ void FLandscapeGrassMapsBuilder::PopulateGrassDataFromAsyncFetchTask(FComponentS
 	State.Stage = EComponentStage::GrassMapsPopulated;
 	PopulatedCount++;
 
-	DEBUG_TRANSITION(State, AsyncFetch, Populated);
+	GRASS_DEBUG_TRANSITION(State, AsyncFetch, Populated);
 	State.TickCount = 0;
 }
 
@@ -1739,7 +1692,7 @@ void FLandscapeGrassMapsBuilder::PopulateGrassDataFromReadback(FComponentState&
 	State.Stage = EComponentStage::GrassMapsPopulated;
 	PopulatedCount++;
 
-	DEBUG_TRANSITION(State, Rendering, Populated);
+	GRASS_DEBUG_TRANSITION(State, Rendering, Populated);
 	State.TickCount = 0;
 }
 
@@ -1890,4 +1843,4 @@ int32 FLandscapeGrassMapsBuilder::GetOutdatedGrassMapCount(bool bInForceUpdate)
 
 #endif // WITH_EDITOR
 #undef LOCTEXT_NAMESPACE
-#undef GRASS_DEBUG_LOG
\ No newline at end of file
+#undef GRASS_DEBUG_TRANSITION
\ No newline at end of file
diff --git a/Engine/Source/Runtime/Landscape/Private/LandscapePrivate.h b/Engine/Source/Runtime/Landscape/Private/LandscapePrivate.h
--- a/Engine/Source/Runtime/Landscape/Private/LandscapePrivate.h
+++ b/Engine/Source/Runtime/Landscape/Private/LandscapePrivate.h
@@ -10,6 +10,8 @@ LANDSCAPE_API DECLARE_LOG_CATEGORY_EXTERN(LogLandscape, Log, All);
 LANDSCAPE_API DECLARE_LOG_CATEGORY_EXTERN(LogLandscapeBP, Display, All);
 LANDSCAPE_API DECLARE_LOG_CATEGORY_EXTERN(LogGrass, Log, All);
 
+#define GRASS_DEBUG_LOG(...) UE_LOG(LogGrass, Verbose, __VA_ARGS__)
+
 /**
  * Landscape stats
  */
