Landscape Editor - Fixed invalid normals when creating a new landscape with WP regions
* Proxies outside of the first landscape region do not serialize correct normals (default normal values are not initialized). During the merge in ResolveLayersTexture, the landscape component texture hash matches so the serialized texture does not get modified. As a result, the normals appear to be correct in the editor but upon reload only the default (incorrect) normals have been serialized
* Fixed an additional issue where the first landscape WP region applies a non-zero rotation

#Synced-CL 49650152

diff --git a/Engine/Source/Editor/LandscapeEditor/Private/LandscapeEditorDetailCustomization_NewLandscape.cpp b/Engine/Source/Editor/LandscapeEditor/Private/LandscapeEditorDetailCustomization_NewLandscape.cpp
--- a/Engine/Source/Editor/LandscapeEditor/Private/LandscapeEditorDetailCustomization_NewLandscape.cpp
+++ b/Engine/Source/Editor/LandscapeEditor/Private/LandscapeEditorDetailCustomization_NewLandscape.cpp
@@ -1090,8 +1090,8 @@ void FLandscapeEditorDetailCustomization_NewLandscape::AddComponents(ULandscapeI
 
 		TArray<FColor> HeightData;
 		const int32 ComponentVerts = (LandscapeComponent->SubsectionSizeQuads + 1) * LandscapeComponent->NumSubsections;
-		const FColor PackedMidpoint = LandscapeDataAccess::PackHeight(LandscapeDataAccess::GetTexHeight(0.0f));
-		HeightData.Init(PackedMidpoint, FMath::Square(ComponentVerts));
+		const FColor DefaultHeightColor = LandscapeDataAccess::GetDefaultPackedHeightColor();
+		HeightData.Init(DefaultHeightColor, FMath::Square(ComponentVerts));
 
 		LandscapeComponent->InitHeightmapData(HeightData, true);
 		LandscapeComponent->UpdateMaterialInstances();
diff --git a/Engine/Source/Editor/LandscapeEditor/Private/LandscapeRegionUtils.cpp b/Engine/Source/Editor/LandscapeEditor/Private/LandscapeRegionUtils.cpp
--- a/Engine/Source/Editor/LandscapeEditor/Private/LandscapeRegionUtils.cpp
+++ b/Engine/Source/Editor/LandscapeEditor/Private/LandscapeRegionUtils.cpp
@@ -24,7 +24,7 @@ ALocationVolume* CreateLandscapeRegionVolume(UWorld* InWorld, ALandscapeProxy* I
 
 	FVector ParentLocation = InParentLandscapeActor->GetActorLocation();
 	FVector Location = ParentLocation + FVector(InRegionCoordinate.X * InRegionSize, InRegionCoordinate.Y * InRegionSize, 0.0) + FVector(InRegionSize / 2.0, InRegionSize / 2.0, 0.0);
-	FRotator Rotation;
+	FRotator Rotation = FRotator::ZeroRotator;
 	FActorSpawnParameters SpawnParameters;
 
 	const FString Label = FString::Printf(TEXT("LandscapeRegion_%i_%i"), InRegionCoordinate.X, InRegionCoordinate.Y);
diff --git a/Engine/Source/Runtime/Landscape/Private/Landscape.cpp b/Engine/Source/Runtime/Landscape/Private/Landscape.cpp
--- a/Engine/Source/Runtime/Landscape/Private/Landscape.cpp
+++ b/Engine/Source/Runtime/Landscape/Private/Landscape.cpp
@@ -2825,6 +2825,7 @@ void ULandscapeComponent::AddDefaultLayerData(const FGuid& InLayerGuid, const TA
 
 		// Setup Heightmap data
 		UTexture2D** LayerHeightmap = InOutCreatedHeightmapTextures.Find(ComponentHeightmap);
+		const FColor DefaultHeightColor = LandscapeDataAccess::GetDefaultPackedHeightColor();
 
 		if (LayerHeightmap == nullptr)
 		{
@@ -2858,15 +2859,8 @@ void ULandscapeComponent::AddDefaultLayerData(const FGuid& InLayerGuid, const TA
 
 								const int32 HeightTexDataIdx = (HeightmapComponentOffsetX + TexX) + (HeightmapComponentOffsetY + TexY) * NewLayerHeightmap->Source.GetSizeX();
 
-								// copy height and normal data
-								const uint16 HeightValue = LandscapeDataAccess::GetTexHeight(0.f);
-
-								Mip0Data[HeightTexDataIdx].R = HeightValue >> 8;
-								Mip0Data[HeightTexDataIdx].G = HeightValue & 255;
-
-								// Normal with get calculated later
-								Mip0Data[HeightTexDataIdx].B = 0;
-								Mip0Data[HeightTexDataIdx].A = 0;
+								// Copy default height and normal data. Normals will be calculated later
+								Mip0Data[HeightTexDataIdx] = DefaultHeightColor;
 							}
 						}
 					}
diff --git a/Engine/Source/Runtime/Landscape/Private/LandscapeConfigHelper.cpp b/Engine/Source/Runtime/Landscape/Private/LandscapeConfigHelper.cpp
--- a/Engine/Source/Runtime/Landscape/Private/LandscapeConfigHelper.cpp
+++ b/Engine/Source/Runtime/Landscape/Private/LandscapeConfigHelper.cpp
@@ -351,11 +351,10 @@ void FLandscapeConfigHelper::CopyRegionToComponent(ULandscapeInfo* InLandscapeIn
 
 	TArray<FColor> HeightInitData;
 	HeightInitData.SetNum(FMath::Square(HeightmapSize));
-	const uint16 DefaultHeight = LandscapeDataAccess::GetTexHeight(0.f);
-	FColor DefaultPackedHeight = LandscapeDataAccess::PackHeight(DefaultHeight);
+	const FColor DefaultHeightColor = LandscapeDataAccess::GetDefaultPackedHeightColor();
 	for (int32 Index = 0; Index < HeightInitData.Num(); ++Index)
 	{
-		HeightInitData[Index] = DefaultPackedHeight;
+		HeightInitData[Index] = DefaultHeightColor;
 	}
 	InComponent->InitHeightmapData(HeightInitData, false);
 
diff --git a/Engine/Source/Runtime/Landscape/Public/LandscapeDataAccess.h b/Engine/Source/Runtime/Landscape/Public/LandscapeDataAccess.h
--- a/Engine/Source/Runtime/Landscape/Public/LandscapeDataAccess.h
+++ b/Engine/Source/Runtime/Landscape/Public/LandscapeDataAccess.h
@@ -25,6 +25,8 @@ namespace LandscapeDataAccess
 {
 	inline constexpr int32 MaxValue = 65535;
 	inline constexpr float MidValue = 32768.f;
+	inline constexpr float MidNormalValue = 127.5f;
+
 	FORCEINLINE float GetLocalHeight(uint16 Height)
 	{
 		return (static_cast<float>(Height) - MidValue) * LANDSCAPE_ZSCALE;
@@ -49,6 +51,16 @@ namespace LandscapeDataAccess
 		return GetLocalHeight(Height);
 	}
 
+	FORCEINLINE FColor PackNormal(const FVector& InNormal)
+	{
+		const FVector Normal = InNormal.GetSafeNormal();
+
+		FColor Color(ForceInit);
+		Color.B = static_cast<uint8>(FMath::Clamp(Normal.X * 127.5 + 127.5, 0.0f, 255.0f));
+		Color.A = static_cast<uint8>(FMath::Clamp(Normal.Y * 127.5 + 127.5, 0.0f, 255.0f));
+		return MoveTemp(Color);
+	}
+
 	FORCEINLINE FVector UnpackNormal(const FColor& InHeightmapSample)
 	{
 		FVector Normal;
@@ -57,6 +69,13 @@ namespace LandscapeDataAccess
 		Normal.Z = FMath::Sqrt(FMath::Max(1.0f - (FMath::Square(Normal.X) + FMath::Square(Normal.Y)), 0.0f));
 		return Normal;
 	}
+
+	FORCEINLINE FColor GetDefaultPackedHeightColor()
+	{
+		const uint16 DefaultHeight = static_cast<uint16>(MidValue);
+		const uint8 DefaultNormal = static_cast<uint8>(MidNormalValue);
+		return FColor(DefaultHeight >> 8, DefaultHeight & 255, DefaultNormal, DefaultNormal);
+	}
 };
 
 #if WITH_EDITOR
