Landscape Editor - Added condition to Weightmap->bForceMiplevelsToBeResident check in UpdatePhysicalMaterialTasks to handle cases where the ALandscape actor has been deleted from a WP Level
* After saving and reopening the level, ForceUpdateLayersContent will not be run leaving the remaining proxy textures unrequested

#Synced-CL 49563940

diff --git a/Engine/Source/Runtime/Landscape/Private/LandscapeEdit.cpp b/Engine/Source/Runtime/Landscape/Private/LandscapeEdit.cpp
--- a/Engine/Source/Runtime/Landscape/Private/LandscapeEdit.cpp
+++ b/Engine/Source/Runtime/Landscape/Private/LandscapeEdit.cpp
@@ -2172,13 +2172,20 @@ void ULandscapeComponent::UpdatePhysicalMaterialTasks()
 	{
 		// ULandscapeSubsystem can start ticking before the weightmaps are loaded.  They are "streamed in", but are still just the default placeholder texture.
 		// Rendering the physical material in this state will silently produce the wrong result.
-		bool bWeightMapsReady = Algo::AllOf(WeightmapTextures, [](UTexture2D* Weightmap)
-			{
-				// We expect FLandscapeTextureStreamingManager to have already started managing the data lifetime of the weightmaps, via UpdateLayersContent.
-				check(!Weightmap || Weightmap->bForceMiplevelsToBeResident);
+		bool bWeightMapsReady = false;
 
-				return Weightmap && !Weightmap->IsDefaultTexture() && !Weightmap->HasPendingInitOrStreaming() && Weightmap->IsFullyStreamedIn();
-			});
+		// WP levels may have a deleted ALandscape actor resulting in a non-editable landscape (Invalid LandscapeInfo). UpdateLayersContent will not run leaving bForceMiplevelsToBeResident unset for the remaining landscape components
+		const ULandscapeInfo* LandscapeInfo = GetLandscapeInfo();
+		if (LandscapeInfo && LandscapeInfo->LandscapeActor.IsValid())
+		{
+			bWeightMapsReady = Algo::AllOf(WeightmapTextures, [](UTexture2D* Weightmap)
+				{
+					// We expect FLandscapeTextureStreamingManager to have already started managing the data lifetime of the weightmaps, via UpdateLayersContent.
+					check(!Weightmap || Weightmap->bForceMiplevelsToBeResident);
+
+					return Weightmap && !Weightmap->IsDefaultTexture() && !Weightmap->HasPendingInitOrStreaming() && Weightmap->IsFullyStreamedIn();
+				});
+		}
 
 		TArray<UPhysicalMaterial*> PhysicalMaterials;
 		if (!bWeightMapsReady)
