* Landscape grass map exporter can now export an arbitrary amount of grass types, instead of always exporting all of those declared in the material
* Removed restriction that all landscape components must use the same list of grass types when generating grass map data for PCG
* Fixed invalid grass map being sampled when setting up a material with more than 1 null grass types

#Synced-CL 49641544

diff --git a/Engine/Plugins/PCG/Shaders/Private/PCGGrassMapUnpackerCS.usf b/Engine/Plugins/PCG/Shaders/Private/PCGGrassMapUnpackerCS.usf
--- a/Engine/Plugins/PCG/Shaders/Private/PCGGrassMapUnpackerCS.usf
+++ b/Engine/Plugins/PCG/Shaders/Private/PCGGrassMapUnpackerCS.usf
@@ -49,8 +49,8 @@ RWTexture2DArray<float> OutUnpackedGrassMaps;
 // Height output. Dimension is (NumComponentsX * InLandscapeComponentResolution, NumComponentsY * InLandscapeComponentResolution, NumGrassMaps).
 RWTexture2D<float> OutUnpackedHeight;
 
-StructuredBuffer<int> InHeightIndexMap;
-StructuredBuffer<int> InGrassIndexMap;
+StructuredBuffer<int> InComponentHeightOffsets;
+StructuredBuffer<int2> InComponentGrassmapOffsetAndChannels;
 int InNumGlobalGrassTypes;
 
 int4 InLinearTileIndexToComponentIndex[PCG_MAX_NUM_LANDSCAPE_COMPONENTS]; // Only x component used, 16 byte alignment on arrays
@@ -99,26 +99,20 @@ void PCGGrassMapUnpacker_CS(const uint3 WindowCoord_Slice : SV_DispatchThreadID)
 	// Heights. The first layer threads also deal with height.
 	if (InOutputHeight == 1 && Slice == 0)
 	{
-		const int HeightExportTileIndex = InHeightIndexMap[ComponentIndex];
-		if (HeightExportTileIndex != -1)
-		{
-			const uint ComponentXOffset = HeightExportTileIndex * InLandscapeComponentResolution;
-			const uint2 InputCoords = uint2(LocalTileCoords.x + ComponentXOffset, LocalTileCoords.y);
+		const int ComponentXOffset = InComponentHeightOffsets[ComponentIndex];
 
-			const float RelativeHeight = (UnpackHeight(InPackedGrassMaps[InputCoords].xy) - LANDSCAPE_MID_VALUE) * InLandscapeGridScale.z * TERRAIN_ZSCALE;
-			OutUnpackedHeight[WindowCoord] = RelativeHeight + InLandscapeLocationZ;
-		}
+		const uint2 InputCoords = uint2(LocalTileCoords.x + ComponentXOffset, LocalTileCoords.y);
+
+		const float RelativeHeight = (UnpackHeight(InPackedGrassMaps[InputCoords].xy) - LANDSCAPE_MID_VALUE) * InLandscapeGridScale.z * TERRAIN_ZSCALE;
+		OutUnpackedHeight[WindowCoord] = RelativeHeight + InLandscapeLocationZ;
 	}
 
 	// Coordinates into input height & grass maps (see packing above).
-	const int GrassExportChannelIndex = InGrassIndexMap[ComponentIndex * InNumGlobalGrassTypes + Slice];
-
-	if (GrassExportChannelIndex != -1)
+	const int2 ComponentAndGrassXOffsetAndChannel = InComponentGrassmapOffsetAndChannels[ComponentIndex * InNumGlobalGrassTypes + Slice];
+	if (ComponentAndGrassXOffsetAndChannel.x != -1)
 	{
-		const int GrassExportTileIndex = GrassExportChannelIndex / 4;
-		const uint ComponentXOffset = GrassExportTileIndex * InLandscapeComponentResolution;
-		// Grass map source channel (see packing above).
-		const uint ExportChannelIndex = GrassExportChannelIndex % 4;
+		const uint ComponentXOffset = ComponentAndGrassXOffsetAndChannel.x;
+		const uint ExportChannelIndex = ComponentAndGrassXOffsetAndChannel.y;
 
 		// Coordinates into packed data which are laid out horizontally (see above).
 		const uint2 InputCoords = uint2(LocalTileCoords.x + ComponentXOffset, LocalTileCoords.y);
diff --git a/Engine/Plugins/PCG/Source/PCG/Private/Elements/PCGGenerateLandscapeTextures.cpp b/Engine/Plugins/PCG/Source/PCG/Private/Elements/PCGGenerateLandscapeTextures.cpp
--- a/Engine/Plugins/PCG/Source/PCG/Private/Elements/PCGGenerateLandscapeTextures.cpp
+++ b/Engine/Plugins/PCG/Source/PCG/Private/Elements/PCGGenerateLandscapeTextures.cpp
@@ -20,6 +20,7 @@
 #include "LandscapeGrassType.h"
 #include "LandscapeGrassWeightExporter.h"
 #include "LandscapeProxy.h"
+#include "LandscapeUtils.h"
 #include "RenderCaptureInterface.h"
 #include "RenderGraphBuilder.h"
 #include "RenderGraphUtils.h"
@@ -250,7 +251,7 @@ bool FPCGGenerateLandscapeTexturesElement::ExecuteInternal(FPCGContext* InContex
 			return true;
 		}
 
-		if (!UE::Landscape::IsRuntimeGrassMapGenerationSupported() && (!ExecutionSource->GetExecutionState().GetWorld() || ExecutionSource->GetExecutionState().GetWorld()->IsGameWorld()))
+		if (!UE::Landscape::SupportsRuntimeGrassMapGeneration() && (!ExecutionSource->GetExecutionState().GetWorld() || ExecutionSource->GetExecutionState().GetWorld()->IsGameWorld()))
 		{
 			UE_LOG(LogPCG, Warning, TEXT("Grass map generation is disabled outside of editor worlds. Try enabling the CVar 'grass.GrassMap.AlwaysBuildRuntimeGenerationResources'."));
 			return true;
@@ -296,27 +297,27 @@ bool FPCGGenerateLandscapeTexturesElement::ExecuteInternal(FPCGContext* InContex
 				continue;
 			}
 
+			ALandscape* Landscape = LandscapeProxy->GetLandscapeActor();
+			if (!Landscape)
+			{
+				continue;
+			}
+
 			// Grow storage to accommodate new components.
 			const int32 NumToReserve = FMath::Min(OverlappingComponents.Num() + LandscapeProxy->LandscapeComponents.Num(), (int32)FPCGGrassMapUnpackerCS::MaxNumLandscapeComponents);
 			OverlappingComponents.Reserve(NumToReserve);
 			Context->LandscapeComponents.Reserve(NumToReserve);
 
-			if (!Context->FirstLandscapeProxy.Get())
+			if (!Context->Landscape.Get())
 			{
-				Context->FirstLandscapeProxy = LandscapeProxy;
+				Context->Landscape = Landscape;
 			}
-			else
+			// TODO [jonathan.bard] : The grass map exporter doesn't support landscape components from different landscapes, simply skip those components for now : 
+			else if (Context->Landscape.Get() != Landscape)
 			{
-				// Validate the proxies are compatible.
-				if (Context->FirstLandscapeProxy->ComponentSizeQuads != LandscapeProxy->ComponentSizeQuads
-					|| Context->FirstLandscapeProxy->SubsectionSizeQuads != LandscapeProxy->SubsectionSizeQuads
-					|| Context->FirstLandscapeProxy->NumSubsections != LandscapeProxy->NumSubsections
-					|| !Context->FirstLandscapeProxy->GetRootComponent()->GetRelativeScale3D().Equals(LandscapeProxy->GetRootComponent()->GetRelativeScale3D())
-					|| !FMath::IsNearlyEqual(Context->FirstLandscapeProxy->GetRootComponent()->GetRelativeLocation().Z, LandscapeProxy->GetRootComponent()->GetRelativeLocation().Z))
-				{
-					PCGLog::LogErrorOnGraph(LOCTEXT("LandscapeProxyMismatch", "Encountered multiple landscape proxies with different settings (component size quads or subsection sizes) or different transforms (scale XYZ or translation Z) which is not supported."), Context);
-					return true;
-				}
+				PCGLog::LogWarningOnGraph(FText::Format(LOCTEXT("LandscapeMismatch", "Encountered multiple landscapes which is not supported at the moment. Only the first landscape will be considered : {0} (current : {1})"), 
+					FText::FromString(Context->Landscape->GetActorNameOrLabel()), FText::FromString(Landscape->GetActorNameOrLabel())), Context);
+				continue;
 			}
 
 			for (int32 ComponentIndex = 0; ComponentIndex < LandscapeProxy->LandscapeComponents.Num(); ++ComponentIndex)
@@ -403,37 +404,11 @@ bool FPCGGenerateLandscapeTexturesElement::ExecuteInternal(FPCGContext* InContex
 			}
 		}
 
-		if (!Context->FirstLandscapeProxy.Get() || Context->LandscapeComponents.IsEmpty() || (Context->SelectedGrassTypes.IsEmpty() && !Settings->bGenerateHeightMap))
+		if (!Context->Landscape.Get() || Context->LandscapeComponents.IsEmpty() || (Context->SelectedGrassTypes.IsEmpty() && !Settings->bGenerateHeightMap))
 		{
 			return true;
 		}
 
-		// Set up index mapping so that export tile/channel index can be looked up from the component index and global grass map index or height index during unpack.
-		{
-			Context->IndexMapping.Init(OverlappingComponents.Num(), Context->GlobalGrassTypes.Num());
-
-			int32 NumExportedTiles = 0;
-
-			for (int32 OverlapIndex = 0; OverlapIndex < OverlappingComponents.Num(); ++OverlapIndex)
-			{
-				ULandscapeComponent* LandscapeComponent = OverlappingComponents[OverlapIndex];
-				check(LandscapeComponent);
-
-				// Register export data tile index that contains height data for this component.
-				Context->IndexMapping.SetHeightEntry(OverlapIndex, NumExportedTiles);
-
-				for (int32 ComponentGrassTypeIndex = 0; ComponentGrassTypeIndex < LandscapeComponent->GetGrassTypes().Num(); ++ComponentGrassTypeIndex)
-				{
-					// Register export data channel index that contains grass map data for this component & grass channel.
-					const int32 GlobalGrassMapIndex = Context->GlobalGrassTypes.IndexOfByKey(LandscapeComponent->GetGrassTypes()[ComponentGrassTypeIndex]);
-					Context->IndexMapping.SetGrassMapEntry(OverlapIndex, GlobalGrassMapIndex, NumExportedTiles, ComponentGrassTypeIndex);
-				}
-
-				// Account for export tiles added for this component.
-				NumExportedTiles += FMath::DivideAndRoundUp<int32>(LandscapeComponent->GetGrassTypes().Num() + /*Channels for height*/2, 4);
-			}
-		}
-
 		Context->bLandscapeComponentsFiltered = true;
 	}
 
@@ -538,20 +513,22 @@ bool FPCGGenerateLandscapeTexturesElement::ExecuteInternal(FPCGContext* InContex
 	// 3. Schedule grass map generation.
 	if (!Context->bGenerationScheduled)
 	{
-		ALandscapeProxy* LandscapeProxy = Context->FirstLandscapeProxy.IsValid() ? Context->FirstLandscapeProxy.Get() : nullptr;
-		if (!LandscapeProxy)
+		ALandscape* Landscape = Context->Landscape.IsValid() ? Context->Landscape.Get() : nullptr;
+		if (!Landscape)
 		{
 			PCGLog::LogErrorOnGraph(LOCTEXT("LandscapeProxyLost", "Reference to landscape proxy actor lost, grass maps will not be generated."), InContext);
 			return true;
 		}
 
-		const FVector LandscapeGridScale = LandscapeProxy->GetRootComponent()->GetRelativeScale3D();
-		const double LandscapeZ = LandscapeProxy->GetRootComponent()->GetRelativeLocation().Z;
+		const FVector LandscapeGridScale = Landscape->GetRootComponent()->GetRelativeScale3D();
+		const double LandscapeZ = Landscape->GetRootComponent()->GetRelativeLocation().Z;
 
 		TArray<ULandscapeComponent*> LandscapeComponents;
 		TArray<FIntVector2> LandscapeTileCoords;
 		LandscapeComponents.Reserve(Context->LandscapeComponents.Num());
-		LandscapeTileCoords.Reserve(Context->LandscapeComponents.Num());
+		// Sister array that stores the global identifier for a landscape component so that we can query info about this component in the grass map exporter
+		TArray<FIntPoint> LandscapeComponentKeys;
+		LandscapeComponentKeys.Reserve(Context->LandscapeComponents.Num());
 
 		for (TWeakObjectPtr<ULandscapeComponent> LandscapeComponentPtr : Context->LandscapeComponents)
 		{
@@ -570,16 +547,25 @@ bool FPCGGenerateLandscapeTexturesElement::ExecuteInternal(FPCGContext* InContex
 			LandscapeTileCoords.Add(FIntVector2(
 				(LandscapeComponent->Bounds.Origin.X - Context->TotalLandscapeComponentBounds.Min.X) / Context->LandscapeComponentExtent,
 				(LandscapeComponent->Bounds.Origin.Y - Context->TotalLandscapeComponentBounds.Min.Y) / Context->LandscapeComponentExtent));
+			LandscapeComponentKeys.Add(LandscapeComponent->GetComponentKey());
+		}
+
+		UE::Landscape::Grass::EGrassWeightExporterFlags Flags = UE::Landscape::Grass::EGrassWeightExporterFlags::None;
+		if (!Context->GlobalGrassTypes.IsEmpty())
+		{
+			Flags |= UE::Landscape::Grass::EGrassWeightExporterFlags::NeedsGrassmap;
+		}
+		if (Settings->bGenerateHeightMap)
+		{
+			Flags |= UE::Landscape::Grass::EGrassWeightExporterFlags::NeedsHeightmap;
 		}
 
 		Context->LandscapeGrassWeightExporter = new FLandscapeGrassWeightExporter(
-			LandscapeProxy,
+			Landscape, 
 			LandscapeComponents,
-			/*bInNeedsGrassmap=*/!Context->GlobalGrassTypes.IsEmpty(),
-			Settings->bGenerateHeightMap,
-			/*HeightMips=*/{},
-			/*bExport=*/false,
-			/*bShouldReadback=*/false);
+			Flags,
+			Context->GlobalGrassTypes, 
+			/*InRequestedHeightMips = */{});
 
 		const FIntPoint& OutputTextureSize = Context->LandscapeGrassWeightExporter->GetTargetSize();
 		const int32 Max2DTextureDimension = GetMax2DTextureDimension();
@@ -590,7 +576,7 @@ bool FPCGGenerateLandscapeTexturesElement::ExecuteInternal(FPCGContext* InContex
 			return true;
 		}
 
-		const int32 ComponentSizeQuads = LandscapeProxy->ComponentSizeQuads;
+		const int32 ComponentSizeQuads = Landscape->ComponentSizeQuads;
 		const double QuadSizeWorld = Context->LandscapeComponentExtent / ComponentSizeQuads;
 		// Each corner of a quad in the landscape corresponds to one texel in the grass map. 
 		const uint32 LandscapeComponentResolution = ComponentSizeQuads + 1;
@@ -648,7 +634,20 @@ bool FPCGGenerateLandscapeTexturesElement::ExecuteInternal(FPCGContext* InContex
 		
 		TWeakPtr<FPCGContextHandle> ContextHandle = Context->GetOrCreateHandle();
 
-		ENQUEUE_RENDER_COMMAND(PCGGenerateGrassMaps)([ContextHandle, LandscapeComponentExtent=Context->LandscapeComponentExtent, LandscapeTileCoords=MoveTemp(LandscapeTileCoords), LandscapeComponentResolution, LandscapeGridScale, LandscapeZ, bGenerateHeightMap=Settings->bGenerateHeightMap, NumTilesX, NumTilesY, OutputResolution, WindowStart, WindowEnd, IndexMapping=MoveTemp(Context->IndexMapping)](FRHICommandListImmediate& RHICmdList)
+		ENQUEUE_RENDER_COMMAND(PCGGenerateGrassMaps)(
+			[ ContextHandle
+			, LandscapeComponentExtent=Context->LandscapeComponentExtent
+			, LandscapeTileCoords = MoveTemp(LandscapeTileCoords)
+			, LandscapeComponentKeys = MoveTemp(LandscapeComponentKeys)
+			, LandscapeComponentResolution
+			, LandscapeGridScale
+			, LandscapeZ
+			, bGenerateHeightMap = Settings->bGenerateHeightMap
+			, NumTilesX
+			, NumTilesY
+			, OutputResolution
+			, WindowStart
+			, WindowEnd](FRHICommandListImmediate& RHICmdList)
 		{
 			LLM_SCOPE_BYTAG(PCG);
 
@@ -699,8 +698,6 @@ bool FPCGGenerateLandscapeTexturesElement::ExecuteInternal(FPCGContext* InContex
 			// Unpack the generated results to simple world aligned textures.
 			FPCGGrassMapUnpackerCS::FParameters* Parameters = GraphBuilder.AllocParameters<FPCGGrassMapUnpackerCS::FParameters>();
 			Parameters->InPackedGrassMaps = GrassMapTextureSRV;
-			Parameters->InHeightIndexMap = GraphBuilder.CreateSRV(IndexMapping.CreateHeightIndexMapBuffer(GraphBuilder));
-			Parameters->InGrassIndexMap = GraphBuilder.CreateSRV(IndexMapping.CreateGrassIndexMapBuffer(GraphBuilder));
 			Parameters->OutUnpackedGrassMaps = GrassMapUAV;
 			Parameters->OutUnpackedHeight = GrassHeightMapUAV;
 			Parameters->InNumTiles = FUintVector2(NumTilesX, NumTilesY);
@@ -709,7 +706,7 @@ bool FPCGGenerateLandscapeTexturesElement::ExecuteInternal(FPCGContext* InContex
 			Parameters->InLandscapeComponentResolution = LandscapeComponentResolution;
 			Parameters->InOutputHeight = bGenerateHeightMap ? 1u : 0u;
 			Parameters->InWindow = FUintVector4(WindowStart.X, WindowStart.Y, WindowEnd.X, WindowEnd.Y);
-			Parameters->InNumGlobalGrassTypes = IndexMapping.GetNumGlobalGrassTypes();
+			Parameters->InNumGlobalGrassTypes = Context->GlobalGrassTypes.Num();
 
 			// Initialize to invalid component indices.
 			for (int Index = 0; Index < FPCGGrassMapUnpackerCS::MaxNumLandscapeComponents; ++Index)
@@ -718,10 +715,75 @@ bool FPCGGenerateLandscapeTexturesElement::ExecuteInternal(FPCGContext* InContex
 			}
 
 			// Now write component mapping.
+			check(LandscapeTileCoords.Num() <= FPCGGrassMapUnpackerCS::MaxNumLandscapeComponents);
 			for (int Index = 0; Index < LandscapeTileCoords.Num(); ++Index)
 			{
-				Parameters->InLinearTileIndexToComponentIndex[LandscapeTileCoords[Index].Y * NumTilesX + LandscapeTileCoords[Index].X].X = Index;
+				const int32 LinearTileIndex = LandscapeTileCoords[Index].Y * NumTilesX + LandscapeTileCoords[Index].X;
+				check((LinearTileIndex >= 0) && (LinearTileIndex < FPCGGrassMapUnpackerCS::MaxNumLandscapeComponents));
+				Parameters->InLinearTileIndexToComponentIndex[LinearTileIndex].X = Index;
+			}
+			check(LandscapeTileCoords.Num() == LandscapeComponentKeys.Num());
+
+			FRDGBufferRef ComponentHeightOffsetsBuffer;
+			if (bGenerateHeightMap)
+			{
+				// X offset for where to find each component's height data in the grass texture (always in the .xy channels)
+				TArray<int32> ComponentHeightOffsets;
+				ComponentHeightOffsets.Reserve(LandscapeComponentKeys.Num());
+				for (int ComponentIndex = 0; ComponentIndex < LandscapeComponentKeys.Num(); ++ComponentIndex)
+				{
+					TValueOrError<FIntRect /*OutputRect*/, FLandscapeGrassWeightExporter_RenderThread::EHeightmapTextureInfoError> TextureInfo = 
+						Context->LandscapeGrassWeightExporter->GetTextureInfoForHeight(LandscapeComponentKeys[ComponentIndex]);
+					// Not supposed to fail : 
+					checkf(!TextureInfo.HasError(), TEXT("Failed to retrieve height map info on component %s : reason : %i"), *LandscapeComponentKeys[ComponentIndex].ToString(), TextureInfo.GetError()); 
+					ComponentHeightOffsets.Add(TextureInfo.GetValue().Min.X);
+				}
+				check(ComponentHeightOffsets.Num() == LandscapeComponentKeys.Num());
+				ComponentHeightOffsetsBuffer = CreateStructuredBuffer(GraphBuilder, TEXT("PCGGenerateLandscapeTexturesComponentHeightOffsets"), ComponentHeightOffsets);
+			}
+			else
+			{
+				ComponentHeightOffsetsBuffer = GSystemTextures.GetDefaultStructuredBuffer(GraphBuilder, sizeof(int32));
 			}
+			Parameters->InComponentHeightOffsets = GraphBuilder.CreateSRV(ComponentHeightOffsetsBuffer);
+
+			FRDGBufferRef ComponentGrassmapOffsetAndChannelsBuffer;
+			if (!Context->GlobalGrassTypes.IsEmpty())
+			{
+				// For each component, for each grass type, X offset for where to find the grass data in the grass texture (.x) and on which channel (.y)
+				//  It's possible a grass map is absent from a given component, in which case, the offset and channel will be set to -1
+				TArray<FIntVector2> ComponentGrassmapOffsetAndChannels;
+				ComponentGrassmapOffsetAndChannels.Reserve(LandscapeComponentKeys.Num() * Context->GlobalGrassTypes.Num());
+				for (int32 ComponentIndex = 0; ComponentIndex < LandscapeComponentKeys.Num(); ++ComponentIndex)
+				{
+					for (int32 GrassmapIndex = 0; GrassmapIndex < Context->GlobalGrassTypes.Num(); ++GrassmapIndex)
+					{
+						FIntVector2& ComponentGrassmapOffsetAndChannel = ComponentGrassmapOffsetAndChannels.Add_GetRef(FIntVector2(-1, -1));
+
+						TValueOrError<TPair<FIntRect /*OutputRect*/, uint8 /*ChannelIndex*/>, FLandscapeGrassWeightExporter_RenderThread::EGrassmapTextureInfoError> TextureInfo =
+							Context->LandscapeGrassWeightExporter->GetTextureInfoForGrass(LandscapeComponentKeys[ComponentIndex], Context->GlobalGrassTypes[GrassmapIndex]);
+						if (TextureInfo.HasError())
+						{
+							// It can fail, but only if that's because the grass map is absent from this component
+							checkf(TextureInfo.GetError() == FLandscapeGrassWeightExporter_RenderThread::EGrassmapTextureInfoError::InvalidGrassType,
+								TEXT("Failed to retrieve grass map info on component %s : reason : %i"), *LandscapeComponentKeys[ComponentIndex].ToString(), TextureInfo.GetError());
+						}
+						else
+						{
+							TPair<FIntRect /*OutputRect*/, uint8 /*ChannelIndex*/> TextureInfoValue = TextureInfo.GetValue();
+							check((TextureInfoValue.Value >= 0) && (TextureInfoValue.Value < 4));
+							ComponentGrassmapOffsetAndChannel =	FIntVector2(TextureInfoValue.Key.Min.X, static_cast<int32>(TextureInfoValue.Value));
+						}
+					}
+				}
+				check(ComponentGrassmapOffsetAndChannels.Num() == LandscapeComponentKeys.Num() * Context->GlobalGrassTypes.Num());
+				ComponentGrassmapOffsetAndChannelsBuffer = CreateStructuredBuffer(GraphBuilder, TEXT("PCGGenerateLandscapeTexturesComponentGrassmapOffsetAndChannels"), ComponentGrassmapOffsetAndChannels);
+			}
+			else
+			{
+				ComponentGrassmapOffsetAndChannelsBuffer = GSystemTextures.GetDefaultStructuredBuffer(GraphBuilder, sizeof(FIntVector2));;
+			}
+			Parameters->InComponentGrassmapOffsetAndChannels = GraphBuilder.CreateSRV(ComponentGrassmapOffsetAndChannelsBuffer);
 
 			TShaderMapRef<FPCGGrassMapUnpackerCS> Shader(GetGlobalShaderMap(GMaxRHIFeatureLevel));
 			const int GroupCountX = FMath::DivideAndRoundUp<int>(OutputResolution.X, FPCGGrassMapUnpackerCS::ThreadGroupDim);
@@ -922,60 +984,4 @@ bool FPCGGenerateLandscapeTexturesElement::ExecuteInternal(FPCGContext* InContex
 	return true;
 }
 
-void FPCGGenerateLandscapeTexturesContext::FGrassIndexMapping::SetHeightEntry(int32 InComponentIndex, int32 InFirstExportTileIndex)
-{
-	if (ensure(HeightTileIndexMap.IsValidIndex(InComponentIndex)))
-	{
-		HeightTileIndexMap[InComponentIndex] = InFirstExportTileIndex;
-	}
-}
-
-void FPCGGenerateLandscapeTexturesContext::FGrassIndexMapping::SetGrassMapEntry(int32 InComponentIndex, int32 InGlobalGrassMapIndex, int32 InFirstExportTileIndex, int32 InComponentGrassTypeIndex)
-{
-	const int32 NumChannels = 4; // Grass export produces RGBA texture.
-	const int32 FirstGrassMapChannel = 2; // First two channels used for landscape heights.
-	const int32 GlobalChannelIndex = InFirstExportTileIndex * NumChannels + InComponentGrassTypeIndex + FirstGrassMapChannel;
-	const int32 MapEntryIndex = InComponentIndex * NumGlobalGrassTypes + InGlobalGrassMapIndex;
-
-	// Expected that initialize method is called first.
-	if (ensure(GrassChannelIndexMap.IsValidIndex(MapEntryIndex)))
-	{
-		GrassChannelIndexMap[MapEntryIndex] = GlobalChannelIndex;
-	}
-}
-
-FRDGBufferRef FPCGGenerateLandscapeTexturesContext::FGrassIndexMapping::CreateHeightIndexMapBuffer(FRDGBuilder& GraphBuilder) const
-{
-	if (!HeightTileIndexMap.IsEmpty())
-	{
-		FRDGBufferDesc Desc = FRDGBufferDesc::CreateStructuredDesc(HeightTileIndexMap.GetTypeSize(), HeightTileIndexMap.Num());
-		FRDGBufferRef Buffer = GraphBuilder.CreateBuffer(Desc, TEXT("PCGGenerateLandscapeTexturesHeightIndexMap"));
-
-		GraphBuilder.QueueBufferUpload(Buffer, MakeConstArrayView(HeightTileIndexMap));
-
-		return Buffer;
-	}
-	else
-	{
-		return GSystemTextures.GetDefaultStructuredBuffer(GraphBuilder, sizeof(int32));
-	}
-}
-
-FRDGBufferRef FPCGGenerateLandscapeTexturesContext::FGrassIndexMapping::CreateGrassIndexMapBuffer(FRDGBuilder& GraphBuilder) const
-{
-	if (!GrassChannelIndexMap.IsEmpty())
-	{
-		FRDGBufferDesc Desc = FRDGBufferDesc::CreateStructuredDesc(GrassChannelIndexMap.GetTypeSize(), GrassChannelIndexMap.Num());
-		FRDGBufferRef Buffer = GraphBuilder.CreateBuffer(Desc, TEXT("PCGGenerateLandscapeTexturesGrassIndexMap"));
-
-		GraphBuilder.QueueBufferUpload(Buffer, MakeConstArrayView(GrassChannelIndexMap));
-
-		return Buffer;
-	}
-	else
-	{
-		return GSystemTextures.GetDefaultStructuredBuffer(GraphBuilder, sizeof(int32));
-	}
-}
-
 #undef LOCTEXT_NAMESPACE
diff --git a/Engine/Plugins/PCG/Source/PCG/Public/Elements/PCGGenerateLandscapeTextures.h b/Engine/Plugins/PCG/Source/PCG/Public/Elements/PCGGenerateLandscapeTextures.h
--- a/Engine/Plugins/PCG/Source/PCG/Public/Elements/PCGGenerateLandscapeTextures.h
+++ b/Engine/Plugins/PCG/Source/PCG/Public/Elements/PCGGenerateLandscapeTextures.h
@@ -9,6 +9,7 @@
 
 #include "PCGGenerateLandscapeTextures.generated.h"
 
+class ALandscape;
 class ALandscapeProxy;
 class FLandscapeGrassWeightExporter;
 class ULandscapeComponent;
@@ -104,41 +105,8 @@ public:
 	TArray<TTuple<TWeakObjectPtr<ULandscapeGrassType>, /*TextureIndex=*/int32>> SelectedGrassTypes;
 
 	TArray<TObjectPtr<ULandscapeGrassType>> GlobalGrassTypes;
-	
-	/** Helper for building index map from global height map / grass map indices to the tile & channel in the data from the grass map exporter. */
-	struct FGrassIndexMapping
-	{
-		void Init(int32 InNumComponents, int32 InNumGlobalGrassTypes)
-		{
-			NumGlobalGrassTypes = InNumGlobalGrassTypes;
-			GrassChannelIndexMap.SetNumUninitialized(InNumComponents * InNumGlobalGrassTypes);
-			FMemory::Memset(GrassChannelIndexMap.GetData(), -1, GrassChannelIndexMap.Num() * GrassChannelIndexMap.GetTypeSize());
 
-			HeightTileIndexMap.SetNumZeroed(InNumComponents);
-		}
-
-		int32 GetNumGlobalGrassTypes() const { return NumGlobalGrassTypes; }
-
-		void SetHeightEntry(int32 InComponentIndex, int32 InFirstExportTileIndex);
-		void SetGrassMapEntry(int32 InComponentIndex, int32 InGlobalGrassMapIndex, int32 InFirstExportTileIndex, int32 InComponentGrassTypeIndex);
-
-		FRDGBufferRef CreateHeightIndexMapBuffer(FRDGBuilder& GraphBuilder) const;
-		FRDGBufferRef CreateGrassIndexMapBuffer(FRDGBuilder& GraphBuilder) const;
-
-	private:
-		// For each component, global tile index containing height data. Height data always packed into first 2 channels of tile so don't need to track channel index.
-		TArray<int32> HeightTileIndexMap;
-
-		// For each component, for each possible global grass map index, store the global channel index in the export data.
-		// To get export tile index, divide this value by 4.
-		TArray<int32> GrassChannelIndexMap;
-
-		int32 NumGlobalGrassTypes = -1;
-	};
-
-	FGrassIndexMapping IndexMapping;
-
-	TWeakObjectPtr<ALandscapeProxy> FirstLandscapeProxy = nullptr;
+	TWeakObjectPtr<ALandscape> Landscape;
 	TArray<TWeakObjectPtr<ULandscapeComponent>> LandscapeComponents;
 
 	/** World-space bounds containing all of the landscape components given to the grass weight exporter. */
@@ -173,7 +141,7 @@ protected:
 public:
 	UE_DEPRECATED(5.8, "No longer used")
 	int32 NumGrassTypes = 0;
-	UE_DEPRECATED(5.8, "Renamed to FirstLandscapeProxy")
+	UE_DEPRECATED(5.8, "Use Landscape")
 	TWeakObjectPtr<ALandscapeProxy> LandscapeProxy = nullptr;
 };
 
diff --git a/Engine/Plugins/PCG/Source/PCGCompute/Internal/PCGGrassMapUnpackerCS.h b/Engine/Plugins/PCG/Source/PCGCompute/Internal/PCGGrassMapUnpackerCS.h
--- a/Engine/Plugins/PCG/Source/PCGCompute/Internal/PCGGrassMapUnpackerCS.h
+++ b/Engine/Plugins/PCG/Source/PCGCompute/Internal/PCGGrassMapUnpackerCS.h
@@ -20,8 +20,8 @@ public:
 
 	BEGIN_SHADER_PARAMETER_STRUCT(FParameters, PCGCOMPUTE_API)
 		SHADER_PARAMETER_RDG_TEXTURE_SRV(Texture2D<float4>, InPackedGrassMaps)
-		SHADER_PARAMETER_RDG_BUFFER_SRV(StructuredBuffer<int32>, InHeightIndexMap)
-		SHADER_PARAMETER_RDG_BUFFER_SRV(StructuredBuffer<int32>, InGrassIndexMap)
+		SHADER_PARAMETER_RDG_BUFFER_SRV(StructuredBuffer<int32>, InComponentHeightOffsets)
+		SHADER_PARAMETER_RDG_BUFFER_SRV(StructuredBuffer<FIntVector2>, InComponentGrassmapOffsetAndChannels)
 		SHADER_PARAMETER_RDG_TEXTURE_UAV(RWTexture2D<float>, OutUnpackedHeight)
 		SHADER_PARAMETER_RDG_TEXTURE_UAV(RWTexture2DArray<float>, OutUnpackedGrassMaps)
 		SHADER_PARAMETER_ARRAY(FIntVector4, InLinearTileIndexToComponentIndex, [MaxNumLandscapeComponents])
diff --git a/Engine/Shaders/Private/LandscapeGrassWeight.usf b/Engine/Shaders/Private/LandscapeGrassWeight.usf
--- a/Engine/Shaders/Private/LandscapeGrassWeight.usf
+++ b/Engine/Shaders/Private/LandscapeGrassWeight.usf
@@ -53,7 +53,13 @@ void VSMain(
 
 #elif PIXELSHADER
 
-int OutputPass;
+int4 PerChannelGrassWeightIndices; // One index per channel in the OutColor. 0 indicates the channel exports grass weight 0, 1 for grass weight 1, etc. 
+
+// 1 special index to indicate we don't want to export anything
+#define GrassWeightInvalidIndex -3
+// 2 special indices to indicate we want to export heights
+#define HeightGrassWeightIndex0 -2
+#define HeightGrassWeightIndex1 -1
 
 // Pixel Shader
 void PSMain( 
@@ -81,153 +87,136 @@ void PSMain(
 #endif // VF_USE_PRIMITIVE_SCENE_DATA
 	float2 PackedHeight = float2(floor(Height / 256), fmod(Height, 256)) / 255;
 
-	// Height = first 2 outputs:
-	OutColor = float4(PackedHeight, 0.0, 0.0);
+	OutColor = 0.0;
 
-	switch(OutputPass)
-	{
-	case 0:
+#define GRASSWEIGHTCASE(N) case N: \
+	{ \
+		OutColor[ChannelIndex] = GetGrassWeight##N(MaterialParameters).x; \
+		break; \
+	}
+
+	UNROLL
+	for (uint ChannelIndex = 0; ChannelIndex < 4; ++ChannelIndex)
 	{
+		switch(PerChannelGrassWeightIndices[ChannelIndex])
+		{
+			case GrassWeightInvalidIndex: 
+				break;
+			case HeightGrassWeightIndex0: 
+			{
+				OutColor[ChannelIndex] = PackedHeight.x; 
+				break;
+			}
+			case HeightGrassWeightIndex1: 
+			{
+				OutColor[ChannelIndex] = PackedHeight.y; 
+				break;
+			}
 #if HAVE_GetGrassWeight0
-		OutColor[2] = GetGrassWeight0(MaterialParameters).x;
+			GRASSWEIGHTCASE(0);
 #endif // HAVE_GetGrassWeight0
 #if HAVE_GetGrassWeight1
-		OutColor[3] = GetGrassWeight1(MaterialParameters).x;
+			GRASSWEIGHTCASE(1);
 #endif // HAVE_GetGrassWeight1
-		break;
-	}
-	case 1:
-	{
 #if HAVE_GetGrassWeight2
-		OutColor[0] = GetGrassWeight2(MaterialParameters).x;
+			GRASSWEIGHTCASE(2);
 #endif // HAVE_GetGrassWeight2
 #if HAVE_GetGrassWeight3
-		OutColor[1] = GetGrassWeight3(MaterialParameters).x;
+			GRASSWEIGHTCASE(3);
 #endif // HAVE_GetGrassWeight3
 #if HAVE_GetGrassWeight4
-		OutColor[2] = GetGrassWeight4(MaterialParameters).x;
+			GRASSWEIGHTCASE(4);
 #endif // HAVE_GetGrassWeight4
 #if HAVE_GetGrassWeight5
-		OutColor[3] = GetGrassWeight5(MaterialParameters).x;
+			GRASSWEIGHTCASE(5);
 #endif // HAVE_GetGrassWeight5
-		break;
-	}
-	case 2:
-	{
 #if HAVE_GetGrassWeight6
-		OutColor[0] = GetGrassWeight6(MaterialParameters).x;
+			GRASSWEIGHTCASE(6);
 #endif // HAVE_GetGrassWeight6
 #if HAVE_GetGrassWeight7
-		OutColor[1] = GetGrassWeight7(MaterialParameters).x;
+			GRASSWEIGHTCASE(7);
 #endif // HAVE_GetGrassWeight7
 #if HAVE_GetGrassWeight8
-		OutColor[2] = GetGrassWeight8(MaterialParameters).x;
+			GRASSWEIGHTCASE(8);
 #endif // HAVE_GetGrassWeight8
 #if HAVE_GetGrassWeight9
-		OutColor[3] = GetGrassWeight9(MaterialParameters).x;
+			GRASSWEIGHTCASE(9);
 #endif // HAVE_GetGrassWeight9
-		break;
-	}
-	case 3:
-	{
+
 #if HAVE_GetGrassWeight10
-		OutColor[0] = GetGrassWeight10(MaterialParameters).x;
+			GRASSWEIGHTCASE(10);
 #endif // HAVE_GetGrassWeight10
 #if HAVE_GetGrassWeight11
-		OutColor[1] = GetGrassWeight11(MaterialParameters).x;
+			GRASSWEIGHTCASE(11);
 #endif // HAVE_GetGrassWeight11
 #if HAVE_GetGrassWeight12
-		OutColor[2] = GetGrassWeight12(MaterialParameters).x;
+			GRASSWEIGHTCASE(12);
 #endif // HAVE_GetGrassWeight12
 #if HAVE_GetGrassWeight13
-		OutColor[3] = GetGrassWeight13(MaterialParameters).x;
+			GRASSWEIGHTCASE(13);
 #endif // HAVE_GetGrassWeight13
-		break;
-	}
-	case 4:
-	{
 #if HAVE_GetGrassWeight14
-		OutColor[0] = GetGrassWeight14(MaterialParameters).x;
+			GRASSWEIGHTCASE(14);
 #endif // HAVE_GetGrassWeight14
 #if HAVE_GetGrassWeight15
-		OutColor[1] = GetGrassWeight15(MaterialParameters).x;
+			GRASSWEIGHTCASE(15);
 #endif // HAVE_GetGrassWeight15
 #if HAVE_GetGrassWeight16
-		OutColor[2] = GetGrassWeight16(MaterialParameters).x;
+			GRASSWEIGHTCASE(16);
 #endif // HAVE_GetGrassWeight16
 #if HAVE_GetGrassWeight17
-		OutColor[3] = GetGrassWeight17(MaterialParameters).x;
+			GRASSWEIGHTCASE(17);
 #endif // HAVE_GetGrassWeight17
-		break;
-	}
-	case 5:
-	{
 #if HAVE_GetGrassWeight18
-		OutColor[0] = GetGrassWeight18(MaterialParameters).x;
+			GRASSWEIGHTCASE(18);
 #endif // HAVE_GetGrassWeight18
 #if HAVE_GetGrassWeight19
-		OutColor[1] = GetGrassWeight19(MaterialParameters).x;
+			GRASSWEIGHTCASE(19);
 #endif // HAVE_GetGrassWeight19
+
 #if HAVE_GetGrassWeight20
-		OutColor[2] = GetGrassWeight20(MaterialParameters).x;
+			GRASSWEIGHTCASE(20);
 #endif // HAVE_GetGrassWeight20
 #if HAVE_GetGrassWeight21
-		OutColor[3] = GetGrassWeight21(MaterialParameters).x;
+			GRASSWEIGHTCASE(21);
 #endif // HAVE_GetGrassWeight21
-		break;
-	}
-	case 6:
-	{
 #if HAVE_GetGrassWeight22
-		OutColor[0] = GetGrassWeight22(MaterialParameters).x;
+			GRASSWEIGHTCASE(22);
 #endif // HAVE_GetGrassWeight22
 #if HAVE_GetGrassWeight23
-		OutColor[1] = GetGrassWeight23(MaterialParameters).x;
+			GRASSWEIGHTCASE(23);
 #endif // HAVE_GetGrassWeight23
 #if HAVE_GetGrassWeight24
-		OutColor[2] = GetGrassWeight24(MaterialParameters).x;
+			GRASSWEIGHTCASE(24);
 #endif // HAVE_GetGrassWeight24
 #if HAVE_GetGrassWeight25
-		OutColor[3] = GetGrassWeight25(MaterialParameters).x;
+			GRASSWEIGHTCASE(25);
 #endif // HAVE_GetGrassWeight25
-		break;
-	}
-	case 7:
-	{
 #if HAVE_GetGrassWeight26
-		OutColor[0] = GetGrassWeight26(MaterialParameters).x;
+			GRASSWEIGHTCASE(26);
 #endif // HAVE_GetGrassWeight26
 #if HAVE_GetGrassWeight27
-		OutColor[1] = GetGrassWeight27(MaterialParameters).x;
+			GRASSWEIGHTCASE(27);
 #endif // HAVE_GetGrassWeight27
 #if HAVE_GetGrassWeight28
-		OutColor[2] = GetGrassWeight28(MaterialParameters).x;
+			GRASSWEIGHTCASE(28);
 #endif // HAVE_GetGrassWeight28
 #if HAVE_GetGrassWeight29
-		OutColor[3] = GetGrassWeight29(MaterialParameters).x;
+			GRASSWEIGHTCASE(29);
 #endif // HAVE_GetGrassWeight29
-		break;
-	}
-	case 8:
-	{
+
 #if HAVE_GetGrassWeight30
-		OutColor[0] = GetGrassWeight30(MaterialParameters).x;
+			GRASSWEIGHTCASE(30);
 #endif // HAVE_GetGrassWeight30
 #if HAVE_GetGrassWeight31
-		OutColor[1] = GetGrassWeight31(MaterialParameters).x;
+			GRASSWEIGHTCASE(31);
 #endif // HAVE_GetGrassWeight31
-		break;
-
-// UMaterialExpressionLandscapeGrassOutput::MaxGrassTypes == 32. Uncomment this if it ever gets increased :
-//#if HAVE_GetGrassWeight32
-//		OutColor[2] = GetGrassWeight32(MaterialParameters);
-//#endif // HAVE_GetGrassWeight32
-//#if HAVE_GetGrassWeight33
-//		OutColor[3] = GetGrassWeight33(MaterialParameters);
-//#endif // HAVE_GetGrassWeight33
-	}
-	default:
-		break;
+
+			default:
+				break;
+		}
 	}
+
+	// UMaterialExpressionLandscapeGrassOutput::MaxGrassTypes == 32. Keep on going if this gets increased
 }
 #endif
\ No newline at end of file
diff --git a/Engine/Source/Runtime/Engine/Private/Materials/MaterialCachedData.cpp b/Engine/Source/Runtime/Engine/Private/Materials/MaterialCachedData.cpp
--- a/Engine/Source/Runtime/Engine/Private/Materials/MaterialCachedData.cpp
+++ b/Engine/Source/Runtime/Engine/Private/Materials/MaterialCachedData.cpp
@@ -627,7 +627,7 @@ void FMaterialCachedExpressionData::UpdateForExpressions(const FMaterialCachedEx
 		{
 			for (const auto& Type : ExpressionGrassOutput->GrassTypes)
 			{
-				GrassTypes.AddUnique(Type.GrassType);
+				GrassTypes.Add(Type.GrassType);
 			}
 		}
 		else if (UMaterialExpressionQualitySwitch* QualitySwitchNode = Cast<UMaterialExpressionQualitySwitch>(Expression))
diff --git a/Engine/Source/Runtime/Landscape/Internal/LandscapeGrassWeightExporter.h b/Engine/Source/Runtime/Landscape/Internal/LandscapeGrassWeightExporter.h
--- a/Engine/Source/Runtime/Landscape/Internal/LandscapeGrassWeightExporter.h
+++ b/Engine/Source/Runtime/Landscape/Internal/LandscapeGrassWeightExporter.h
@@ -3,8 +3,10 @@
 
 #include "CoreMinimal.h"
 #include "LandscapeComponent.h"
+#include "Templates/ValueOrError.h"
 #include "MeshMaterialShader.h"
 
+class ALandscape;
 class ALandscapeProxy;
 class FLandscapeAsyncTextureReadback;
 class FLandscapeComponentSceneProxy;
@@ -17,9 +19,29 @@ namespace UE::Landscape
 {
 	UE_DEPRECATED(5.8, "Use ULandscapeComponent::CanRenderGrassMap instead")
 	LANDSCAPE_API bool CanRenderGrassMap(ULandscapeComponent* Component);
+	UE_DEPRECATED(5.8, "Renamed to SupportsRuntimeGrassMapGeneration (LandscapeUtils.h)")
 	LANDSCAPE_API bool IsRuntimeGrassMapGenerationSupported();
-}
+} // UE::Landscape
 
+namespace UE::Landscape::Grass
+{
+	void AddGrassWeightShaderTypes(FMaterialShaderTypes& InOutShaderTypes);
+
+	enum class EGrassWeightExporterFlags
+	{
+		None = 0,
+		NeedsGrassmap = (1 << 0),
+		NeedsHeightmap = (1 << 1),
+		RenderImmediately = (1 << 2),
+		ReadbackToCPU = (1 << 3),
+		Default = NeedsGrassmap | NeedsHeightmap | RenderImmediately | ReadbackToCPU,
+	};
+	ENUM_CLASS_FLAGS(EGrassWeightExporterFlags);
+} // UE::Landscape::Grass
+
+
+// ----------------------------------------------------------------------------------
+// 
 // Hacky base class to avoid 8 bytes of padding after the vtable
 class FLandscapeGrassWeightExporter_RenderThread_FixLayout
 {
@@ -32,86 +54,147 @@ class FLandscapeGrassWeightExporter_RenderThread : public FLandscapeGrassWeightE
 {
 	friend class FLandscapeGrassWeightExporter;
 	friend class FLandscapeGrassMapsBuilder;
+	friend class FComponentRenderInfo;
 
 private:
-	FLandscapeGrassWeightExporter_RenderThread(const TArray<int32>& InHeightMips, bool bInReadbackToCPU = true);
+	FLandscapeGrassWeightExporter_RenderThread(int32 InComponentSizeVerts, int32 InSubsectionSizeQuads, int32 InNumSubsections, 
+		TConstArrayView<int32> InHeightMips, UE::Landscape::Grass::EGrassWeightExporterFlags InFlags);
 
 public:
 	virtual ~FLandscapeGrassWeightExporter_RenderThread();
 
 	const FIntPoint& GetTargetSize() const { return TargetSize; }
 
-	/** Renders the component to the given texture. */
+	int32 GetTileSize() const
+	{
+		return ComponentSizeVerts;
+	}
+
+	/** 
+	 * Renders the components to the given texture. 
+	 */
 	LANDSCAPE_API void RenderLandscapeComponentToTexture_RenderThread(FRDGBuilder& GraphBuilder, FRDGTextureRef OutputTexture);
 
+	enum class EHeightmapTextureInfoError
+	{
+		InvalidComponentKey,
+		InvalidFlags,
+		InvalidMipIndex,
+	};
+	/**
+	 * Return the location of the height data for the requested component / mip in the output texture
+	 * @param InComponentKey Key of the component being requested
+	 * @param InMipIndex Requested mip index. If < 0, will return mip 0. Otherwise, the mip needs to have been requested in HeightMips
+	 * 
+	 * @return If the component / mip data is found, FIntRect contains the pixels containing the requested height data in the generated texture (channels R and G). 
+	 *  Returns an error otherwise.
+	 */
+	LANDSCAPE_API TValueOrError<FIntRect /*OutputRect*/, EHeightmapTextureInfoError> GetTextureInfoForHeight(const FIntPoint& InComponentKey, int32 InMipIndex = INDEX_NONE);
+
+	enum class EGrassmapTextureInfoError
+	{
+		InvalidComponentKey, 
+		InvalidFlags,
+		InvalidGrassType,
+	};
+	/**
+	 * Return the location of the grass data for the requested component in the output texture
+	 * @param InComponentKey Key of the component being requested
+	 * @param InGrassType 
+	 * 
+	 * @return If the component / grass data is found, a FIntRect that contains the pixels containing the requested grass data in the generated texture, along with the channel index. 
+	 *  Returns an error otherwise.
+	 */
+	// TODO [jonathan.bard] : replace grass type with FName
+	LANDSCAPE_API TValueOrError<TPair<FIntRect /*OutputRect*/, uint8 /*ChannelIndex*/>, EGrassmapTextureInfoError> GetTextureInfoForGrass(const FIntPoint& InComponentKey, ULandscapeGrassType* InGrassType);
+
+private:
+	/** Creates a texture and renders the component to it, and then and triggers a readback of the texture. */
+	void RenderLandscapeComponentToTexture_RenderThread(FRHICommandListImmediate& RHICmdList, FLandscapeAsyncTextureReadback* AsyncReadbackPtr);
+
 private:
 	struct FComponentInfo
 	{
-		TObjectPtr<ULandscapeComponent> Component = nullptr;
-		TArray<TObjectPtr<ULandscapeGrassType>> RequestedGrassTypes;
-		FVector2D ViewOffset = FVector2D::ZeroVector;
-		int32 PixelOffsetX = 0;
-		FLandscapeComponentSceneProxy* SceneProxy = nullptr;
-		int32 NumPasses = 0;
-		int32 FirstHeightMipsPassIndex = MAX_int32;
+		FComponentInfo(ULandscapeComponent* InComponent, UE::Landscape::Grass::EGrassWeightExporterFlags InFlags,
+			TConstArrayView<TObjectPtr<ULandscapeGrassType>> InRequestedGrassTypes, TConstArrayView<int32> InRequestedHeightMips);
 
-		FComponentInfo(ULandscapeComponent* InComponent, bool bInNeedsGrassmap, bool bInNeedsHeightmap, const TArray<int32>& InHeightMips)
-			: Component(InComponent)
-			, SceneProxy((FLandscapeComponentSceneProxy*)InComponent->SceneProxy)
+		int32 GetNumPasses() const
 		{
-			if (bInNeedsGrassmap)
-			{
-				RequestedGrassTypes = InComponent->GetGrassTypes();
-			}
-			int32 NumGrassMaps = RequestedGrassTypes.Num();
-			if (bInNeedsHeightmap || NumGrassMaps > 0)
+			return PerPassPerChannelGrassWeightIndices.Num();
+		}
+
+		bool HasMipHeightData() const
+		{
+			return FirstHeightMipsPassIndex != MAX_int32;
+		}
+
+		struct FPassChannelIndex
+		{
+			int32 PassIndex = INDEX_NONE;
+			int32 ChannelIndex = INDEX_NONE;
+
+			bool operator == (const FPassChannelIndex& InOther) const
 			{
-				// 2 channels for the heightmap, and one channel for each grass map, packed into 4 channel render targets
-				NumPasses += FMath::DivideAndRoundUp(2 /* heightmap */ + NumGrassMaps, 4);
+				return (PassIndex == InOther.PassIndex) && (ChannelIndex == InOther.ChannelIndex);
 			}
 
-#if WITH_EDITORONLY_DATA
-			// since we don't read HeightMips unless we are in editor, there's no reason to add passes for it unless we are in editor
-			if (InHeightMips.Num() > 0)
+			bool IsValidChannelIndex() const
 			{
-				FirstHeightMipsPassIndex = NumPasses;
-				NumPasses += InHeightMips.Num();
+				return (ChannelIndex >= 0) && (ChannelIndex < 4);
 			}
-#endif // WITH_EDITORONLY_DATA
-		}
+		};
+
+		// Identifier for this component (use its key because it's a render-thread object so we don't want to access TObjectPtr there)
+		FIntPoint ComponentKey = FIntPoint(MAX_int32, MAX_int32);
+		// TODO [jonathan.bard] : this is a render thread object, this should not have a TObjectPtr (replace with FName)
+		TMap<TObjectPtr<ULandscapeGrassType>, FPassChannelIndex> PerGrassTypePassChannelIndices;
+		TArray<FIntVector4> PerPassPerChannelGrassWeightIndices;
+		FVector2D ViewOffset = FVector2D::ZeroVector;
+		int32 PixelOffsetX = 0;
+		FLandscapeComponentSceneProxy* SceneProxy = nullptr;
+		int32 FirstHeightMipsPassIndex = MAX_int32;
+		UE::Landscape::Grass::EGrassWeightExporterFlags Flags = UE::Landscape::Grass::EGrassWeightExporterFlags::None;
 	};
 
+	int32 ComponentSizeVerts = 0;
+	int32 SubsectionSizeQuads = 0;
+	int32 NumSubsections = 0;
 	FSceneInterface* SceneInterface = nullptr;
 	TArray<FComponentInfo, TInlineAllocator<1>> ComponentInfos;
-	FIntPoint TargetSize;
+	FIntPoint TargetSize = FIntPoint(ForceInit);
 	TArray<int32> HeightMips;
-	float PassOffsetX;
-	FVector ViewOrigin;
+	float PassOffsetX = 0.0f;
+	FVector ViewOrigin = FVector(ForceInit);
+	UE::Landscape::Grass::EGrassWeightExporterFlags Flags = UE::Landscape::Grass::EGrassWeightExporterFlags::None;
 
 	// game thread synchronous, do not access directly from render thread
 	FLandscapeAsyncTextureReadback* GameThreadAsyncReadbackPtr = nullptr;
 
-	FMatrix ViewRotationMatrix;
-	FMatrix ProjectionMatrix;
-
-	/** Creates a texture and renders the component to it, and then and triggers a readback of the texture. */
-	void RenderLandscapeComponentToTexture_RenderThread(FRHICommandListImmediate& RHICmdList, FLandscapeAsyncTextureReadback* AsyncReadbackPtr);
+	FMatrix ViewRotationMatrix = FMatrix::Identity;
+	FMatrix ProjectionMatrix = FMatrix::Identity;
 };
 
+
+// ----------------------------------------------------------------------------------
+
 class FLandscapeGrassWeightExporter : public FLandscapeGrassWeightExporter_RenderThread
 {
 	friend class FLandscapeGrassMapsBuilder;
 
-private:
-	TObjectPtr<ALandscapeProxy> LandscapeProxy;
-	int32 ComponentSizeVerts;
-	int32 SubsectionSizeQuads;
-	int32 NumSubsections;
-	TArray<TObjectPtr<ULandscapeGrassType>> GrassTypes;
-
 public:
+	UE_DEPRECATED(5.8, "Use the new constructor with explicit flags and explicit grass names")
 	LANDSCAPE_API FLandscapeGrassWeightExporter(ALandscapeProxy* InLandscapeProxy, TArrayView<ULandscapeComponent* const> InLandscapeComponents, bool bInNeedsGrassmap = true, bool bInNeedsHeightmap = true, const TArray<int32>& InHeightMips = {}, bool bInRenderImmediately = true, bool bInReadbackToCPU = true);
 
+	/**
+	 * @param InLandscape - Parent landscape. All requested component must share the same parent landscape
+	 * @param InRequestedComponents - List of components for which to export weight (/height) maps
+	 * @param InFlags - Flags for specializing what to export and how to read back the results
+	 * @param InRequestedGrassTypes - (Optional) List of grass types requested for readback. If empty, all grass types are exported
+	 * @param InRequestedHeightMips - (Optional) List of mip levels requested for reading back heightmaps
+	 */
+	LANDSCAPE_API FLandscapeGrassWeightExporter(ALandscape* InLandscape, TArrayView<ULandscapeComponent* const> InRequestedComponents, UE::Landscape::Grass::EGrassWeightExporterFlags InFlags,
+		TConstArrayView<TObjectPtr<ULandscapeGrassType>> InRequestedGrassTypes = {}, TConstArrayView<int32> InRequestedHeightMips = {});
+
 	// If using the async readback path, check its status and update if needed. Return true when the AsyncReadbackResults are available.
 	// You must call this periodically, or the async readback may not complete.
 	// bInForceFinish will force the RenderThread to wait until GPU completes the readback, ensuring the readback is completed after the render thread executes the command.
@@ -137,10 +220,8 @@ private:
 	void ApplyResults();
 
 	void CancelAndSelfDestruct();
-};
-
-namespace UE::Landscape::Grass
-{
-	void AddGrassWeightShaderTypes(FMaterialShaderTypes& InOutShaderTypes);
-}
 
+private:
+	TObjectPtr<ALandscape> Landscape;
+	TArray<TObjectPtr<ULandscapeGrassType>> GrassTypes;
+};
diff --git a/Engine/Source/Runtime/Landscape/Private/LandscapeEdit.cpp b/Engine/Source/Runtime/Landscape/Private/LandscapeEdit.cpp
--- a/Engine/Source/Runtime/Landscape/Private/LandscapeEdit.cpp
+++ b/Engine/Source/Runtime/Landscape/Private/LandscapeEdit.cpp
@@ -1655,7 +1655,6 @@ void ULandscapeComponent::UpdateCollisionData(bool bInUpdateHeightfieldRegion)
 {
 	TArray64<uint8> CollisionMipData;
 	TArray64<uint8> SimpleCollisionMipData;
-	TArray64<uint8> XYOffsetMipData;
 
 	check( GetHeightmap()->Source.IsValid() );
 
diff --git a/Engine/Source/Runtime/Landscape/Private/LandscapeEditLayers.cpp b/Engine/Source/Runtime/Landscape/Private/LandscapeEditLayers.cpp
--- a/Engine/Source/Runtime/Landscape/Private/LandscapeEditLayers.cpp
+++ b/Engine/Source/Runtime/Landscape/Private/LandscapeEditLayers.cpp
@@ -5423,8 +5423,8 @@ void ALandscape::OnDirtyHeightmap(FTextureToComponentHelper const& MapHelper, UT
 
 				if (NumDifferentPixels > 0)
 				{
-					FFileHelper::CreateBitmap(*(FilePattern + "_a(pre).bmp"), SizeU, SizeV, InOldData, &SubRegion, &IFileManager::Get(), /*OutFilename = */nullptr, /*bool bInWriteAlpha = */true);
-					FFileHelper::CreateBitmap(*(FilePattern + "_b(post).bmp"), SizeU, SizeV, InNewData, &SubRegion, &IFileManager::Get(), /*OutFilename = */nullptr, /*bool bInWriteAlpha = */true);
+					FFileHelper::CreateBitmap(*(FilePattern + "_a(pre).bmp"), SizeU, SizeV, InOldData, &SubRegion, &IFileManager::Get(), /*OutFilename = */nullptr, /*bInWriteAlpha = */true);
+					FFileHelper::CreateBitmap(*(FilePattern + "_b(post).bmp"), SizeU, SizeV, InNewData, &SubRegion, &IFileManager::Get(), /*OutFilename = */nullptr, /*bInWriteAlpha = */true);
 
 					if (bDumpDiffDetails)
 					{
diff --git a/Engine/Source/Runtime/Landscape/Private/LandscapeGrass.cpp b/Engine/Source/Runtime/Landscape/Private/LandscapeGrass.cpp
--- a/Engine/Source/Runtime/Landscape/Private/LandscapeGrass.cpp
+++ b/Engine/Source/Runtime/Landscape/Private/LandscapeGrass.cpp
@@ -29,6 +29,7 @@
 #include "LandscapePrivate.h"
 #include "LandscapeProxy.h"
 #include "LandscapeInfo.h"
+#include "LandscapeUtilsPrivate.h"
 #include "LightMap.h"
 #include "Engine/MapBuildDataRegistry.h"
 #include "ShadowMap.h"
@@ -960,15 +961,14 @@ bool ULandscapeComponent::CanRenderGrassMap() const
 
 #if WITH_EDITOR
 
-namespace UE::Landscape
+namespace UE::Landscape::Grass
 {
 	extern uint32 ComputeGrassMapGenerationHash(const ULandscapeComponent* Component, UMaterialInterface* Material);
-	extern void SubmitGPUCommands(bool bBlockUntilRTComplete, bool bBlockRTUntilGPUComplete);
-}
+} // namespace UE::Landscape::Grass
 
 uint32 ULandscapeComponent::ComputeGrassMapGenerationHash() const
 {
-	return UE::Landscape::ComputeGrassMapGenerationHash(this, GetLandscapeMaterial());
+	return UE::Landscape::Grass::ComputeGrassMapGenerationHash(this, GetLandscapeMaterial());
 }
 
 bool ULandscapeComponent::IsGrassMapOutdated() const
@@ -978,7 +978,11 @@ bool ULandscapeComponent::IsGrassMapOutdated() const
 
 TArray<uint16> ULandscapeComponent::RenderWPOHeightmap(int32 LOD)
 {
+	using namespace UE::Landscape::Grass;
+
 	TArray<uint16> Results;
+	ALandscape* Landscape = GetLandscapeActor();
+	check(Landscape != nullptr);
 
 	if (ensure(SceneProxy))
 	{
@@ -995,13 +999,14 @@ TArray<uint16> ULandscapeComponent::RenderWPOHeightmap(int32 LOD)
 
 		if (LOD == 0)
 		{
-			FLandscapeGrassWeightExporter Exporter(GetLandscapeProxy(), { this }, /*bInNeedsGrassmap = */ false, /*bInNeedsHeightmap =*/ true, {});
+			EGrassWeightExporterFlags LandscapeGrassWeightExporterFlags = EGrassWeightExporterFlags::NeedsHeightmap | EGrassWeightExporterFlags::RenderImmediately | EGrassWeightExporterFlags::ReadbackToCPU;
+			FLandscapeGrassWeightExporter Exporter(Landscape, /*InRequestedComponents = */{ this }, LandscapeGrassWeightExporterFlags);
 
 			while (!Exporter.IsAsyncReadbackComplete())
 			{
-				UE::Landscape::SubmitGPUCommands(/* bBlockUntilRTComplete =  */ true, /* bBlockRTUntilGPUComplete =  */ true);
+				UE::Landscape::Private::SubmitGPUCommands(/* bBlockUntilRTComplete =  */ true, /* bBlockRTUntilGPUComplete =  */ true);
 				bool bOutRenderCommandsQueued;
-				Exporter.CheckAndUpdateAsyncReadback(bOutRenderCommandsQueued, true);
+				Exporter.CheckAndUpdateAsyncReadback(bOutRenderCommandsQueued, /*bInForceFinish = */true);
 			}
 
 			TMap<ULandscapeComponent*, TUniquePtr<FLandscapeComponentGrassData>, TInlineSetAllocator<1>> TempGrassData;
@@ -1010,15 +1015,16 @@ TArray<uint16> ULandscapeComponent::RenderWPOHeightmap(int32 LOD)
 		}
 		else
 		{
+			EGrassWeightExporterFlags LandscapeGrassWeightExporterFlags = EGrassWeightExporterFlags::RenderImmediately | EGrassWeightExporterFlags::ReadbackToCPU;
 			TArray<int32> HeightMips;
 			HeightMips.Add(LOD);
-			FLandscapeGrassWeightExporter Exporter(GetLandscapeProxy(), { this }, /*bInNeedsGrassmap = */ false, /*bInNeedsHeightmap =*/ false, MoveTemp(HeightMips));
+			FLandscapeGrassWeightExporter Exporter(Landscape, /*InRequestedComponents = */{ this }, LandscapeGrassWeightExporterFlags, /*InRequestedGrassTypes = */{}, MakeConstArrayView(HeightMips));
 
 			while (!Exporter.IsAsyncReadbackComplete())
 			{
-				UE::Landscape::SubmitGPUCommands(/* bBlockUntilRTComplete =  */ true, /* bBlockRTUntilGPUComplete =  */ true);
+				UE::Landscape::Private::SubmitGPUCommands(/* bBlockUntilRTComplete =  */ true, /* bBlockRTUntilGPUComplete =  */ true);
 				bool bOutRenderCommandsQueued;
-				Exporter.CheckAndUpdateAsyncReadback(bOutRenderCommandsQueued, true);
+				Exporter.CheckAndUpdateAsyncReadback(bOutRenderCommandsQueued, /*bInForceFinish = */true);
 			}
 
 			TMap<ULandscapeComponent*, TUniquePtr<FLandscapeComponentGrassData>, TInlineSetAllocator<1>> TempGrassData;
diff --git a/Engine/Source/Runtime/Landscape/Private/LandscapeGrassMapsBuilder.cpp b/Engine/Source/Runtime/Landscape/Private/LandscapeGrassMapsBuilder.cpp
--- a/Engine/Source/Runtime/Landscape/Private/LandscapeGrassMapsBuilder.cpp
+++ b/Engine/Source/Runtime/Landscape/Private/LandscapeGrassMapsBuilder.cpp
@@ -6,6 +6,8 @@
 #include "LandscapeAsyncTextureReadback.h"
 #include "LandscapePrivate.h"
 #include "LandscapeProxy.h"
+#include "LandscapeUtils.h"
+#include "LandscapeUtilsPrivate.h"
 #include "LandscapeGrassWeightExporter.h"
 #include "MaterialCachedData.h"
 #include "LandscapeEditTypes.h"
@@ -147,132 +149,112 @@ FLandscapeGrassMapsBuilder::FLandscapeGrassMapsBuilder(UWorld* InOwner, FLandsca
 {
 }
 
-namespace UE::Landscape
+namespace UE::Landscape::Grass
 {
-	uint32 ComputeGrassMapGenerationHash(const ULandscapeComponent* Component, UMaterialInterface* Material)
-	{
-		TRACE_CPUPROFILER_EVENT_SCOPE(ComputeGrassMapGenerationHash);
-		uint32 Hash = 0;
 
-		// Change this hash key to invalidate all cached grass density maps, if the generation functions change.
-		static uint32 InitialHash = GetTypeHash(FGuid("216D95C7651D4095ADC6A8459B4F181D"));
+uint32 ComputeGrassMapGenerationHash(const ULandscapeComponent* Component, UMaterialInterface* Material)
+{
+	TRACE_CPUPROFILER_EVENT_SCOPE(ComputeGrassMapGenerationHash);
+	uint32 Hash = 0;
 
-		Hash = InitialHash;
+	// Change this hash key to invalidate all cached grass density maps, if the generation functions change.
+	static uint32 InitialHash = GetTypeHash(FGuid("216D95C7651D4095ADC6A8459B4F181D"));
 
-		// we only include material and texture hashes in editor (there is no automatic detection of changes in non-editor builds)
-	#if WITH_EDITOR
-		// Take into account any material state change : (excluding texture state)
-		Hash = FCrc::TypeCrc32(Material->ComputeAllStateCRC(), Hash);
+	Hash = InitialHash;
 
-		// hash the heightmap texture (we use the component's heightmap texture hash that ignores normals)
-		Hash = FCrc::TypeCrc32(ULandscapeTextureHash::GetHash(Component->GetHeightmap()), Hash);
+	// we only include material and texture hashes in editor (there is no automatic detection of changes in non-editor builds)
+#if WITH_EDITOR
+	// Take into account any material state change : (excluding texture state)
+	Hash = FCrc::TypeCrc32(Material->ComputeAllStateCRC(), Hash);
 
-		// hash the weightmap textures
-		for (UTexture2D* Weightmap : Component->GetWeightmapTextures())
-		{
-			check(Weightmap->Source.IsValid());
-			Hash = FCrc::TypeCrc32(ULandscapeTextureHash::GetHash(Weightmap), Hash);
-		}
-	#endif // WITH_EDITOR
+	// hash the heightmap texture (we use the component's heightmap texture hash that ignores normals)
+	Hash = FCrc::TypeCrc32(ULandscapeTextureHash::GetHash(Component->GetHeightmap()), Hash);
 
-		return Hash;
+	// hash the weightmap textures
+	for (UTexture2D* Weightmap : Component->GetWeightmapTextures())
+	{
+		check(Weightmap->Source.IsValid());
+		Hash = FCrc::TypeCrc32(ULandscapeTextureHash::GetHash(Weightmap), Hash);
 	}
+#endif // WITH_EDITOR
 
-	uint32 ComputeGrassInstanceGenerationHash(uint32 GrassMapGenerationHash, const TArray<TObjectPtr<ULandscapeGrassType>>& GrassTypes)
-	{
-		TRACE_CPUPROFILER_EVENT_SCOPE(ComputeGrassInstanceGenerationHash);
-		// grass maps are input to grass instances
-		uint32 Hash = GrassMapGenerationHash;
+	return Hash;
+}
 
-		// If anything changes in the grass types, that affects the grass instances
-		for (ULandscapeGrassType* GrassType : GrassTypes)
-		{
-			Hash = FCrc::TypeCrc32(GrassType ? GrassType->StateHash : 0, Hash);
-		}
+uint32 ComputeGrassInstanceGenerationHash(uint32 GrassMapGenerationHash, const TArray<TObjectPtr<ULandscapeGrassType>>& GrassTypes)
+{
+	TRACE_CPUPROFILER_EVENT_SCOPE(ComputeGrassInstanceGenerationHash);
+	// grass maps are input to grass instances
+	uint32 Hash = GrassMapGenerationHash;
 
-		return Hash;
+	// If anything changes in the grass types, that affects the grass instances
+	for (ULandscapeGrassType* GrassType : GrassTypes)
+	{
+		Hash = FCrc::TypeCrc32(GrassType ? GrassType->StateHash : 0, Hash);
 	}
 
+	return Hash;
+}
+
+// Deprecated
+bool CanRenderGrassMap(ULandscapeComponent *Component)
+{
+	return Component->CanRenderGrassMap();
+}
+
+} // namespace UE::Landscape::Grass
+
+namespace UE::Landscape::Grass::Private
+{
+
 #if WITH_EDITOR
-	void CompileGrassMapShader(const ULandscapeComponent* Component)
+void CompileGrassMapShader(const ULandscapeComponent* Component)
+{
+	if (Component->GetMaterialInstanceCount(false) > 0)
 	{
-		if (Component->GetMaterialInstanceCount(false) > 0)
+		if (UMaterialInstance* MaterialInstance = Component->GetMaterialInstance(0))
 		{
-			if (UMaterialInstance* MaterialInstance = Component->GetMaterialInstance(0))
+			if (FMaterialResource* MaterialResource = MaterialInstance->GetMaterialResource(GetFeatureLevelShaderPlatform_Checked(Component->GetWorld()->GetFeatureLevel())))
 			{
-				if (FMaterialResource* MaterialResource = MaterialInstance->GetMaterialResource(GetFeatureLevelShaderPlatform_Checked(Component->GetWorld()->GetFeatureLevel())))
-				{
-					MaterialResource->FinishCompilation();
-				}
+				MaterialResource->FinishCompilation();
 			}
 		}
 	}
+}
 #endif // WITH_EDITOR
 
-	// Deprecated
-	bool CanRenderGrassMap(ULandscapeComponent *Component)
-	{
-		return Component->CanRenderGrassMap();
-	}
-
-	bool IsRuntimeGrassMapGenerationSupported()
+// calculates the minimum distance between any cameras and the specified Worldbounds.
+static inline double CalculateMinDistanceToCamerasSquared(const TArray<FVector>& Cameras, const FBoxSphereBounds& WorldBounds)
+{
+	if (!Cameras.Num())
 	{
-		return GGrassMapUseRuntimeGeneration || GGrassMapAlwaysBuildRuntimeGenerationResources;
+		return 0.0f;
 	}
-
-	// calculates the minimum distance between any cameras and the specified Worldbounds.
-	static inline double CalculateMinDistanceToCamerasSquared(const TArray<FVector>& Cameras, const FBoxSphereBounds& WorldBounds)
+	double MinSqrDistance = MAX_dbl;
+	for (const FVector& CameraPos : Cameras)
 	{
-		if (!Cameras.Num())
-		{
-			return 0.0f;
-		}
-		double MinSqrDistance = MAX_dbl;
-		for (const FVector& CameraPos : Cameras)
-		{
-			MinSqrDistance = FMath::Min<double>(MinSqrDistance, WorldBounds.ComputeSquaredDistanceFromBoxToPoint(CameraPos));
-		}
-		return MinSqrDistance;
+		MinSqrDistance = FMath::Min<double>(MinSqrDistance, WorldBounds.ComputeSquaredDistanceFromBoxToPoint(CameraPos));
 	}
+	return MinSqrDistance;
+}
 
-	void SubmitGPUCommands(bool bBlockUntilRTComplete, bool bBlockRTUntilGPUComplete)
-	{
-		FEvent* ResultsReadyEvent = nullptr;
-		if (bBlockUntilRTComplete)
-		{
-			ResultsReadyEvent = FPlatformProcess::GetSynchEventFromPool(true);
-		}
-
-		ENQUEUE_RENDER_COMMAND(FFlushResourcesCommand)(
-			[ResultsReadyEvent, bBlockRTUntilGPUComplete](FRHICommandListImmediate& RHICmdList)
-			{
-				RHICmdList.ImmediateFlush(EImmediateFlushType::FlushRHIThreadFlushResources);
-				RHICmdList.FlushResources();
-				RHICmdList.SubmitCommandsAndFlushGPU();
+} // namespace UE::Landscape::Grass::Private
 
-				if (ResultsReadyEvent)
-				{
-					if (bBlockRTUntilGPUComplete)
-					{
-						// Block render thread waiting for GPU to complete.  Note this can be very expensive on some platforms.
-						RHICmdList.BlockUntilGPUIdle();
-					}
-					
-					ResultsReadyEvent->Trigger();
-				}
-			});
-
-		if (ResultsReadyEvent)
-		{
-			// block game thread waiting for render thread to tell us the GPU is complete
-			ResultsReadyEvent->Wait();
-			ResultsReadyEvent->Reset();
-			FPlatformProcess::ReturnSynchEventToPool(ResultsReadyEvent);
-		}
+namespace UE::Landscape
+{
 
-		return;
-	}
+// Deprecated
+bool CanRenderGrassMap(ULandscapeComponent* Component)
+{
+	return Component->CanRenderGrassMap();
 }
+// Deprecated
+bool IsRuntimeGrassMapGenerationSupported()
+{
+	return UE::Landscape::SupportsRuntimeGrassMapGeneration();
+}
+
+} // namespace UE::Landscape
 
 FLandscapeGrassMapsBuilder::~FLandscapeGrassMapsBuilder()
 {
@@ -355,9 +337,9 @@ FLandscapeGrassMapsBuilder::FComponentState::FComponentState(ULandscapeComponent
 	UMaterialInterface* Material = Component->GetLandscapeMaterial();
 	if (Material)
 	{
-		GrassMapGenerationHash = UE::Landscape::ComputeGrassMapGenerationHash(Component, Material);
+		GrassMapGenerationHash = UE::Landscape::Grass::ComputeGrassMapGenerationHash(Component, Material);
 		const TArray<TObjectPtr<ULandscapeGrassType>>& GrassTypes = Material->GetCachedExpressionData().GrassTypes;
-		GrassInstanceGenerationHash = UE::Landscape::ComputeGrassInstanceGenerationHash(GrassMapGenerationHash, GrassTypes);
+		GrassInstanceGenerationHash = UE::Landscape::Grass::ComputeGrassInstanceGenerationHash(GrassMapGenerationHash, GrassTypes);
 	}
 #endif // WITH_EDITOR
 }
@@ -367,7 +349,7 @@ void FLandscapeGrassMapsBuilder::FPendingComponent::UpdatePriorityDistance(const
 	TRACE_CPUPROFILER_EVENT_SCOPE(FLandscapeGrassMapsBuilder::FPendingComponent::UpdatePriorityDistance);
 	ULandscapeComponent* Component = State->Component;
 	FBoxSphereBounds WorldBounds = Component->CalcBounds(Component->GetComponentTransform());
-	const double MinSqrDistanceToComponent = UE::Landscape::CalculateMinDistanceToCamerasSquared(Cameras, WorldBounds);
+	const double MinSqrDistanceToComponent = UE::Landscape::Grass::Private::CalculateMinDistanceToCamerasSquared(Cameras, WorldBounds);
 	const double ThresholdDistance = Component->GrassTypeSummary.MaxInstanceDiscardDistance * MustHaveDistanceScale;
 	PriorityKey = MinSqrDistanceToComponent - ThresholdDistance * ThresholdDistance;
 }
@@ -500,7 +482,7 @@ bool FLandscapeGrassMapsBuilder::UpdateTrackedComponents(const TArray<FVector>&
 						}
 
 						// check if any dependencies changed
-						const uint32 CurGrassMapGenerationHash = UE::Landscape::ComputeGrassMapGenerationHash(Component, Material);
+						const uint32 CurGrassMapGenerationHash = UE::Landscape::Grass::ComputeGrassMapGenerationHash(Component, Material);
 						if (State->GrassMapGenerationHash != CurGrassMapGenerationHash)
 						{
 							break; // cancel and evict to restart process
@@ -508,7 +490,7 @@ bool FLandscapeGrassMapsBuilder::UpdateTrackedComponents(const TArray<FVector>&
 
 						// check if any grass types have changed -- this invalidates foliage instances but not the grass maps
 						const TArray<TObjectPtr<ULandscapeGrassType>>& GrassTypes = Component->GetGrassTypes();
-						const uint32 CurGrassInstanceGenerationHash = UE::Landscape::ComputeGrassInstanceGenerationHash(CurGrassMapGenerationHash, GrassTypes);
+						const uint32 CurGrassInstanceGenerationHash = UE::Landscape::Grass::ComputeGrassInstanceGenerationHash(CurGrassMapGenerationHash, GrassTypes);
 						if (State->GrassInstanceGenerationHash != CurGrassInstanceGenerationHash)
 						{
 							ComponentsToRemoveFoliageInstances.Add(Component);
@@ -973,7 +955,7 @@ void FLandscapeGrassMapsBuilder::DebugRenderComponents(
 						{
 	#if WITH_EDITOR
 							// in editor, try to force compilation to complete
-							UE::Landscape::CompileGrassMapShader(State->Component);
+							UE::Landscape::Grass::Private::CompileGrassMapShader(State->Component);
 							if (!Component->CanRenderGrassMap())
 	#endif // WITH_EDITOR
 							{
@@ -1171,7 +1153,7 @@ bool FLandscapeGrassMapsBuilder::BuildGrassMapsNowForComponents(
 				{
 #if WITH_EDITOR
 					// in editor, try to force compilation to complete
-					UE::Landscape::CompileGrassMapShader(State->Component);
+					UE::Landscape::Grass::Private::CompileGrassMapShader(State->Component);
 					if (!Component->CanRenderGrassMap())
 #endif // WITH_EDITOR
 					{
@@ -1193,7 +1175,7 @@ bool FLandscapeGrassMapsBuilder::BuildGrassMapsNowForComponents(
 		{
 			// TODO [chris.tchou] it currently seems to be faster to block here; otherwise it takes a long time to complete the readback
 			// not sure why this is, something must be getting starved in the non-blocking path.
-			UE::Landscape::SubmitGPUCommands(/* bBlockUntilRTComplete =  */ true, /* bBlockRTUntilGPUComplete =  */ false);
+			UE::Landscape::Private::SubmitGPUCommands(/* bBlockUntilRTComplete =  */ true, /* bBlockRTUntilGPUComplete =  */ false);
 			LastFlush = CurTime;
 		}
 		
@@ -1290,7 +1272,8 @@ bool FLandscapeGrassMapsBuilder::CancelAndEvict(FComponentState& State, bool bCa
 					// Because Components can be re-registered and pick up where they left off, we don't want to cancel the operation outright just yet.
 					// To avoid potential render thread stalls we let the operation continue asynchronously for up to four ticks,
 					// before forcing completion (this also gives us more time for re-register)
-					if (!State.ActiveRender->CheckAndUpdateAsyncReadback(bNewRenderCommands, State.TickCount > 4))
+					const bool bForceFinish = (State.TickCount > 4);
+					if (!State.ActiveRender->CheckAndUpdateAsyncReadback(bNewRenderCommands, bForceFinish))
 					{
 						// we can't evict yet.. must wait for the readback to complete
 						return false;
@@ -1404,9 +1387,9 @@ bool FLandscapeGrassMapsBuilder::TryFastpathsFromPending(FComponentState& State,
 		{
 			if (UMaterialInterface* Material = Component->GetLandscapeMaterial())
 			{
-				State.GrassMapGenerationHash = UE::Landscape::ComputeGrassMapGenerationHash(Component, Material);
+				State.GrassMapGenerationHash = UE::Landscape::Grass::ComputeGrassMapGenerationHash(Component, Material);
 				const TArray<TObjectPtr<ULandscapeGrassType>>& GrassTypes = Material->GetCachedExpressionData().GrassTypes;
-				State.GrassInstanceGenerationHash = UE::Landscape::ComputeGrassInstanceGenerationHash(State.GrassMapGenerationHash, GrassTypes);
+				State.GrassInstanceGenerationHash = UE::Landscape::Grass::ComputeGrassInstanceGenerationHash(State.GrassMapGenerationHash, GrassTypes);
 			}
 		}
 
@@ -1449,7 +1432,7 @@ bool FLandscapeGrassMapsBuilder::StartGrassMapGeneration(FComponentState& State,
 #if WITH_EDITOR
 		if (bForceCompileShaders)
 		{
-			UE::Landscape::CompileGrassMapShader(Component);
+			UE::Landscape::Grass::Private::CompileGrassMapShader(Component);
 		}
 
 		if (!Component->CanRenderGrassMap())
@@ -1622,7 +1605,7 @@ void FLandscapeGrassMapsBuilder::KickOffRenderAndReadback(FComponentState& State
 	check(State.ActiveRender == nullptr);
 	constexpr bool bInNeedsGrassmap = true;
 	constexpr bool bInNeedsHeightmap = true;
-	State.ActiveRender.Reset(new FLandscapeGrassWeightExporter(Component->GetLandscapeProxy(), { Component }, bInNeedsGrassmap, bInNeedsHeightmap, MoveTemp(HeightMips)));
+	State.ActiveRender.Reset(new FLandscapeGrassWeightExporter(Component->GetLandscapeActor(), { Component }, UE::Landscape::Grass::EGrassWeightExporterFlags::Default, /*InRequestedGrassTypes = */{}, MakeConstArrayView(HeightMips)));
 	check(State.ActiveRender != nullptr);
 
 	State.Stage = EComponentStage::Rendering;
@@ -1730,7 +1713,7 @@ bool FLandscapeGrassMapsBuilder::FComponentState::IsBeyondEvictionRange(const TA
 {
 	check(Stage == EComponentStage::GrassMapsPopulated);
 	const FBoxSphereBounds WorldBounds = Component->CalcBounds(Component->GetComponentTransform());
-	const float MinSqrDistanceToComponent = UE::Landscape::CalculateMinDistanceToCamerasSquared(Cameras, WorldBounds);
+	const float MinSqrDistanceToComponent = UE::Landscape::Grass::Private::CalculateMinDistanceToCamerasSquared(Cameras, WorldBounds);
 	const float DiscardDistanceScale = GGrassMapGuardBandDiscardMultiplier * GGrassCullDistanceScale;
 	const float MinEvictDistance = Component->GrassTypeSummary.MaxInstanceDiscardDistance * DiscardDistanceScale;
 	return (MinSqrDistanceToComponent > MinEvictDistance * MinEvictDistance);
diff --git a/Engine/Source/Runtime/Landscape/Private/LandscapeGrassWeightExporter.cpp b/Engine/Source/Runtime/Landscape/Private/LandscapeGrassWeightExporter.cpp
--- a/Engine/Source/Runtime/Landscape/Private/LandscapeGrassWeightExporter.cpp
+++ b/Engine/Source/Runtime/Landscape/Private/LandscapeGrassWeightExporter.cpp
@@ -14,24 +14,61 @@
 #include "TextureResource.h"
 #include "RenderCaptureInterface.h"
 #include "ShaderPlatformCachedIniValue.h"
+#include "Landscape.h"
 #include "LandscapePrivate.h"
+#include "LandscapeSubsystem.h"
 #include "LandscapeAsyncTextureReadback.h"
+#include "Materials/MaterialExpressionLandscapeGrassOutput.h"
 
-#if UE_BUILD_DEBUG
+#if WITH_EDITOR
 #include "Misc/FileHelper.h"
-#endif
+#endif // WITH_EDITOR
 
 class FLandscapeGrassWeightVS;
 class FLandscapeGrassWeightPS;
 
+// TODO [jonathan.bard] : Remove this after we've let this run for a while
+#define GRASS_ALLOW_EXTRAVALIDATION 1
+
+#define GRASS_ENABLE_EXTRAVALIDATION (!UE_BUILD_SHIPPING && GRASS_ALLOW_EXTRAVALIDATION)
+
+// ----------------------------------------------------------------------------------
+
 int32 GRenderCaptureNextGrassmapDraws = 0;
 static FAutoConsoleVariableRef CVarRenderCaptureNextGrassmapDraws(
 	TEXT("grass.GrassMap.RenderCaptureNextDraws"),
 	GRenderCaptureNextGrassmapDraws,
 	TEXT("Trigger render captures during the next N grassmap draw calls."));
 
+static TAutoConsoleVariable<int32> CVarDumpGrassmaps(
+	TEXT("grass.GrassMap.Dump"),
+	0,
+	TEXT("Dumps every rendered grass map to disk."));
+
+
+// ----------------------------------------------------------------------------------
+
 extern int32 GGrassMapAlwaysBuildRuntimeGenerationResources;
 extern int32 GGrassMapUseRuntimeGeneration;
+extern int32 GGrassEnable;
+
+
+// ----------------------------------------------------------------------------------
+
+DECLARE_GPU_STAT(LandscapeGrassMaps);
+
+
+// ----------------------------------------------------------------------------------
+
+namespace UE::Landscape::Grass::Private
+{
+	static constexpr int32 GrassWeightInvalidIndex = -3;
+	static constexpr int32 HeightGrassWeightIndex0 = -2;
+	static constexpr int32 HeightGrassWeightIndex1 = -1;
+} // namespace UE::Landscape::Grass::Private
+
+
+// ----------------------------------------------------------------------------------
 
 BEGIN_SHADER_PARAMETER_STRUCT(FLandscapeGrassPassParameters, )
 	SHADER_PARAMETER_STRUCT_REF(FViewUniformShaderParameters, View)
@@ -40,14 +77,18 @@ BEGIN_SHADER_PARAMETER_STRUCT(FLandscapeGrassPassParameters, )
 	RENDER_TARGET_BINDING_SLOTS()
 END_SHADER_PARAMETER_STRUCT()
 
+
+// ----------------------------------------------------------------------------------
+
 class FLandscapeGrassWeightShaderElementData : public FMeshMaterialShaderElementData
 {
 public:
-	int32 OutputPass;
+	FIntVector4 PerChannelGrassWeightIndices;
 	FVector2f RenderOffset;
 };
 
-extern int32 GGrassEnable;
+
+// ----------------------------------------------------------------------------------
 
 static bool ShouldCacheLandscapeGrassShaders(const FMeshMaterialShaderPermutationParameters& Parameters)
 {
@@ -87,6 +128,9 @@ static bool ShouldCacheLandscapeGrassShaders(const FMeshMaterialShaderPermutatio
 	return bShouldCache;
 }
 
+
+// ----------------------------------------------------------------------------------
+
 class FLandscapeGrassWeightVS : public FMeshMaterialShader
 {
 	DECLARE_SHADER_TYPE(FLandscapeGrassWeightVS, MeshMaterial);
@@ -129,10 +173,13 @@ public:
 
 IMPLEMENT_MATERIAL_SHADER_TYPE(, FLandscapeGrassWeightVS, TEXT("/Engine/Private/LandscapeGrassWeight.usf"), TEXT("VSMain"), SF_Vertex);
 
+
+// ----------------------------------------------------------------------------------
+
 class FLandscapeGrassWeightPS : public FMeshMaterialShader
 {
 	DECLARE_SHADER_TYPE(FLandscapeGrassWeightPS, MeshMaterial);
-	LAYOUT_FIELD(FShaderParameter, OutputPassParameter);
+	LAYOUT_FIELD(FShaderParameter, PerChannelGrassWeightIndicesParameter);
 public:
 
 	static bool ShouldCompilePermutation(const FMeshMaterialShaderPermutationParameters& Parameters)
@@ -143,7 +190,7 @@ public:
 	FLandscapeGrassWeightPS(const ShaderMetaType::CompiledShaderInitializerType& Initializer)
 		: FMeshMaterialShader(Initializer)
 	{
-		OutputPassParameter.Bind(Initializer.ParameterMap, TEXT("OutputPass"));
+		PerChannelGrassWeightIndicesParameter.Bind(Initializer.ParameterMap, TEXT("PerChannelGrassWeightIndices"));
 		PassUniformBuffer.Bind(Initializer.ParameterMap, FSceneTextureUniformParameters::FTypeInfo::GetStructMetadata()->GetShaderVariableName());
 	}
 
@@ -161,18 +208,24 @@ public:
 	{
 		FMeshMaterialShader::GetShaderBindings(Scene, FeatureLevel, PrimitiveSceneProxy, MaterialRenderProxy, Material, ShaderElementData, ShaderBindings);
 
-		ShaderBindings.Add(OutputPassParameter, ShaderElementData.OutputPass);
+		ShaderBindings.Add(PerChannelGrassWeightIndicesParameter, ShaderElementData.PerChannelGrassWeightIndices);
 	}
 };
 
 IMPLEMENT_MATERIAL_SHADER_TYPE(, FLandscapeGrassWeightPS, TEXT("/Engine/Private/LandscapeGrassWeight.usf"), TEXT("PSMain"), SF_Pixel);
 
+
+// ----------------------------------------------------------------------------------
+
 void UE::Landscape::Grass::AddGrassWeightShaderTypes(FMaterialShaderTypes& InOutShaderTypes)
 {
 	InOutShaderTypes.AddShaderType<FLandscapeGrassWeightVS>();
 	InOutShaderTypes.AddShaderType<FLandscapeGrassWeightPS>();
 }
 
+
+// ----------------------------------------------------------------------------------
+
 class FLandscapeGrassWeightMeshProcessor : public FMeshPassProcessor
 {
 public:
@@ -180,7 +233,7 @@ public:
 
 	void AddMeshBatch(const FMeshBatch& RESTRICT MeshBatch,
 		uint64 BatchElementMask,
-		int32 NumPasses,
+		TConstArrayView<FIntVector4> PerPassPerChannelGrassWeightIndices,
 		FVector2D ViewOffset,
 		float PassOffsetX,
 		int32 FirstHeightMipsPassIndex,
@@ -206,7 +259,7 @@ private:
 		int32 StaticMeshId,
 		const FMaterialRenderProxy& MaterialRenderProxy,
 		const FMaterial& MaterialResource,
-		int32 NumPasses,
+		TConstArrayView<FIntVector4> PerPassPerChannelGrassWeightIndices,
 		FVector2D ViewOffset,
 		float PassOffsetX,
 		int32 FirstHeightMipsPassIndex,
@@ -218,7 +271,7 @@ private:
 		const FPrimitiveSceneProxy* RESTRICT PrimitiveSceneProxy,
 		const FMaterialRenderProxy& RESTRICT MaterialRenderProxy,
 		const FMaterial& RESTRICT MaterialResource,
-		int32 NumPasses,
+		TConstArrayView<FIntVector4> PerPassPerChannelGrassWeightIndices,
 		FVector2D ViewOffset,
 		float PassOffsetX,
 		int32 FirstHeightMipsPassIndex,
@@ -238,7 +291,7 @@ FLandscapeGrassWeightMeshProcessor::FLandscapeGrassWeightMeshProcessor(const FSc
 
 void FLandscapeGrassWeightMeshProcessor::AddMeshBatch(const FMeshBatch& RESTRICT MeshBatch,
 	uint64 BatchElementMask,
-	int32 NumPasses,
+	TConstArrayView<FIntVector4> PerPassPerChannelGrassWeightIndices,
 	FVector2D ViewOffset,
 	float PassOffsetX,
 	int32 FirstHeightMipsPassIndex,
@@ -251,7 +304,7 @@ void FLandscapeGrassWeightMeshProcessor::AddMeshBatch(const FMeshBatch& RESTRICT
 		const FMaterial* Material = MaterialRenderProxy->GetMaterialNoFallback(FeatureLevel);
 		if (Material)
 		{
-			if (TryAddMeshBatch(MeshBatch, BatchElementMask, PrimitiveSceneProxy, -1, *MaterialRenderProxy, *Material, NumPasses, ViewOffset, PassOffsetX, FirstHeightMipsPassIndex, HeightMips))
+			if (TryAddMeshBatch(MeshBatch, BatchElementMask, PrimitiveSceneProxy, -1, *MaterialRenderProxy, *Material, PerPassPerChannelGrassWeightIndices, ViewOffset, PassOffsetX, FirstHeightMipsPassIndex, HeightMips))
 			{
 				break;
 			}
@@ -268,14 +321,14 @@ bool FLandscapeGrassWeightMeshProcessor::TryAddMeshBatch(
 	int32 StaticMeshId,
 	const FMaterialRenderProxy& MaterialRenderProxy,
 	const FMaterial& MaterialResource,
-	int32 NumPasses,
+	TConstArrayView<FIntVector4> PerPassPerChannelGrassWeightIndices,
 	FVector2D ViewOffset,
 	float PassOffsetX,
 	int32 FirstHeightMipsPassIndex,
 	const TArray<int32>& HeightMips)
 {
 	check(MeshBatch.VertexFactory != nullptr);
-	return Process(MeshBatch, BatchElementMask, PrimitiveSceneProxy, MaterialRenderProxy, MaterialResource, NumPasses, ViewOffset, PassOffsetX, FirstHeightMipsPassIndex, HeightMips);
+	return Process(MeshBatch, BatchElementMask, PrimitiveSceneProxy, MaterialRenderProxy, MaterialResource, PerPassPerChannelGrassWeightIndices, ViewOffset, PassOffsetX, FirstHeightMipsPassIndex, HeightMips);
 }
 
 bool FLandscapeGrassWeightMeshProcessor::Process(
@@ -284,7 +337,7 @@ bool FLandscapeGrassWeightMeshProcessor::Process(
 	const FPrimitiveSceneProxy* RESTRICT PrimitiveSceneProxy,
 	const FMaterialRenderProxy& RESTRICT MaterialRenderProxy,
 	const FMaterial& RESTRICT MaterialResource,
-	int32 NumPasses,
+	TConstArrayView<FIntVector4> PerPassPerChannelGrassWeightIndices,
 	FVector2D ViewOffset,
 	float PassOffsetX,
 	int32 FirstHeightMipsPassIndex,
@@ -316,10 +369,11 @@ bool FLandscapeGrassWeightMeshProcessor::Process(
 	ShaderElementData.InitializeMeshMaterialData(ViewIfDynamicMeshCommand, PrimitiveSceneProxy, MeshBatch, -1, true);
 
 	const FMeshDrawCommandSortKey SortKey = CalculateMeshStaticSortKey(PassShaders.VertexShader, PassShaders.PixelShader);
-
+	
+	const int32 NumPasses = PerPassPerChannelGrassWeightIndices.Num();
 	for (int32 PassIndex = 0; PassIndex < NumPasses; ++PassIndex)
 	{
-		ShaderElementData.OutputPass = (PassIndex >= FirstHeightMipsPassIndex) ? 0 : PassIndex;
+		ShaderElementData.PerChannelGrassWeightIndices = PerPassPerChannelGrassWeightIndices[PassIndex];
 		ShaderElementData.RenderOffset = FVector2f(ViewOffset) + FVector2f(PassOffsetX * PassIndex, 0);	// LWC_TODO: Precision loss
 
 		uint64 Mask = (PassIndex >= FirstHeightMipsPassIndex) ? HeightMips[PassIndex - FirstHeightMipsPassIndex] : BatchElementMask;
@@ -396,10 +450,18 @@ IPSOCollector* CreateLandscapeGrassWeightPSOCollector(ERHIFeatureLevel::Type Fea
 FRegisterPSOCollectorCreateFunction RegisterLandscapeGrassWeightPSOCollector(&CreateLandscapeGrassWeightPSOCollector, EShadingPath::Deferred, LandscapeGrassWeightMeshPassName);
 FRegisterPSOCollectorCreateFunction RegisterMobileLandscapeGrassWeightPSOCollector(&CreateLandscapeGrassWeightPSOCollector, EShadingPath::Mobile, LandscapeGrassWeightMeshPassName);
 
-FLandscapeGrassWeightExporter_RenderThread::FLandscapeGrassWeightExporter_RenderThread(const TArray<int32>& InHeightMips, bool bInReadbackToCPU)
-	: HeightMips(InHeightMips)
+
+// ----------------------------------------------------------------------------------
+
+FLandscapeGrassWeightExporter_RenderThread::FLandscapeGrassWeightExporter_RenderThread(int32 InComponentSizeVerts, int32 InSubsectionSizeQuads, int32 InNumSubsections, 
+	TConstArrayView<int32> InHeightMips, UE::Landscape::Grass::EGrassWeightExporterFlags InFlags)
+	: ComponentSizeVerts(InComponentSizeVerts)
+	, SubsectionSizeQuads(InSubsectionSizeQuads)
+	, NumSubsections(InNumSubsections)
+	, HeightMips(InHeightMips)
+	, Flags(InFlags)
 {
-	if (bInReadbackToCPU)
+	if (EnumHasAnyFlags(InFlags, UE::Landscape::Grass::EGrassWeightExporterFlags::ReadbackToCPU))
 	{
 		// even when doing a synchronous readback, we use the async readback structure
 		GameThreadAsyncReadbackPtr = new FLandscapeAsyncTextureReadback();
@@ -415,8 +477,6 @@ FLandscapeGrassWeightExporter_RenderThread::~FLandscapeGrassWeightExporter_Rende
 	}
 }
 
-DECLARE_GPU_STAT(LandscapeGrassMaps);
-
 void FLandscapeGrassWeightExporter_RenderThread::RenderLandscapeComponentToTexture_RenderThread(FRHICommandListImmediate& RHICmdList, FLandscapeAsyncTextureReadback* AsyncReadbackPtr)
 {
 	FRDGBuilder GraphBuilder(RHICmdList);
@@ -433,8 +493,10 @@ void FLandscapeGrassWeightExporter_RenderThread::RenderLandscapeComponentToTextu
 
 	RenderLandscapeComponentToTexture_RenderThread(GraphBuilder, OutputTexture);
 
-	check(AsyncReadbackPtr != nullptr);
-	AsyncReadbackPtr->StartReadback_RenderThread(GraphBuilder, OutputTexture);
+	if (AsyncReadbackPtr != nullptr)
+	{
+		AsyncReadbackPtr->StartReadback_RenderThread(GraphBuilder, OutputTexture);
+	}
 
 	GraphBuilder.Execute();
 }
@@ -459,7 +521,7 @@ void FLandscapeGrassWeightExporter_RenderThread::RenderLandscapeComponentToTextu
 	GetRendererModule().CreateAndInitSingleView(GraphBuilder.RHICmdList, &ViewFamily, &ViewInitOptions);
 
 	const FSceneView* View = ViewFamily.Views[0];
-	auto* PassParameters = GraphBuilder.AllocParameters<FLandscapeGrassPassParameters>();
+	FLandscapeGrassPassParameters* PassParameters = GraphBuilder.AllocParameters<FLandscapeGrassPassParameters>();
 	PassParameters->View = View->ViewUniformBuffer;
 	PassParameters->Scene = GetSceneUniformBufferRef(GraphBuilder, *View);
 	PassParameters->RenderTargets[0] = FRenderTargetBinding(OutputTexture, ERenderTargetLoadAction::EClear);
@@ -475,44 +537,172 @@ void FLandscapeGrassWeightExporter_RenderThread::RenderLandscapeComponentToTextu
 
 			const uint64 DefaultBatchElementMask = 1 << 0; // LOD 0 only
 
-			for (auto& ComponentInfo : ComponentInfos)
+			for (FComponentInfo& ComponentInfo : ComponentInfos)
 			{
 				if (ensure(ComponentInfo.SceneProxy))
 				{
 					const FMeshBatch& Mesh = ComponentInfo.SceneProxy->GetGrassMeshBatch();
 					Mesh.MaterialRenderProxy->UpdateUniformExpressionCacheIfNeeded(View->GetFeatureLevel());
 
-					PassMeshProcessor.AddMeshBatch(Mesh, DefaultBatchElementMask, ComponentInfo.NumPasses, ComponentInfo.ViewOffset, PassOffsetX, ComponentInfo.FirstHeightMipsPassIndex, HeightMips, ComponentInfo.SceneProxy);
+					PassMeshProcessor.AddMeshBatch(Mesh, DefaultBatchElementMask, ComponentInfo.PerPassPerChannelGrassWeightIndices, ComponentInfo.ViewOffset, PassOffsetX, ComponentInfo.FirstHeightMipsPassIndex, HeightMips, ComponentInfo.SceneProxy);
 				}
 			}
 		});
 }
 
+TValueOrError<FIntRect /*OutputRect*/, FLandscapeGrassWeightExporter_RenderThread::EHeightmapTextureInfoError> FLandscapeGrassWeightExporter_RenderThread::GetTextureInfoForHeight(const FIntPoint& InComponentKey, int32 InMipIndex)
+{
+	using namespace UE::Landscape::Grass;
+	using namespace UE::Landscape::Grass::Private;
+
+	const FComponentInfo* FoundComponentInfo = ComponentInfos.FindByPredicate([InComponentKey](const FComponentInfo& InComponentInfo) { return (InComponentInfo.ComponentKey == InComponentKey); });
+	if (FoundComponentInfo == nullptr)
+	{
+		return MakeError(EHeightmapTextureInfoError::InvalidComponentKey);
+	}
+
+	// This means : get the height (mip 0), even if not in the requested mips (requires EGrassWeightExporterFlags::NeedsHeightmap)
+	if (InMipIndex < 0)
+	{
+		if (!EnumHasAnyFlags(Flags, EGrassWeightExporterFlags::NeedsHeightmap))
+		{
+			return MakeError(EHeightmapTextureInfoError::InvalidFlags);
+		}
+
+		// Height is always in the first pass if present :
+		check(!FoundComponentInfo->PerPassPerChannelGrassWeightIndices.IsEmpty());
+		// Height is always in the 2 first channels :
+		check((FoundComponentInfo->PerPassPerChannelGrassWeightIndices[0][0] == HeightGrassWeightIndex0)
+			&& (FoundComponentInfo->PerPassPerChannelGrassWeightIndices[0][1] == HeightGrassWeightIndex1));
+		FIntPoint PixelOffset(FoundComponentInfo->PixelOffsetX, 0);
+		FIntRect OutputRect(PixelOffset, PixelOffset + FIntPoint(GetTileSize(), GetTileSize()));
+		return MakeValue(OutputRect);
+	}
+	
+	// Other mips :
+	if (!HeightMips.Contains(InMipIndex))
+	{
+		return MakeError(EHeightmapTextureInfoError::InvalidMipIndex);
+	}
+
+	int32 MipOffset = HeightMips.IndexOfByKey(InMipIndex);
+	check(MipOffset != INDEX_NONE);
+	int32 PassIndexForMip = FoundComponentInfo->FirstHeightMipsPassIndex + MipOffset;
+	check(PassIndexForMip < FoundComponentInfo->GetNumPasses());
+
+	check(FoundComponentInfo->PerPassPerChannelGrassWeightIndices.IsValidIndex(PassIndexForMip));
+	// Height is always in the 2 first channels :
+	check((FoundComponentInfo->PerPassPerChannelGrassWeightIndices[PassIndexForMip][0] == HeightGrassWeightIndex0)
+		&& (FoundComponentInfo->PerPassPerChannelGrassWeightIndices[PassIndexForMip][1] == HeightGrassWeightIndex1));
+
+	// Passes don't care about mip size so the offset is always a multiple of GetTileSize() :
+	FIntPoint PixelOffset(FoundComponentInfo->PixelOffsetX + PassIndexForMip * GetTileSize(), 0);
+	
+	// Compute the number of vertices of that mip (i.e. without subsection duplicates)
+	int32 MipSize = NumSubsections * (SubsectionSizeQuads >> InMipIndex) + 1;
+
+	FIntRect OutputRect(PixelOffset, PixelOffset + FIntPoint(MipSize, MipSize));
+	return MakeValue(OutputRect);
+}
+
+TValueOrError<TPair<FIntRect /*OutputRect*/, uint8 /*ChannelIndex*/>, FLandscapeGrassWeightExporter_RenderThread::EGrassmapTextureInfoError> FLandscapeGrassWeightExporter_RenderThread::GetTextureInfoForGrass(const FIntPoint& InComponentKey, ULandscapeGrassType* InGrassType)
+{
+	using namespace UE::Landscape::Grass;
+	using namespace UE::Landscape::Grass::Private;
+
+	const FComponentInfo* FoundComponentInfo = ComponentInfos.FindByPredicate([InComponentKey](const FComponentInfo& InComponentInfo) { return (InComponentInfo.ComponentKey == InComponentKey); });
+	if (FoundComponentInfo == nullptr)
+	{
+		return MakeError(EGrassmapTextureInfoError::InvalidComponentKey);
+	}
+
+	if (!EnumHasAnyFlags(Flags, EGrassWeightExporterFlags::NeedsGrassmap))
+	{
+		return MakeError(EGrassmapTextureInfoError::InvalidFlags);
+	}
+
+	const FComponentInfo::FPassChannelIndex* FoundPassChannelIndex = FoundComponentInfo->PerGrassTypePassChannelIndices.Find(InGrassType);
+	if (FoundPassChannelIndex == nullptr)
+	{
+		return MakeError(EGrassmapTextureInfoError::InvalidGrassType);
+	}
+
+	check(FoundComponentInfo->PerPassPerChannelGrassWeightIndices.IsValidIndex(FoundPassChannelIndex->PassIndex));
+	check(FoundPassChannelIndex->IsValidChannelIndex());
+	check(FoundComponentInfo->PerPassPerChannelGrassWeightIndices[FoundPassChannelIndex->PassIndex][FoundPassChannelIndex->ChannelIndex] != GrassWeightInvalidIndex);
+
+	FIntPoint PixelOffset(FoundComponentInfo->PixelOffsetX + FoundPassChannelIndex->PassIndex * GetTileSize(), 0);
+	FIntRect OutputRect(PixelOffset, PixelOffset + FIntPoint(GetTileSize(), GetTileSize()));
+	return MakeValue(MakeTuple(OutputRect, FoundPassChannelIndex->ChannelIndex));
+}
+
+namespace UE::Landscape::Private
+{
+// Remove once the old FLandscapeGrassWeightExporter constructor is deprecated : 
+UE::Landscape::Grass::EGrassWeightExporterFlags BuildLandscapeGrassWeightExporterConstructionFlags(bool bInNeedsGrassmap, bool bInNeedsHeightmap, bool bInRenderImmediately, bool bInReadbackToCPU)
+{
+	using namespace UE::Landscape::Grass;
+
+	EGrassWeightExporterFlags LandscapeGrassWeightExporterFlags = EGrassWeightExporterFlags::None;
+	if (bInNeedsGrassmap)
+	{
+		LandscapeGrassWeightExporterFlags |= EGrassWeightExporterFlags::NeedsGrassmap;
+	}
+	if (bInNeedsHeightmap)
+	{
+		LandscapeGrassWeightExporterFlags |= EGrassWeightExporterFlags::NeedsHeightmap;
+	}
+	if (bInRenderImmediately)
+	{
+		LandscapeGrassWeightExporterFlags |= EGrassWeightExporterFlags::RenderImmediately;
+	}
+	if (bInReadbackToCPU)
+	{
+		LandscapeGrassWeightExporterFlags |= EGrassWeightExporterFlags::ReadbackToCPU;
+	}
+	return LandscapeGrassWeightExporterFlags;
+}
+} // namespace UE::Landscape::Private
+
+
+// Deprecated
 FLandscapeGrassWeightExporter::FLandscapeGrassWeightExporter(ALandscapeProxy* InLandscapeProxy, TArrayView<ULandscapeComponent* const> InLandscapeComponents, bool bInNeedsGrassmap, bool bInNeedsHeightmap, const TArray<int32>& InHeightMips, bool bInRenderImmediately, bool bInReadbackToCPU)
-	: FLandscapeGrassWeightExporter_RenderThread(InHeightMips, bInReadbackToCPU)
-	, LandscapeProxy(InLandscapeProxy)
-	, ComponentSizeVerts(InLandscapeProxy->ComponentSizeQuads + 1)
-	, SubsectionSizeQuads(InLandscapeProxy->SubsectionSizeQuads)
-	, NumSubsections(InLandscapeProxy->NumSubsections)
+	: FLandscapeGrassWeightExporter(
+		InLandscapeProxy->GetLandscapeActor(),
+		InLandscapeComponents,
+		UE::Landscape::Private::BuildLandscapeGrassWeightExporterConstructionFlags(bInNeedsGrassmap, bInNeedsHeightmap, bInRenderImmediately, bInReadbackToCPU),
+		/*InRequestedGrassTypes = */{},
+		MakeArrayView(InHeightMips))
 {
-	check(InLandscapeComponents.Num() > 0);
-	SceneInterface = InLandscapeComponents[0]->GetScene();
+}
+
+FLandscapeGrassWeightExporter::FLandscapeGrassWeightExporter(ALandscape* InLandscape, TArrayView<ULandscapeComponent* const> InRequestedComponents, UE::Landscape::Grass::EGrassWeightExporterFlags InFlags,
+	TConstArrayView<TObjectPtr<ULandscapeGrassType>> InRequestedGrassTypes, TConstArrayView<int32> InRequestedHeightMips)
+	: FLandscapeGrassWeightExporter_RenderThread(/*ComponentSizeVerts = */InLandscape->ComponentSizeQuads + 1, InLandscape->SubsectionSizeQuads, InLandscape->NumSubsections, InRequestedHeightMips, InFlags)
+	, Landscape(InLandscape)
+{
+	check(InLandscape != nullptr);
+	check(InRequestedComponents.Num() > 0);
+	SceneInterface = InRequestedComponents[0]->GetScene();
 
 	// todo: use a 2d target?
 	const int32 SingleTileWidth = ComponentSizeVerts;
 	TargetSize = FIntPoint(0, ComponentSizeVerts);
 
 	// First compute the total render target size and prepare ComponentInfos (each component has its own number of needed passes because some might have different materials, thus different grass types) :
-	ComponentInfos.Reserve(InLandscapeComponents.Num());
+	ComponentInfos.Reserve(InRequestedComponents.Num());
 	int32 CurrentPixelOffsetX = 0;
-	for (ULandscapeComponent* Component : InLandscapeComponents)
+	FIntPoint MinSectionBase(MAX_int32, MAX_int32);
+	for (ULandscapeComponent* Component : InRequestedComponents)
 	{
+		check(InLandscape == Component->GetLandscapeActor());
+		MinSectionBase = MinSectionBase.ComponentMin(Component->GetSectionBase());
 		ensure(Component->SceneProxy);
 
-		FComponentInfo& ComponentInfo = ComponentInfos.Add_GetRef(FComponentInfo(Component, bInNeedsGrassmap, bInNeedsHeightmap, InHeightMips));
+		FComponentInfo& ComponentInfo = ComponentInfos.Emplace_GetRef(Component, InFlags, InRequestedGrassTypes, InRequestedHeightMips);
 		ComponentInfo.PixelOffsetX = CurrentPixelOffsetX;
 
-		CurrentPixelOffsetX += ComponentInfo.NumPasses * SingleTileWidth;
+		CurrentPixelOffsetX += ComponentInfo.GetNumPasses() * SingleTileWidth;
 	}
 	TargetSize.X = CurrentPixelOffsetX;
 
@@ -520,9 +710,13 @@ FLandscapeGrassWeightExporter::FLandscapeGrassWeightExporter(ALandscapeProxy* In
 	PassOffsetX = 2.0f * (float)SingleTileWidth / (float)TargetSize.X;
 
 	// Then compute FComponentInfo's ViewOffset with the knowledge of the total render target size : 
-	for (FComponentInfo& ComponentInfo : ComponentInfos)
+	const int32 NumComponents = ComponentInfos.Num();
+	check(NumComponents == InRequestedComponents.Num());
+	for (int32 ComponentIndex = 0; ComponentIndex < NumComponents; ++ComponentIndex)
 	{
-		FIntPoint ComponentOffset = (ComponentInfo.Component->GetSectionBase() - LandscapeProxy->GetSectionBase());
+		FComponentInfo& ComponentInfo = ComponentInfos[ComponentIndex];
+		ULandscapeComponent* Component = InRequestedComponents[ComponentIndex];
+		FIntPoint ComponentOffset = (Component->GetSectionBase() - InLandscape->GetSectionBase());
 		FVector2D ViewOffset(-ComponentOffset.X, ComponentOffset.Y);
 		ViewOffset.X += ComponentInfo.PixelOffsetX;
 		ViewOffset /= (FVector2D(TargetSize) * 0.5f);
@@ -530,13 +724,13 @@ FLandscapeGrassWeightExporter::FLandscapeGrassWeightExporter(ALandscapeProxy* In
 	}
 
 	// center of target area in world
-	FVector TargetCenter = LandscapeProxy->GetTransform().TransformPosition(FVector(TargetSizeMinusOne, 0.f) * 0.5f);
+	FVector TargetCenter = InLandscape->GetTransform().TransformPosition(FVector(TargetSizeMinusOne, 0.f) * 0.5f);
 
 	// extent of target in world space
-	FVector TargetExtent = FVector(TargetSize, 0.0f) * LandscapeProxy->GetActorScale() * 0.5f;
+	FVector TargetExtent = FVector(TargetSize, 0.0f) * InLandscape->GetActorScale() * 0.5f;
 
 	ViewOrigin = TargetCenter;
-	ViewRotationMatrix = FInverseRotationMatrix(LandscapeProxy->GetActorRotation());
+	ViewRotationMatrix = FInverseRotationMatrix(InLandscape->GetActorRotation());
 	ViewRotationMatrix *= FMatrix(FPlane(1.0f, 0.0f, 0.0f, 0.0f),
 		FPlane(0.0f, -1.0f, 0.0f, 0.0f),
 		FPlane(0.0f, 0.0f, -1.0f, 0.0f),
@@ -549,7 +743,7 @@ FLandscapeGrassWeightExporter::FLandscapeGrassWeightExporter(ALandscapeProxy* In
 		0.5f / ZOffset,
 		ZOffset);
 
-	if (bInRenderImmediately)
+	if (EnumHasAnyFlags(InFlags, UE::Landscape::Grass::EGrassWeightExporterFlags::RenderImmediately))
 	{
 		UE::RenderCommandPipe::FSyncScope SyncScope;
 
@@ -582,11 +776,18 @@ bool FLandscapeGrassWeightExporter::IsAsyncReadbackComplete()
 struct FByteBuffer2DView : public IBuffer2DView<uint8>
 {
 	uint8* BufferStart = nullptr;
-	int32 ByteStrideY = 0;
 	int32 ByteStrideX = 0;
+	int32 ByteStrideY = 0;
 	int32 NumX = 0;
 	int32 NumY = 0;
 
+	FByteBuffer2DView(int32 InByteStrideX, int32 InByteStrideY, int32 InNumX, int32 InNumY)
+		: ByteStrideX(InByteStrideX)
+		, ByteStrideY(InByteStrideY)
+		, NumX(InNumX)
+		, NumY(InNumY)
+	{}
+
 	// copy elements from buffer to Dest, in X then Y order
 	virtual void CopyTo(uint8* Dest, int32 SizeInBytes) const override
 	{
@@ -627,23 +828,13 @@ struct FByteBuffer2DView : public IBuffer2DView<uint8>
 		return (MaxBits == 0);
 	}
 
-	virtual int32 Num() const override { return NumX * NumY; }
+	virtual int32 Num() const override 
+	{ 
+		return NumX * NumY; 
+	}
 };
 
-// FColor memory layout matches our BGRA GPU layout only on little endian CPUs!
-#if PLATFORM_LITTLE_ENDIAN
-	#define BGRA_AS_FCOLOR_BLUE B
-	#define BGRA_AS_FCOLOR_GREEN G
-	#define BGRA_AS_FCOLOR_RED R
-	#define BGRA_AS_FCOLOR_ALPHA A
-#else
-	#define BGRA_AS_FCOLOR_BLUE A
-	#define BGRA_AS_FCOLOR_GREEN R
-	#define BGRA_AS_FCOLOR_RED G
-	#define BGRA_AS_FCOLOR_ALPHA B
-#endif // PLATFORM_LITTLE_ENDIAN
-
-struct FHeightBuffer2DView : IBuffer2DView<uint16>
+struct FHeightBuffer2DView : public IBuffer2DView<uint16>
 {
 	FColor* BufferStart = nullptr;
 	int32 StrideY = 0;
@@ -653,13 +844,14 @@ struct FHeightBuffer2DView : IBuffer2DView<uint16>
 	// copy elements from buffer to Dest, in X then Y order
 	virtual void CopyTo(uint16* Dest, int32 Count) const override
 	{
+		static_assert(PLATFORM_LITTLE_ENDIAN, "Big-endian is not supported anymore");
 		for (int y = 0; Count > 0 && y < NumY; y++)
 		{
 			FColor* Src = BufferStart + y * StrideY;
 			int32 CopyCountX = FMath::Min(Count, NumX);
 			while (CopyCountX--)
 			{
-				*Dest = (((uint16) Src->BGRA_AS_FCOLOR_RED) << 8) + (uint16)(Src->BGRA_AS_FCOLOR_GREEN);
+				*Dest = (((uint16)Src->R) << 8) + (uint16)(Src->G);
 				Dest++;
 				Src++;
 			}
@@ -704,7 +896,11 @@ void FLandscapeGrassWeightExporter::CancelAndSelfDestruct()
 
 TMap<ULandscapeComponent*, TUniquePtr<FLandscapeComponentGrassData>, TInlineSetAllocator<1>> FLandscapeGrassWeightExporter::FetchResults(bool bFreeAsyncReadback)
 {
+	using namespace UE::Landscape::Grass;
+	using namespace UE::Landscape::Grass::Private;
+
 	TRACE_CPUPROFILER_EVENT_SCOPE(FetchResults);
+
 	TMap<ULandscapeComponent*, TUniquePtr<FLandscapeComponentGrassData>, TInlineSetAllocator<1>> Results;
 	TArray<FColor> Samples;
 
@@ -722,16 +918,22 @@ TMap<ULandscapeComponent*, TUniquePtr<FLandscapeComponentGrassData>, TInlineSetA
 
 	Results.Reserve(ComponentInfos.Num());
 	FHeightBuffer2DView HeightData;
-	TMap<ULandscapeGrassType*, IBuffer2DView<uint8>*> WeightData;
+	TMap<ULandscapeGrassType*, FByteBuffer2DView> WeightData;
+
+	ULandscapeInfo* LandscapeInfo = Landscape->GetLandscapeInfo();
+	check(LandscapeInfo != nullptr);
 
-	for (auto& ComponentInfo : ComponentInfos)
+	for (FComponentInfo& ComponentInfo : ComponentInfos)
 	{
-		ULandscapeComponent* Component = ComponentInfo.Component;
+		ULandscapeComponent** ComponentPtr = LandscapeInfo->XYtoComponentMap.Find(ComponentInfo.ComponentKey);
+		if ((ComponentPtr == nullptr) || (*ComponentPtr == nullptr))
+		{
+			continue;
+		}
+		ULandscapeComponent* Component = *ComponentPtr;
 
 		TUniquePtr<FLandscapeComponentGrassData> NewGrassData = MakeUnique<FLandscapeComponentGrassData>(Component);
 
-		int32 ComponentSizeVerts2 = FMath::Square(ComponentSizeVerts);
-
 		HeightData.NumX = ComponentSizeVerts;
 		HeightData.NumY = ComponentSizeVerts;
 		HeightData.StrideY = TargetSize.X;
@@ -741,79 +943,166 @@ TMap<ULandscapeComponent*, TUniquePtr<FLandscapeComponentGrassData>, TInlineSetA
 #endif // WITH_EDITORONLY_DATA
 
 		WeightData.Empty();
+		WeightData.Reserve(ComponentInfo.PerGrassTypePassChannelIndices.Num());
 
-		// this array is in 1:1 correspondence with ComponentInfo.RequestedGrassTypes
-		TArray<FByteBuffer2DView> GrassWeightArrays;
-		GrassWeightArrays.SetNum(ComponentInfo.RequestedGrassTypes.Num());
+		for (auto& ItPair : ComponentInfo.PerGrassTypePassChannelIndices)
+		{
+			// Note: WeightData points directly at the elements of GrassWeightArrays (DO NOT REALLOCATE GRASSWEIGHTARRAYS)
+			WeightData.Add(ItPair.Key, FByteBuffer2DView(/*InByteStrideX = */4, /*InByteStrideY = */TargetSize.X * 4, /*InNumX = */ComponentSizeVerts, /*InNumY = */ComponentSizeVerts));
+		}
 
-		for (int32 Index = 0; Index < GrassWeightArrays.Num(); Index++)
+		auto FindGrassTypeForPassAndGrassWeight = [&](const FComponentInfo::FPassChannelIndex& InPassChannelIndex) -> ULandscapeGrassType*
 		{
-			FByteBuffer2DView* WeightView = &GrassWeightArrays[Index];
-			ULandscapeGrassType* GrassType = ComponentInfo.RequestedGrassTypes[Index];
+			for (auto& ItPair : ComponentInfo.PerGrassTypePassChannelIndices)
+			{
+				if (ItPair.Value == InPassChannelIndex)
+				{
+					return ItPair.Key;
+				}
+			}
+			return nullptr;
+		};
 
-			WeightView->NumX = ComponentSizeVerts;
-			WeightView->NumY = ComponentSizeVerts;
-			WeightView->ByteStrideX = 4;
-			WeightView->ByteStrideY = TargetSize.X * 4;
-			WeightView->BufferStart = nullptr;
+#if WITH_EDITOR
+		// output debug bitmap
+		if (CVarDumpGrassmaps.GetValueOnGameThread())
+		{
+			FString WorldName = Component->GetWorld()->GetName();
+			ULandscapeSubsystem* LandscapeSubsystem = Component->GetWorld()->GetSubsystem<ULandscapeSubsystem>();
+			check(LandscapeSubsystem != nullptr);
+			const FDateTime CurrentTime = LandscapeSubsystem->GetAppCurrentDateTime();
+			FString ParentLandscapeActorName = Landscape->GetActorLabel();
+			ALandscapeProxy* Proxy = Cast<ALandscapeProxy>(Component->GetOwner());
+			check(Proxy != nullptr);
+			FString ActorName = Proxy->GetActorLabel();
+			FString FilePath = FString::Format(TEXT("{0}/LandscapeGrass/{1}/{2}/{3}/Grassmaps/{4}-{5}"), 
+				{ FPaths::ProjectSavedDir(), CurrentTime.ToString(), WorldName, ParentLandscapeActorName, ActorName, Component->GetName() });
+			FStringBuilderBase StrBuilder;
+			for (int32 PassIndex = 0; PassIndex < ComponentInfo.GetNumPasses(); ++PassIndex)
+			{
+				auto GetChannelString = [&](const FComponentInfo::FPassChannelIndex& InPassChannelIndex) -> FString
+					{
+						FString Result(TEXT("<empty>"));
+						int32 GrassWeightIndex = ComponentInfo.PerPassPerChannelGrassWeightIndices[InPassChannelIndex.PassIndex][InPassChannelIndex.ChannelIndex];
+						if (GrassWeightIndex != GrassWeightInvalidIndex)
+						{
+							if (InPassChannelIndex.PassIndex < ComponentInfo.FirstHeightMipsPassIndex)
+							{
+								if (GrassWeightIndex == HeightGrassWeightIndex0)
+								{
+									Result = FString(TEXT("[Mip=0]Height0"));
+								}
+								else if (GrassWeightIndex == HeightGrassWeightIndex1)
+								{
+									Result = FString(TEXT("[Mip=0]Height1"));
+								}
+								else 
+								{
+									ULandscapeGrassType* GrassType = FindGrassTypeForPassAndGrassWeight(InPassChannelIndex);
+									check(GrassType != nullptr);
+									Result = GrassType->GetName();
+								}
+							}
+							else
+							{
+								int32 HeightMipIndex = InPassChannelIndex.PassIndex - ComponentInfo.FirstHeightMipsPassIndex;
+								check(HeightMips.IsValidIndex(HeightMipIndex));
+								if (GrassWeightIndex == HeightGrassWeightIndex0)
+								{
+									Result = FString::Printf(TEXT("[Mip=%i]Height0"), HeightMips[HeightMipIndex]);
+								}
+								else if (GrassWeightIndex == HeightGrassWeightIndex1)
+								{
+									Result = FString::Printf(TEXT("[Mip=%i]Height1"), HeightMips[HeightMipIndex]);
+								}
+							}
+						}
+						return Result;
+					};
+
+				int32 OffsetX = ComponentInfo.PixelOffsetX + PassIndex * ComponentSizeVerts;
+				StrBuilder.Appendf(TEXT("- OffsetX = %i :\n"), OffsetX);
+				StrBuilder.Appendf(TEXT("R : %s\n"), *GetChannelString(FComponentInfo::FPassChannelIndex { PassIndex, 0 }));
+				StrBuilder.Appendf(TEXT("G : %s\n"), *GetChannelString(FComponentInfo::FPassChannelIndex { PassIndex, 1 }));
+				StrBuilder.Appendf(TEXT("B : %s\n"), *GetChannelString(FComponentInfo::FPassChannelIndex { PassIndex, 2 }));
+				StrBuilder.Appendf(TEXT("A : %s\n"), *GetChannelString(FComponentInfo::FPassChannelIndex { PassIndex, 3 }));
+			}
+			FFileHelper::SaveStringToFile(StrBuilder.ToView(), *(FilePath + ".desc.txt"));
+			FFileHelper::CreateBitmap(*(FilePath + ".bmp"), TargetSize.X, TargetSize.Y, Samples.GetData(), /*SubRegion = */nullptr, &IFileManager::Get(), /*OutFilename = */nullptr, /*bInWriteAlpha = */!ComponentInfo.PerGrassTypePassChannelIndices.IsEmpty());
+		}
+#endif // WITH_EDITOR
+	
+#if GRASS_ENABLE_EXTRAVALIDATION
+		FIntRect HeightRect;
+		if (EnumHasAnyFlags(Flags, EGrassWeightExporterFlags::NeedsHeightmap))
+		{
+			TValueOrError<FIntRect /*OutputRect*/, EHeightmapTextureInfoError> TextureInfo = GetTextureInfoForHeight(ComponentInfo.ComponentKey);
+			checkf(!TextureInfo.HasError(), TEXT("GetTextureInfoForHeight failed (error : %i)"), static_cast<int32>(TextureInfo.GetError()));
+			HeightRect = TextureInfo.GetValue();
+		}
 
-			// Note: WeightData points directly at the elements of GrassWeightArrays (DO NOT REALLOCATE GRASSWEIGHTARRAYS)
-			WeightData.Add(GrassType, WeightView);
+		TMap<int32, FIntRect> HeightMipsRects;
+		HeightMipsRects.Reserve(HeightMips.Num());
+		for (int32 MipIndex : HeightMips)
+		{
+			TValueOrError<FIntRect /*OutputRect*/, EHeightmapTextureInfoError> TextureInfo = GetTextureInfoForHeight(ComponentInfo.ComponentKey, MipIndex);
+			checkf(!TextureInfo.HasError(), TEXT("GetTextureInfoForHeight failed for mip %i (error : %i)"), MipIndex, static_cast<int32>(TextureInfo.GetError()));
+			HeightMipsRects.Add(MipIndex, TextureInfo.GetValue());
 		}
 
-		// output debug bitmap
-#if UE_BUILD_DEBUG
-		static bool bOutputGrassBitmap = false;
-		if (bOutputGrassBitmap)
+		TMap<ULandscapeGrassType*, TPair<FIntRect, int32>> GrassRectsAndChannelIndices;
+		for (auto& ItPair : ComponentInfo.PerGrassTypePassChannelIndices)
 		{
-			FString TempPath = FPaths::ScreenShotDir();
-			TempPath += TEXT("/GrassDebug");
-			IFileManager::Get().MakeDirectory(*TempPath, true);
-			FFileHelper::CreateBitmap(*(TempPath / "Grass"), TargetSize.X, TargetSize.Y, Samples.GetData(), nullptr, &IFileManager::Get(), nullptr, ComponentInfo.RequestedGrassTypes.Num() >= 2);
+			TValueOrError<TPair<FIntRect /*OutputRect*/, uint8 /*ChannelIndex*/>, EGrassmapTextureInfoError> TextureInfo = GetTextureInfoForGrass(ComponentInfo.ComponentKey, ItPair.Key);
+			checkf(!TextureInfo.HasError(), TEXT("GetTextureInfoForGrass failed for grass type %s (error : %i)"), *ItPair.Key->GetName(), static_cast<int32>(TextureInfo.GetError()));
+			GrassRectsAndChannelIndices.Add(ItPair.Key, TextureInfo.GetValue());
 		}
-#endif
+#endif // GRASS_ENABLE_EXTRAVALIDATION
 
-		int32 GrassTypeCount = ComponentInfo.RequestedGrassTypes.Num();
-		for (int32 PassIdx = 0; PassIdx < ComponentInfo.NumPasses; PassIdx++)
+		for (int32 PassIdx = 0; PassIdx < ComponentInfo.GetNumPasses(); PassIdx++)
 		{
 			FColor* SampleData = &Samples[ComponentInfo.PixelOffsetX + PassIdx * ComponentSizeVerts];
+			const FIntVector4& PerChannelGrassWeightIndices = ComponentInfo.PerPassPerChannelGrassWeightIndices[PassIdx];
 
 			if (PassIdx < ComponentInfo.FirstHeightMipsPassIndex)
 			{
+				int32 ChannelIndex = 0;
 				if (PassIdx == 0)	// height in RG, grass weights in BA
 				{
+#if GRASS_ENABLE_EXTRAVALIDATION
+					FColor* HeightSampleData = &Samples[HeightRect.Min.X];
+					check(HeightSampleData == SampleData);
+#endif // GRASS_ENABLE_EXTRAVALIDATION
+
+					check((PerChannelGrassWeightIndices[0] == HeightGrassWeightIndex0) && (PerChannelGrassWeightIndices[1] == HeightGrassWeightIndex1));
 					HeightData.BufferStart = SampleData;
-					if (GrassTypeCount > 0)
-					{
-						GrassWeightArrays[0].BufferStart = &SampleData->BGRA_AS_FCOLOR_BLUE;
-						if (GrassTypeCount > 1)
-						{
-							GrassWeightArrays[1].BufferStart = &SampleData->BGRA_AS_FCOLOR_ALPHA;
-						}
-					}
+					ChannelIndex += 2;
 				}
-				else
-				{
-					int32 TypeIdx = PassIdx * 4 - 2;
 
-					GrassWeightArrays[TypeIdx+0].BufferStart = &SampleData->BGRA_AS_FCOLOR_RED;
-					if (GrassTypeCount > TypeIdx+1)
+				for (; ChannelIndex < 4; ++ChannelIndex)
+				{
+					if (PerChannelGrassWeightIndices[ChannelIndex] != GrassWeightInvalidIndex)
 					{
-						GrassWeightArrays[TypeIdx+1].BufferStart = &SampleData->BGRA_AS_FCOLOR_GREEN;
-						if (GrassTypeCount > TypeIdx + 2)
-						{
-							GrassWeightArrays[TypeIdx + 2].BufferStart = &SampleData->BGRA_AS_FCOLOR_BLUE;
-							if (GrassTypeCount > TypeIdx + 3)
-							{
-								GrassWeightArrays[TypeIdx + 3].BufferStart = &SampleData->BGRA_AS_FCOLOR_ALPHA;
-							}
-						}
+						ULandscapeGrassType* GrassType = FindGrassTypeForPassAndGrassWeight(FComponentInfo::FPassChannelIndex { PassIdx, ChannelIndex });
+
+#if GRASS_ENABLE_EXTRAVALIDATION
+						TPair<FIntRect, int32>* RectAndIndex = GrassRectsAndChannelIndices.Find(GrassType);
+						check(RectAndIndex != nullptr);
+						check(ChannelIndex == RectAndIndex->Value);
+						FColor* GrassSampleData = &Samples[RectAndIndex->Key.Min.X];
+						check(GrassSampleData == SampleData);
+#endif // GRASS_ENABLE_EXTRAVALIDATION
+
+						FByteBuffer2DView* ByteBufferView = WeightData.Find(GrassType);
+						check((GrassType != nullptr) && (ByteBufferView != nullptr));
+						static_assert(PLATFORM_LITTLE_ENDIAN, "Big-endian is not supported anymore");
+						ByteBufferView->BufferStart = reinterpret_cast<uint8*>((ChannelIndex == 0) ? &SampleData->R : (ChannelIndex == 1) ? &SampleData->G : (ChannelIndex == 2) ? &SampleData->B : &SampleData->A);
 					}
 				}
 			}
 			else // PassIdx >= FirstHeightMipsPassIndex
 			{
-#if WITH_EDITORONLY_DATA
+#if WITH_EDITOR
 				const int32 Mip = HeightMips[PassIdx - ComponentInfo.FirstHeightMipsPassIndex];
 				int32 MipSizeVerts = NumSubsections * (SubsectionSizeQuads >> Mip);
 				TArray<uint16>& MipHeightData = NewGrassData->HeightMipData.Add(Mip);
@@ -824,20 +1113,30 @@ TMap<ULandscapeComponent*, TUniquePtr<FLandscapeComponentGrassData>, TInlineSetA
 					FColor* SrcSample = &SampleData[y * TargetSize.X];
 					for (int32 x = 0; x < MipSizeVerts; x++)
 					{
-						*DstMipHeight++ = (((uint16)SrcSample->BGRA_AS_FCOLOR_RED) << 8) + (uint16)(SrcSample->BGRA_AS_FCOLOR_GREEN);
+						static_assert(PLATFORM_LITTLE_ENDIAN, "Big-endian is not supported anymore");
+						*DstMipHeight++ = (((uint16)SrcSample->R) << 8) + (uint16)(SrcSample->G);
 						SrcSample++;
 					}
 				}
-#endif // WITH_EDITORONLY_DATA
+
+#if GRASS_ENABLE_EXTRAVALIDATION
+				FIntRect* MipRect = HeightMipsRects.Find(Mip);
+				check(MipRect != nullptr);
+				FColor* MipSampleData = &Samples[MipRect->Min.X];
+				check(MipSampleData == SampleData);
+#endif // GRASS_ENABLE_EXTRAVALIDATION
+
+#endif // WITH_EDITOR
 			}
 		}
-
-		#undef BGRA_AS_FCOLOR_BLUE
-		#undef BGRA_AS_FCOLOR_GREEN
-		#undef BGRA_AS_FCOLOR_RED
-		#undef BGRA_AS_FCOLOR_ALPHA
 		
-		NewGrassData->InitializeFrom(&HeightData, WeightData, /* bStripEmptyWeights = */ true);
+		TMap<ULandscapeGrassType*, IBuffer2DView<uint8>*> WeightDataPtrs;
+		WeightDataPtrs.Reserve(WeightData.Num());
+		for (auto& ItPair : WeightData)
+		{
+			WeightDataPtrs.Add(ItPair.Key, static_cast<IBuffer2DView<uint8>*>(&ItPair.Value));
+		}
+		NewGrassData->InitializeFrom(&HeightData, WeightDataPtrs, /* bStripEmptyWeights = */ true);
 		Results.Add(Component, MoveTemp(NewGrassData));
 	}
 
@@ -874,3 +1173,87 @@ void FLandscapeGrassWeightExporter::ApplyResults(TMap<ULandscapeComponent*, TUni
 #endif // WITH_EDITOR
 	}
 }
+
+FLandscapeGrassWeightExporter_RenderThread::FComponentInfo::FComponentInfo(ULandscapeComponent* InComponent, UE::Landscape::Grass::EGrassWeightExporterFlags InFlags,
+	TConstArrayView<TObjectPtr<ULandscapeGrassType>> InRequestedGrassTypes, TConstArrayView<int32> InRequestedHeightMips) 
+	: ComponentKey(InComponent->GetComponentKey())
+	, SceneProxy((FLandscapeComponentSceneProxy*)InComponent->SceneProxy)
+	, Flags(InFlags)
+{
+	using namespace UE::Landscape::Grass;
+	using namespace UE::Landscape::Grass::Private;
+
+	int32 CurrentPassIndex = -1;
+	int32 CurrentChannelIndex = -1;
+
+	auto RequestNewPass = [&]() 
+		{
+			// Setting the channel index to -1 will ensure that a new pass is created in the next call to FindOrAddChannel :
+			CurrentChannelIndex = -1;
+		};
+
+	auto FindOrAddChannel = [&]() -> FPassChannelIndex
+		{
+			if ((CurrentChannelIndex < 0) || (CurrentChannelIndex >= 3))
+			{
+				// Add a new pass :
+				PerPassPerChannelGrassWeightIndices.Emplace(GrassWeightInvalidIndex, GrassWeightInvalidIndex, GrassWeightInvalidIndex, GrassWeightInvalidIndex);
+				CurrentChannelIndex = -1;
+				++CurrentPassIndex;
+			}
+			return FPassChannelIndex{ CurrentPassIndex, ++CurrentChannelIndex };
+		};
+
+	auto SetValueAtPassIndex = [&](FPassChannelIndex InPassChannelIndex, int32 InValue)
+		{
+			check(PerPassPerChannelGrassWeightIndices.IsValidIndex(InPassChannelIndex.PassIndex));
+			check(InPassChannelIndex.IsValidChannelIndex());
+			PerPassPerChannelGrassWeightIndices[InPassChannelIndex.PassIndex][InPassChannelIndex.ChannelIndex] = InValue;
+		};
+
+	// Start with 2 channels for heightmaps if needed : 
+	if (EnumHasAnyFlags(InFlags, EGrassWeightExporterFlags::NeedsHeightmap))
+	{
+		SetValueAtPassIndex(FindOrAddChannel(), HeightGrassWeightIndex0);
+		SetValueAtPassIndex(FindOrAddChannel(), HeightGrassWeightIndex1);
+	}
+
+	if (EnumHasAnyFlags(InFlags, EGrassWeightExporterFlags::NeedsGrassmap))
+	{
+		TArray<TObjectPtr<ULandscapeGrassType>> RenderableGrassTypes = InComponent->GetGrassTypes();
+		TArray<TObjectPtr<ULandscapeGrassType>> RequestedGrassTypes = InRequestedGrassTypes.IsEmpty() ? RenderableGrassTypes : TArray<TObjectPtr<ULandscapeGrassType>>(InRequestedGrassTypes);
+		for (TObjectPtr<ULandscapeGrassType> RequestedGrassType : RequestedGrassTypes)
+		{
+			if (RequestedGrassType != nullptr)
+			{
+				if (int32 GrassOutputIndex = RenderableGrassTypes.IndexOfByKey(RequestedGrassType);
+					GrassOutputIndex != INDEX_NONE)
+				{
+					check(GrassOutputIndex < UMaterialExpressionLandscapeGrassOutput::MaxGrassTypes);
+					FPassChannelIndex PassChannelIndex = FindOrAddChannel();
+					SetValueAtPassIndex(PassChannelIndex, GrassOutputIndex);
+					PerGrassTypePassChannelIndices.Add(RequestedGrassType, PassChannelIndex);
+				}
+			}
+		}
+	}
+
+	// since we don't read HeightMips unless we are in editor, there's no reason to add passes for it unless we are in editor
+#if WITH_EDITOR
+	// Add a new pass for each height mip, for each pass we'll use just 2 channels :  
+	if (!InRequestedHeightMips.IsEmpty())
+	{
+		FirstHeightMipsPassIndex = CurrentPassIndex + 1;
+
+		for (int32 MipIndex : InRequestedHeightMips)
+		{
+			// Request a new pass for each height mip :
+			RequestNewPass();
+			SetValueAtPassIndex(FindOrAddChannel(), HeightGrassWeightIndex0);
+			SetValueAtPassIndex(FindOrAddChannel(), HeightGrassWeightIndex1);
+		}
+	}
+#endif // WITH_EDITOR
+}
+
+#undef GRASS_ENABLE_EXTRAVALIDATION
\ No newline at end of file
diff --git a/Engine/Source/Runtime/Landscape/Private/LandscapeRender.cpp b/Engine/Source/Runtime/Landscape/Private/LandscapeRender.cpp
--- a/Engine/Source/Runtime/Landscape/Private/LandscapeRender.cpp
+++ b/Engine/Source/Runtime/Landscape/Private/LandscapeRender.cpp
@@ -54,6 +54,7 @@ LandscapeRender.cpp: New terrain rendering
 #include "RenderGraphBuilder.h"
 #include "Scalability.h"
 #include "Rendering/CustomRenderPass.h"
+#include "LandscapeUtils.h"
 #include "LandscapeUtilsPrivate.h"
 #include "SceneRendererInterface.h"
 
@@ -243,9 +244,6 @@ extern TAutoConsoleVariable<int32> CVarLandscapeShowDirty;
 
 extern RENDERER_API TAutoConsoleVariable<float> CVarStaticMeshLODDistanceScale;
 
-extern int32 GGrassMapUseRuntimeGeneration;
-extern int32 GGrassMapAlwaysBuildRuntimeGenerationResources;
-
 #if !UE_BUILD_SHIPPING
 uint32 GVarDumpLandscapeLODsCurrentFrame = 0;
 bool GVarDumpLandscapeLODsDetails = false;
@@ -367,7 +365,7 @@ namespace UE::Landscape
 		#if WITH_EDITOR
 			return true;
 		#else
-			return (GGrassMapUseRuntimeGeneration || GGrassMapAlwaysBuildRuntimeGenerationResources);
+			return SupportsRuntimeGrassMapGeneration();
 		#endif // WITH_EDITOR
 	}
 
diff --git a/Engine/Source/Runtime/Landscape/Private/LandscapeUtils.cpp b/Engine/Source/Runtime/Landscape/Private/LandscapeUtils.cpp
--- a/Engine/Source/Runtime/Landscape/Private/LandscapeUtils.cpp
+++ b/Engine/Source/Runtime/Landscape/Private/LandscapeUtils.cpp
@@ -36,6 +36,9 @@ extern const size_t ChannelOffsets[4];
 extern UNREALED_API UEditorEngine* GEditor;
 #endif // WITH_EDITOR
 
+extern int32 GGrassMapUseRuntimeGeneration;
+extern int32 GGrassMapAlwaysBuildRuntimeGenerationResources;
+
 static TAutoConsoleVariable<bool> CVarLandscapeDebugShaders(
 	TEXT("landscape.DebugShaders"),
 	false,
@@ -93,6 +96,11 @@ FString GetLandscapeToolTargetTypeFlagsAsString(ELandscapeToolTargetTypeFlags In
 	return *FString::Join(TargetTypeStrings, TEXT(","));
 }
 
+bool SupportsRuntimeGrassMapGeneration()
+{
+	return GGrassMapUseRuntimeGeneration || GGrassMapAlwaysBuildRuntimeGenerationResources;
+}
+
 
 // ----------------------------------------------------------------------------------
 
diff --git a/Engine/Source/Runtime/Landscape/Private/LandscapeUtilsPrivate.cpp b/Engine/Source/Runtime/Landscape/Private/LandscapeUtilsPrivate.cpp
--- a/Engine/Source/Runtime/Landscape/Private/LandscapeUtilsPrivate.cpp
+++ b/Engine/Source/Runtime/Landscape/Private/LandscapeUtilsPrivate.cpp
@@ -6,6 +6,7 @@
 #include "LandscapeComponent.h"
 #include "LandscapeProxy.h"
 #include "LandscapeInfo.h"
+#include "RenderingThread.h"
 
 #if WITH_EDITOR
 #include "EditorDirectories.h"
@@ -179,4 +180,42 @@ int32 BlitWeightChannelsToUint8(uint8* DestBuffer, const FBlitBuffer2DDesc& Dest
 	return BlitRGChannelsToSingleTmpl(DestBuffer, Dest, SrcBuffer, Src, ClipRect, GetWeight);
 };
 
+void SubmitGPUCommands(bool bBlockUntilRTComplete, bool bBlockRTUntilGPUComplete)
+{
+	FEvent* ResultsReadyEvent = nullptr;
+	if (bBlockUntilRTComplete)
+	{
+		ResultsReadyEvent = FPlatformProcess::GetSynchEventFromPool(true);
+	}
+
+	ENQUEUE_RENDER_COMMAND(FFlushResourcesCommand)(
+		[ResultsReadyEvent, bBlockRTUntilGPUComplete](FRHICommandListImmediate& RHICmdList)
+		{
+			RHICmdList.ImmediateFlush(EImmediateFlushType::FlushRHIThreadFlushResources);
+			RHICmdList.FlushResources();
+			RHICmdList.SubmitCommandsAndFlushGPU();
+
+			if (ResultsReadyEvent)
+			{
+				if (bBlockRTUntilGPUComplete)
+				{
+					// Block render thread waiting for GPU to complete.  Note this can be very expensive on some platforms.
+					RHICmdList.BlockUntilGPUIdle();
+				}
+
+				ResultsReadyEvent->Trigger();
+			}
+		});
+
+	if (ResultsReadyEvent)
+	{
+		// block game thread waiting for render thread to tell us the GPU is complete
+		ResultsReadyEvent->Wait();
+		ResultsReadyEvent->Reset();
+		FPlatformProcess::ReturnSynchEventToPool(ResultsReadyEvent);
+	}
+
+	return;
+}
+
 } // end namespace UE::Landscape::Private
\ No newline at end of file
diff --git a/Engine/Source/Runtime/Landscape/Private/LandscapeUtilsPrivate.h b/Engine/Source/Runtime/Landscape/Private/LandscapeUtilsPrivate.h
--- a/Engine/Source/Runtime/Landscape/Private/LandscapeUtilsPrivate.h
+++ b/Engine/Source/Runtime/Landscape/Private/LandscapeUtilsPrivate.h
@@ -347,4 +347,6 @@ namespace UE::Landscape::Private
 	 */
 	int32 BlitWeightChannelsToUint8(uint8* DestBuffer, const FBlitBuffer2DDesc& Dest, uint32* SrcBuffer, const FBlitBuffer2DDesc& Src, FIntRect ClipRect);
 
+	void SubmitGPUCommands(bool bBlockUntilRTComplete, bool bBlockRTUntilGPUComplete);
+
 } // end namespace UE::Landscape::Private
diff --git a/Engine/Source/Runtime/Landscape/Public/LandscapeUtils.h b/Engine/Source/Runtime/Landscape/Public/LandscapeUtils.h
--- a/Engine/Source/Runtime/Landscape/Public/LandscapeUtils.h
+++ b/Engine/Source/Runtime/Landscape/Public/LandscapeUtils.h
@@ -47,6 +47,8 @@ LANDSCAPE_API ELandscapeToolTargetTypeFlags GetLandscapeToolTargetTypeAsFlags(EL
 LANDSCAPE_API ELandscapeToolTargetType GetLandscapeToolTargetTypeSingleFlagAsType(ELandscapeToolTargetTypeFlags InSingleFlag);
 LANDSCAPE_API FString GetLandscapeToolTargetTypeFlagsAsString(ELandscapeToolTargetTypeFlags InTargetTypeFlags);
 
+LANDSCAPE_API bool SupportsRuntimeGrassMapGeneration();
+
 // ----------------------------------------------------------------------------------
 
 /** This struct is usually meant to be allocated on the game thread (where there's no FRDGBuilder, which is Render Thread-only) and allows to queue successive operations (lambdas) onto a single render command,
