Add support for landscape delayed proxy creation or default material rendering when PSOs are still compiling with custom optional landscape cvars to set the delayed proxy creation mode

#Synced-CL 49762744

diff --git a/Engine/Source/Runtime/Engine/Private/PSOPrecache.cpp b/Engine/Source/Runtime/Engine/Private/PSOPrecache.cpp
--- a/Engine/Source/Runtime/Engine/Private/PSOPrecache.cpp
+++ b/Engine/Source/Runtime/Engine/Private/PSOPrecache.cpp
@@ -111,8 +111,13 @@ bool ShouldBoostPSOPrecachePriorityOnDraw()
 	return FApp::CanEverRender() && PipelineStateCache::IsPSOPrecachingEnabled() && GPSODrawnComponentBoostStrategy && !GIsEditor;
 }
 
-EPSOPrecacheProxyCreationStrategy GetPSOPrecacheProxyCreationStrategy()
+EPSOPrecacheProxyCreationStrategy GetPSOPrecacheProxyCreationStrategy(int32 OverrideValue)
 {
+	if (OverrideValue >= 0 && OverrideValue <= 2)
+	{
+		return EPSOPrecacheProxyCreationStrategy(OverrideValue);
+	}
+
 	if (GPSOProxyCreationWhenPSOReady != 1)
 	{
 		return EPSOPrecacheProxyCreationStrategy::AlwaysCreate;
diff --git a/Engine/Source/Runtime/Engine/Public/PSOPrecacheComponent.h b/Engine/Source/Runtime/Engine/Public/PSOPrecacheComponent.h
--- a/Engine/Source/Runtime/Engine/Public/PSOPrecacheComponent.h
+++ b/Engine/Source/Runtime/Engine/Public/PSOPrecacheComponent.h
@@ -24,7 +24,7 @@ enum class EPSOPrecacheProxyCreationStrategy : uint8
 	UseDefaultMaterialUntilPSOPrecached = 2
 };
 
-extern ENGINE_API EPSOPrecacheProxyCreationStrategy GetPSOPrecacheProxyCreationStrategy();
+extern ENGINE_API EPSOPrecacheProxyCreationStrategy GetPSOPrecacheProxyCreationStrategy(int32 OverrideValue = -1);
 
 /**
  * Delay component proxy creation when it's requested PSOs are still precaching
diff --git a/Engine/Source/Runtime/Landscape/Classes/LandscapeComponent.h b/Engine/Source/Runtime/Landscape/Classes/LandscapeComponent.h
--- a/Engine/Source/Runtime/Landscape/Classes/LandscapeComponent.h
+++ b/Engine/Source/Runtime/Landscape/Classes/LandscapeComponent.h
@@ -865,6 +865,9 @@ public:
 	virtual void SetMaterial(int32 ElementIndex, UMaterialInterface* Material) override;
 	virtual void PreFeatureLevelChange(ERHIFeatureLevel::Type PendingFeatureLevel) override;
 #endif
+	
+	/** Collect all the PSO precache data used by the landscape component */
+	LANDSCAPE_API virtual void CollectPSOPrecacheData(const FPSOPrecacheParams& BasePrecachePSOParams, FMaterialInterfacePSOPrecacheParamsList& OutParams) override;
 	//~ End UPrimitiveComponent Interface.
 
 	//~ Begin USceneComponent Interface.
diff --git a/Engine/Source/Runtime/Landscape/Private/Landscape.cpp b/Engine/Source/Runtime/Landscape/Private/Landscape.cpp
--- a/Engine/Source/Runtime/Landscape/Private/Landscape.cpp
+++ b/Engine/Source/Runtime/Landscape/Private/Landscape.cpp
@@ -1242,51 +1242,7 @@ void ULandscapeComponent::PostLoad()
 
 	Super::PostLoad();
 
-	if (IsComponentPSOPrecachingEnabled())
-	{
-		TArray<UMaterialInterface*> Materials;
-		bool bGetDebugMaterials = false;
-		GetUsedMaterials(Materials, bGetDebugMaterials);
-
-		FPSOPrecacheParams PrecachePSOParams;
-		SetupPrecachePSOParams(PrecachePSOParams);
-
-		PrecachePSOParams.bIgnoreMaskMaterial = !UE::Landscape::UseLandscapeVertexMask();
-		PrecachePSOParams.bRequireRTDynamicGeometryUpdate = true;
-
-		FPSOPrecacheVertexFactoryDataList VertexFactoryDataList;
-		VertexFactoryDataList.Add(FPSOPrecacheVertexFactoryData(&FLandscapeVertexFactory::StaticType));
-
-		// we need the fixed grid vertex factory for both virtual texturing and grass
-		if (NeedsFixedGridVertexFactory(GMaxRHIShaderPlatform))
-		{
-			VertexFactoryDataList.Add(FPSOPrecacheVertexFactoryData(&FLandscapeFixedGridVertexFactory::StaticType));
-		}
-
-		if (Culling::UseCulling(GMaxRHIShaderPlatform))
-		{
-			VertexFactoryDataList.Add(FPSOPrecacheVertexFactoryData(Culling::GetTileVertexFactoryType()));
-		}
-
-		ILandscapeModule& LandscapeModule = FModuleManager::GetModuleChecked<ILandscapeModule>("Landscape");
-		const UE::Landscape::FSetupLandscapePSOPrecacheVertexFactoryDataListDelegate& SetupLandscapePSOPrecacheVertexFactoryDataListDelegate = LandscapeModule.GetSetupLandscapePSOPrecacheVertexFactoryDataListDelegate();
-		if (SetupLandscapePSOPrecacheVertexFactoryDataListDelegate.IsBound())
-		{
-			SetupLandscapePSOPrecacheVertexFactoryDataListDelegate.Execute(this, VertexFactoryDataList);
-		}
-
-		TArray<FMaterialPSOPrecacheRequestID> MaterialPrecacheRequestIDs;
-		for (UMaterialInterface* MaterialInterface : Materials)
-		{
-			if (MaterialInterface)
-			{
-				if (IsComponentPSOPrecachingEnabled())
-				{
-					MaterialInterface->PrecachePSOs(VertexFactoryDataList, PrecachePSOParams, EPSOPrecachePriority::High, MaterialPrecacheRequestIDs);
-				}
-			}
-		}
-	}
+	PrecachePSOs();
 
 #if WITH_EDITOR
 	ALandscapeProxy* LandscapeProxy = GetLandscapeProxy();
@@ -1602,6 +1558,49 @@ PRAGMA_ENABLE_DEPRECATION_WARNINGS
 	GrassData->ConditionalDiscardDataOnLoad();
 }
 
+void ULandscapeComponent::CollectPSOPrecacheData(const FPSOPrecacheParams& BasePrecachePSOParams, FMaterialInterfacePSOPrecacheParamsList& OutParams)
+{
+	FPSOPrecacheParams PrecachePSOParams = BasePrecachePSOParams;
+	PrecachePSOParams.bIgnoreMaskMaterial = !UE::Landscape::UseLandscapeVertexMask();
+	PrecachePSOParams.bRequireRTDynamicGeometryUpdate = true;
+
+	FPSOPrecacheVertexFactoryDataList VertexFactoryDataList;
+	VertexFactoryDataList.Add(FPSOPrecacheVertexFactoryData(&FLandscapeVertexFactory::StaticType));
+
+	// we need the fixed grid vertex factory for both virtual texturing and grass
+	if (UE::Landscape::NeedsFixedGridVertexFactory(GMaxRHIShaderPlatform))
+	{
+		VertexFactoryDataList.Add(FPSOPrecacheVertexFactoryData(&FLandscapeFixedGridVertexFactory::StaticType));
+	}
+
+	if (UE::Landscape::Culling::UseCulling(GMaxRHIShaderPlatform))
+	{
+		VertexFactoryDataList.Add(FPSOPrecacheVertexFactoryData(UE::Landscape::Culling::GetTileVertexFactoryType()));
+	}
+
+	ILandscapeModule& LandscapeModule = FModuleManager::GetModuleChecked<ILandscapeModule>("Landscape");
+	const UE::Landscape::FSetupLandscapePSOPrecacheVertexFactoryDataListDelegate& SetupLandscapePSOPrecacheVertexFactoryDataListDelegate = LandscapeModule.GetSetupLandscapePSOPrecacheVertexFactoryDataListDelegate();
+	if (SetupLandscapePSOPrecacheVertexFactoryDataListDelegate.IsBound())
+	{
+		SetupLandscapePSOPrecacheVertexFactoryDataListDelegate.Execute(this, VertexFactoryDataList);
+	}
+
+	TArray<UMaterialInterface*> Materials;
+	bool bGetDebugMaterials = false;
+	GetUsedMaterials(Materials, bGetDebugMaterials);
+	for (UMaterialInterface* MaterialInterface : Materials)
+	{
+		if (MaterialInterface)
+		{
+			FMaterialInterfacePSOPrecacheParams& ComponentParams = OutParams[OutParams.AddDefaulted()];
+			ComponentParams.Priority = EPSOPrecachePriority::High;
+			ComponentParams.MaterialInterface = MaterialInterface;
+			ComponentParams.VertexFactoryDataList = VertexFactoryDataList;
+			ComponentParams.PSOPrecacheParams = PrecachePSOParams;
+		}
+	}
+}
+
 #if WITH_EDITORONLY_DATA
 void ULandscapeComponent::DeclareConstructClasses(TArray<FTopLevelAssetPath>& OutConstructClasses, const UClass* SpecificSubclass)
 {
@@ -2098,6 +2097,12 @@ void ULandscapeComponent::BeginDestroy()
 
 FPrimitiveSceneProxy* ULandscapeComponent::CreateSceneProxy()
 {
+	if (CheckPSOPrecachingAndBoostPriority(EPSOPrecachePriority::Highest) && UE::Landscape::GetLandscapePSOPrecacheProxyCreationStrategy() == EPSOPrecacheProxyCreationStrategy::DelayUntilPSOPrecached)
+	{
+		UE_LOG(LogLandscape, Verbose, TEXT("Skipping CreateSceneProxy for ULandscapeComponent %s (Its PSOs are still compiling)"), *GetFullName());
+		return nullptr;
+	}
+	
 	FPrimitiveSceneProxy* LocalSceneProxy = nullptr;
 
 	ILandscapeModule& LandscapeModule = FModuleManager::GetModuleChecked<ILandscapeModule>("Landscape");
diff --git a/Engine/Source/Runtime/Landscape/Private/LandscapeRender.cpp b/Engine/Source/Runtime/Landscape/Private/LandscapeRender.cpp
--- a/Engine/Source/Runtime/Landscape/Private/LandscapeRender.cpp
+++ b/Engine/Source/Runtime/Landscape/Private/LandscapeRender.cpp
@@ -238,6 +238,26 @@ FAutoConsoleVariableRef CVarLandscapeRuntimeVirtualTextureRenderWithQuad(
 	ECVF_RenderThreadSafe | ECVF_Scalability
 );
 
+static int32 GLandscapePSOPrecacheProxyCreationStrategy = -1;
+static FAutoConsoleVariableRef CVarLandscapePSOPrecacheProxyCreationStrategy(
+	TEXT("r.PSOPrecache.LandscapeProxyCreationStrategy"), 
+	GLandscapePSOPrecacheProxyCreationStrategy, 
+	TEXT("-1 : Use global PSOPrecache settings (default)\n")
+	TEXT(" 0 : Always create the proxy even when PSOs are still compiling\n")
+	TEXT(" 1 : Delay proxy creation when PSOs are still compiling\n")
+	TEXT(" 2 : Use default material when PSOs are still compiling"),
+	ECVF_ReadOnly);
+
+static int32 GLandscapePSOPrecacheMeshProxyCreationStrategy = 1;
+static FAutoConsoleVariableRef CVarLandscapePSOPrecacheMeshProxyCreationStrategy(
+	TEXT("r.PSOPrecache.LandscapeMeshProxyCreationStrategy"), 
+	GLandscapePSOPrecacheMeshProxyCreationStrategy, 
+	TEXT("-1 : Use global PSOPrecache settings\n")
+	TEXT(" 0 : Always create the proxy even when PSOs are still compiling\n")
+	TEXT(" 1 : Delay proxy creation when PSOs are still compiling (default)\n")
+	TEXT(" 2 : Use default material when PSOs are still compiling"),
+	ECVF_ReadOnly);
+
 #if WITH_EDITOR
 extern TAutoConsoleVariable<int32> CVarLandscapeShowDirty;
 #endif
@@ -373,6 +393,11 @@ namespace UE::Landscape
 	{
 		return CVarRayTracingLandscapeVertexMask.GetValueOnAnyThread() != 0;
 	}
+
+	EPSOPrecacheProxyCreationStrategy GetLandscapePSOPrecacheProxyCreationStrategy()
+	{
+		return GetPSOPrecacheProxyCreationStrategy(GLandscapePSOPrecacheProxyCreationStrategy);
+	}
 } // namespace UE::Landscape
 
 //
@@ -1637,6 +1662,11 @@ FLandscapeComponentSceneProxy::FLandscapeComponentSceneProxy(ULandscapeComponent
 	check(ComponentLightInfo);
 
 	const bool bHasStaticLighting = ComponentLightInfo->GetLightMap() || ComponentLightInfo->GetShadowMap();
+	
+	bool bForceDefaultMaterial = false;
+#if UE_WITH_PSO_PRECACHING
+	bForceDefaultMaterial = InComponent->GetPSOPrecacheComponentData().IsPSOPrecaching() && UE::Landscape::GetLandscapePSOPrecacheProxyCreationStrategy() == EPSOPrecacheProxyCreationStrategy::UseDefaultMaterialUntilPSOPrecached;
+#endif // UE_WITH_PSO_PRECACHING
 
 	check(AvailableMaterialInterfaces.Num() == AvailableMaterials.Num());
 	// Check material usage and validity. Replace invalid entries by default material so that indexing AvailableMaterials with LODIndexToMaterialIndex still works :
@@ -1661,7 +1691,7 @@ FLandscapeComponentSceneProxy::FLandscapeComponentSceneProxy(ULandscapeComponent
 				bIsValidMaterial &= !bHasStaticLighting || MaterialInterface->CheckMaterialUsage_Concurrent(MATUSAGE_StaticLighting);
 			}
 
-			if (!bIsValidMaterial)
+			if (!bIsValidMaterial || bForceDefaultMaterial)
 			{
 				// Replace the landscape material by the default material : 
 				MaterialInterface = UMaterial::GetDefaultMaterial(MD_Surface);
@@ -4830,6 +4860,13 @@ FPrimitiveSceneProxy* ULandscapeMeshProxyComponent::CreateSceneProxy()
 		return nullptr;
 	}
 
+	// Check mesh overwrite version of proxy creation strategy to allow different strategy for HLODs
+	if (CheckPSOPrecachingAndBoostPriority(EPSOPrecachePriority::Highest) && GetPSOPrecacheProxyCreationStrategy(GLandscapePSOPrecacheMeshProxyCreationStrategy) == EPSOPrecacheProxyCreationStrategy::DelayUntilPSOPrecached)
+	{
+		UE_LOG(LogLandscape, Verbose, TEXT("Skipping CreateSceneProxy for ULandscapeMeshProxyComponent %s (Its PSOs are still compiling)"), *GetFullName());
+		return nullptr;
+	}
+
 	return new FLandscapeMeshProxySceneProxy(this, LandscapeGuid, ProxyComponentBases, ProxyComponentCentersObjectSpace, ComponentXVectorObjectSpace, ComponentYVectorObjectSpace, GetComponentTransform(), ComponentResolution, ProxyLOD, LODGroupKey, ALandscapeProxy::ComputeLandscapeKey(GetWorld(), LODGroupKey, LandscapeGuid));
 }
 
@@ -4919,6 +4956,12 @@ FPrimitiveSceneProxy* ULandscapeNaniteComponent::CreateSceneProxy()
 	// Is Nanite supported, and is there built Nanite data for this static mesh?
 	if (IsEnabled() && ShouldCreateNaniteProxy(&MaterialAudit))
 	{
+		if (CheckPSOPrecachingAndBoostPriority(EPSOPrecachePriority::Highest) && GetPSOPrecacheProxyCreationStrategy(GLandscapePSOPrecacheProxyCreationStrategy) == EPSOPrecacheProxyCreationStrategy::DelayUntilPSOPrecached)
+		{
+			UE_LOG(LogLandscape, Verbose, TEXT("Skipping CreateSceneProxy for ULandscapeNaniteComponent %s (Its PSOs are still compiling)"), *GetFullName());
+			return nullptr;
+		}
+
 		return ::new FLandscapeNaniteSceneProxy(MaterialAudit, this);
 	}
 
diff --git a/Engine/Source/Runtime/Landscape/Public/LandscapeRender.h b/Engine/Source/Runtime/Landscape/Public/LandscapeRender.h
--- a/Engine/Source/Runtime/Landscape/Public/LandscapeRender.h
+++ b/Engine/Source/Runtime/Landscape/Public/LandscapeRender.h
@@ -94,6 +94,7 @@ namespace UE::Landscape
 	bool NeedsFixedGridVertexFactory(EShaderPlatform InShaderPlatform);
 	bool ShouldBuildGrassMapRenderingResources();
 	bool UseLandscapeVertexMask();
+	EPSOPrecacheProxyCreationStrategy GetLandscapePSOPrecacheProxyCreationStrategy();
 } // namespace UE::Landscape
 
 /** The uniform shader parameters for a landscape draw call. */
