Landscape Editor - Removed setter calls from ULandscapeEditLayerBase::PostEditPropertyChanged
* Call Broadcast directly since the Setter functions will incorrectly find that bHasChanged == false (the property has already been modified here)

#Synced-CL 49736387

diff --git a/Engine/Source/Runtime/Landscape/Private/LandscapeEditLayer.cpp b/Engine/Source/Runtime/Landscape/Private/LandscapeEditLayer.cpp
--- a/Engine/Source/Runtime/Landscape/Private/LandscapeEditLayer.cpp
+++ b/Engine/Source/Runtime/Landscape/Private/LandscapeEditLayer.cpp
@@ -3,6 +3,7 @@
 #include "LandscapeEditLayer.h"
 #include "LandscapeEditTypes.h"
 #include "Landscape.h"
+#include "Algo/AllOf.h"
 
 #include UE_INLINE_GENERATED_CPP_BY_NAME(LandscapeEditLayer)
 
@@ -224,12 +225,15 @@ void ULandscapeEditLayerBase::AddOrUpdateWeightmapAllocationLayerBlend(TObjectPt
 
 	bFoundValue = bInValue;
 	
-	if (bInModify)
+	if (bHasValueChanged)
 	{
-		Modify();
+		if (bInModify)
+		{
+			Modify();
+		}
 	}
 
-	BroadcastOnLayerDataChanged(GET_MEMBER_NAME_CHECKED(ULandscapeEditLayerBase, WeightmapLayerAllocationBlend), /*bInUserTriggered = */true, /*bInRequiresLandscapeUpdate =*/ true, /*bHasValueChanged =*/ true, /*InChangeType = */EPropertyChangeType::ValueSet);
+	BroadcastOnLayerDataChanged(GET_MEMBER_NAME_CHECKED(ULandscapeEditLayerBase, WeightmapLayerAllocationBlend), /*bInUserTriggered = */true, /*bInRequiresLandscapeUpdate =*/ true, bHasValueChanged, /*InChangeType = */EPropertyChangeType::ValueSet);
 }
 
 const TMap<TObjectPtr<ULandscapeLayerInfoObject>, bool>& ULandscapeEditLayerBase::GetWeightmapLayerAllocationBlend() const
@@ -239,14 +243,24 @@ const TMap<TObjectPtr<ULandscapeLayerInfoObject>, bool>& ULandscapeEditLayerBase
 
 void ULandscapeEditLayerBase::SetWeightmapLayerAllocationBlend(const TMap<TObjectPtr<ULandscapeLayerInfoObject>, bool>& InWeightmapLayerAllocationBlend, bool bInModify)
 {
+	const bool bHasValueChanged = !(WeightmapLayerAllocationBlend.Num() == InWeightmapLayerAllocationBlend.Num() &&
+		Algo::AllOf(WeightmapLayerAllocationBlend, [&InWeightmapLayerAllocationBlend](const TPair<TObjectPtr<ULandscapeLayerInfoObject>, bool>& ExistingPair)
+			{
+				const bool* BlendValue = InWeightmapLayerAllocationBlend.Find(ExistingPair.Key);
+				return BlendValue != nullptr && *BlendValue == ExistingPair.Value;
+			}));
+
 	WeightmapLayerAllocationBlend = InWeightmapLayerAllocationBlend;
 
-	if (bInModify)
+	if (bHasValueChanged)
 	{
-		Modify();
+		if (bInModify)
+		{
+			Modify();
+		}
 	}
 
-	BroadcastOnLayerDataChanged(GET_MEMBER_NAME_CHECKED(ULandscapeEditLayerBase, WeightmapLayerAllocationBlend), /*bInUserTriggered = */true, /*bInRequiresLandscapeUpdate =*/ true, /*bHasValueChanged =*/ true, /*InChangeType = */EPropertyChangeType::ValueSet);
+	BroadcastOnLayerDataChanged(GET_MEMBER_NAME_CHECKED(ULandscapeEditLayerBase, WeightmapLayerAllocationBlend), /*bInUserTriggered = */true, /*bInRequiresLandscapeUpdate =*/ true, bHasValueChanged, /*InChangeType = */EPropertyChangeType::ValueSet);
 }
 
 void ULandscapeEditLayerBase::PostEditChangeProperty(FPropertyChangedEvent& PropertyChangedEvent)
@@ -255,34 +269,8 @@ void ULandscapeEditLayerBase::PostEditChangeProperty(FPropertyChangedEvent& Prop
 
 	const FName MemberPropertyName = PropertyChangedEvent.MemberProperty ? PropertyChangedEvent.MemberProperty->GetFName() : NAME_None;
 
-	if (MemberPropertyName == GET_MEMBER_NAME_CHECKED(ULandscapeEditLayerBase, Guid))
-	{
-		SetGuid(Guid, /*bInModify = */true);
-	}
-	else if (MemberPropertyName == GET_MEMBER_NAME_CHECKED(ULandscapeEditLayerBase, LayerName))
-	{
-		SetName(LayerName, /*bInModify = */true);
-	}
-	else if (MemberPropertyName == GET_MEMBER_NAME_CHECKED(ULandscapeEditLayerBase, HeightmapAlpha))
-	{
-		SetAlphaForTargetType(ELandscapeToolTargetType::Heightmap, HeightmapAlpha, /*bInModify = */true, PropertyChangedEvent.ChangeType);
-	}
-	else if (MemberPropertyName == GET_MEMBER_NAME_CHECKED(ULandscapeEditLayerBase, WeightmapAlpha))
-	{
-		SetAlphaForTargetType(ELandscapeToolTargetType::Weightmap, WeightmapAlpha, /*bInModify = */true, PropertyChangedEvent.ChangeType);
-	}
-	else if (MemberPropertyName == GET_MEMBER_NAME_CHECKED(ULandscapeEditLayerBase, bLocked))
-	{
-		SetLocked(bLocked, /*bInModify = */true);
-	}
-	else if (MemberPropertyName == GET_MEMBER_NAME_CHECKED(ULandscapeEditLayerBase, bVisible))
-	{
-		SetVisible(bVisible, /*bInModify = */true);
-	}
-	else if (MemberPropertyName == GET_MEMBER_NAME_CHECKED(ULandscapeEditLayerBase, WeightmapLayerAllocationBlend))
-	{
-		SetWeightmapLayerAllocationBlend(WeightmapLayerAllocationBlend, /*bInModify = */true);
-	}
+	// Currently, every property requires a landscape update but this could be specialized per property to avoid unnecessary updates when new properties are added
+	BroadcastOnLayerDataChanged(MemberPropertyName, /*bInUserTriggered = */true, /* bRequiresLandscapeUpdate =*/ true, /*bHasValueChanged =*/true, PropertyChangedEvent.ChangeType);
 }
 
 void ULandscapeEditLayerBase::PostEditUndo()
@@ -335,7 +323,7 @@ void ULandscapeEditLayerBase::BroadcastOnLayerDataChanged(FName InPropertyName,
 	FProperty* Property = nullptr;
 	if (InPropertyName != NAME_None)
 	{
-		Property = FindFProperty<FProperty>(StaticClass(), InPropertyName);
+		Property = FindFProperty<FProperty>(GetClass(), InPropertyName);
 		check(Property != nullptr);
 	}
 
