Prevent int overflow when creating very large landscapes.

#Synced-CL 49559354

diff --git a/Engine/Source/Runtime/Landscape/Private/LandscapeEditLayerMergeRenderContext.cpp b/Engine/Source/Runtime/Landscape/Private/LandscapeEditLayerMergeRenderContext.cpp
--- a/Engine/Source/Runtime/Landscape/Private/LandscapeEditLayerMergeRenderContext.cpp
+++ b/Engine/Source/Runtime/Landscape/Private/LandscapeEditLayerMergeRenderContext.cpp
@@ -240,15 +240,15 @@ bool FMergeRenderBatch::operator<(const FMergeRenderBatch& InOther) const
 {
 	// Sort batches in a deterministic order (upper-left -> upper-right -> lower-left- -> lower-right) to ensure that components with a conflict on their
 	//  borders get rendered in a deterministic way : 
-	auto ComputeSortIndex = [this](const FIntPoint& InSectionBase) -> int32
+	auto ComputeSortIndex = [this](const FIntPoint& InSectionBase) -> int64
 	{
-		int32 AbsoluteLinearIndex = (InSectionBase.X - LandscapeExtent.Min.X) + (InSectionBase.Y - LandscapeExtent.Min.Y) * (LandscapeExtent.Max.X - LandscapeExtent.Min.X + 1);
+		int64 AbsoluteLinearIndex = ((int64)InSectionBase.X - LandscapeExtent.Min.X) + ((int64)InSectionBase.Y - LandscapeExtent.Min.Y) * ((int64)LandscapeExtent.Max.X - LandscapeExtent.Min.X + 1);
 		check(AbsoluteLinearIndex >= 0);
 		return AbsoluteLinearIndex;
 	};
 	
-	int32 LinearIndex = ComputeSortIndex(SectionRect.Min);
-	int32 OtherLinearIndex = ComputeSortIndex(InOther.SectionRect.Min);
+	int64 LinearIndex = ComputeSortIndex(SectionRect.Min);
+	int64 OtherLinearIndex = ComputeSortIndex(InOther.SectionRect.Min);
 	if (LinearIndex < OtherLinearIndex)
 	{
 		return true;
