Speculative crash fix for when the Slate callback for allowing the deletion of actors gets called on landscape actors after they have been deleted.

#Synced-CL 49989314

diff --git a/Engine/Source/Runtime/Landscape/Private/LandscapeEdit.cpp b/Engine/Source/Runtime/Landscape/Private/LandscapeEdit.cpp
--- a/Engine/Source/Runtime/Landscape/Private/LandscapeEdit.cpp
+++ b/Engine/Source/Runtime/Landscape/Private/LandscapeEdit.cpp
@@ -6482,8 +6482,7 @@ bool ALandscape::CanDeleteSelectedActor(FText& OutReason) const
 	{
 		// Allow Delete of Actor if all other related actors have been deleted
 		ULandscapeInfo* Info = GetLandscapeInfo();
-		check(Info);
-		return Info->CanDeleteLandscape(OutReason);
+		return Info ? Info->CanDeleteLandscape(OutReason) : false;
 	}
 
 	return true;
