Improved/standardized formatting of debug landscape files (for grass.grassmap.dump, landscape.DumpHeightmapDiff and landscape.DumpWeightmapDiff), it now shows the component's key (index in the overall landscape) and its section base (which is visible in landscape mode, when in Select mode)

#Synced-CL 49989314

diff --git a/Engine/Source/Runtime/Landscape/Classes/LandscapeComponent.h b/Engine/Source/Runtime/Landscape/Classes/LandscapeComponent.h
--- a/Engine/Source/Runtime/Landscape/Classes/LandscapeComponent.h
+++ b/Engine/Source/Runtime/Landscape/Classes/LandscapeComponent.h
@@ -1029,6 +1029,9 @@ public:
 		return GetSectionBase() / ComponentSizeQuads;
 	}
 
+	/** @return a user-friendly name to easily identify this component when dumping debug info, for example */
+	LANDSCAPE_API FString GetDebugName() const;
+
 	/** @param InSectionBase new section base for a component */
 	void SetSectionBase(FIntPoint InSectionBase)
 	{
diff --git a/Engine/Source/Runtime/Landscape/Private/Landscape.cpp b/Engine/Source/Runtime/Landscape/Private/Landscape.cpp
--- a/Engine/Source/Runtime/Landscape/Private/Landscape.cpp
+++ b/Engine/Source/Runtime/Landscape/Private/Landscape.cpp
@@ -1969,6 +1969,13 @@ ALandscapeProxy* ULandscapeComponent::GetLandscapeProxy() const
 	return CastChecked<ALandscapeProxy>(GetOuter());
 }
 
+FString ULandscapeComponent::GetDebugName() const
+{
+	const FIntPoint ComponentKey = GetComponentKey();
+	const FIntPoint SectionBase = GetSectionBase();
+	return FString::Format(TEXT("Component_{0}_{1}[{2}_{3}]"), { ComponentKey.X, ComponentKey.Y, SectionBase.X, SectionBase.Y });
+}
+
 int32 ULandscapeComponent::GetNumRelevantMips() const
 {
 	const int32 TextureSize = (SubsectionSizeQuads + 1) * NumSubsections;
diff --git a/Engine/Source/Runtime/Landscape/Private/LandscapeEditLayers.cpp b/Engine/Source/Runtime/Landscape/Private/LandscapeEditLayers.cpp
--- a/Engine/Source/Runtime/Landscape/Private/LandscapeEditLayers.cpp
+++ b/Engine/Source/Runtime/Landscape/Private/LandscapeEditLayers.cpp
@@ -5226,7 +5226,7 @@ void ALandscape::OnDirtyWeightmap(FTextureToComponentHelper const& MapHelper, UT
 						check(Proxy);
 						FString ActorName = Proxy->GetActorLabel();
 						FString FilePattern = FString::Format(TEXT("{0}/LandscapeLayers/{1}/{2}/{3}/Weightmaps/{4}/{5}-{6}-{7}[mip{8}]"), 
-							{ FPaths::ProjectSavedDir(), CurrentTime.ToString(), WorldName, ParentLandscapeActorName, AllocInfo.GetLayerName().ToString(), ActorName, Component->GetName(), InWeightmap->GetName(), InMipLevel});
+							{ FPaths::ProjectSavedDir(), CurrentTime.ToString(), WorldName, ParentLandscapeActorName, AllocInfo.GetLayerName().ToString(), ActorName, Component->GetDebugName(), InWeightmap->GetName(), InMipLevel});
 
 						FFileHelper::EColorChannel ColorChannel = GetWeightmapColorChannel(AllocInfo);
 						FFileHelper::CreateBitmap(*(FilePattern + "_a(pre).bmp"), SizeU, SizeV, InOldData, /*SubRectangle = */nullptr, &IFileManager::Get(), /*OutFilename = */nullptr, /*bool bInWriteAlpha = */true, ColorChannel);
@@ -5361,7 +5361,8 @@ void ALandscape::OnDirtyHeightmap(FTextureToComponentHelper const& MapHelper, UT
 				ALandscapeProxy* Proxy = Cast<ALandscapeProxy>(Component->GetOwner());
 				check(Proxy);
 				FString ActorName = Proxy->GetActorLabel();
-				FString FilePattern = FString::Format(TEXT("{0}/LandscapeLayers/{1}/{2}/{3}/Heightmaps/{4}-{5}-{6}[mip{7}]"), { FPaths::ProjectSavedDir(), CurrentTime.ToString(), WorldName, ParentLandscapeActorName, ActorName, Component->GetName(), InHeightmap->GetName(), InMipLevel });
+				FString FilePattern = FString::Format(TEXT("{0}/LandscapeLayers/{1}/{2}/{3}/Heightmaps/{4}-{5}-{6}[mip{7}]"), 
+					{ FPaths::ProjectSavedDir(), CurrentTime.ToString(), WorldName, ParentLandscapeActorName, ActorName, Component->GetDebugName(), InHeightmap->GetName(), InMipLevel });
 
 				const int32 SizeU = InHeightmap->Source.GetSizeX() >> InMipLevel;
 				const int32 SizeV = InHeightmap->Source.GetSizeY() >> InMipLevel;
diff --git a/Engine/Source/Runtime/Landscape/Private/LandscapeGrassWeightExporter.cpp b/Engine/Source/Runtime/Landscape/Private/LandscapeGrassWeightExporter.cpp
--- a/Engine/Source/Runtime/Landscape/Private/LandscapeGrassWeightExporter.cpp
+++ b/Engine/Source/Runtime/Landscape/Private/LandscapeGrassWeightExporter.cpp
@@ -977,8 +977,8 @@ TMap<ULandscapeComponent*, TUniquePtr<FLandscapeComponentGrassData>, TInlineSetA
 			ALandscapeProxy* Proxy = Cast<ALandscapeProxy>(Component->GetOwner());
 			check(Proxy != nullptr);
 			FString ActorName = Proxy->GetActorLabel();
-			FString FilePath = FString::Format(TEXT("{0}/LandscapeGrass/{1}/{2}/{3}/Grassmaps/{4}-{5}"), 
-				{ FPaths::ProjectSavedDir(), CurrentTime.ToString(), WorldName, ParentLandscapeActorName, ActorName, Component->GetName() });
+			FString FilePath = FString::Format(TEXT("{0}/LandscapeGrass/{1}/{2}/{3}/{4}-{5}"), 
+				{ FPaths::ProjectSavedDir(), CurrentTime.ToString(), WorldName, ParentLandscapeActorName, ActorName, Component->GetDebugName() });
 			FStringBuilderBase StrBuilder;
 			for (int32 PassIndex = 0; PassIndex < ComponentInfo.GetNumPasses(); ++PassIndex)
 			{
