[Navigation] Prevent multithread access on CachedHeightFieldSamples possible through INavRelevantInterface::GatherGeometrySlice.

#Synced-CL 49739737

diff --git a/Engine/Source/Runtime/Landscape/Classes/LandscapeHeightfieldCollisionComponent.h b/Engine/Source/Runtime/Landscape/Classes/LandscapeHeightfieldCollisionComponent.h
--- a/Engine/Source/Runtime/Landscape/Classes/LandscapeHeightfieldCollisionComponent.h
+++ b/Engine/Source/Runtime/Landscape/Classes/LandscapeHeightfieldCollisionComponent.h
@@ -10,6 +10,7 @@
 #include "AI/Navigation/NavigationTypes.h"
 #include "Components/PrimitiveComponent.h"
 #include "Serialization/BulkData.h"
+#include "HAL/CriticalSection.h"
 
 #include "LandscapeHeightfieldCollisionComponent.generated.h"
 
@@ -175,6 +176,8 @@ public:
 	/** Cached PxHeightFieldSamples values for navmesh generation. Note that it's being used only if navigation octree is set up for lazy geometry exporting */
 	int32 HeightfieldRowsCount;
 	int32 HeightfieldColumnsCount;
+
+	mutable FRWLock CachedHeightFieldSamplesLock;
 	FNavHeightfieldSamples CachedHeightFieldSamples;
 
 	enum ECollisionQuadFlags : uint8
diff --git a/Engine/Source/Runtime/Landscape/Private/LandscapeCollision.cpp b/Engine/Source/Runtime/Landscape/Private/LandscapeCollision.cpp
--- a/Engine/Source/Runtime/Landscape/Private/LandscapeCollision.cpp
+++ b/Engine/Source/Runtime/Landscape/Private/LandscapeCollision.cpp
@@ -286,7 +286,11 @@ void ULandscapeHeightfieldCollisionComponent::OnUnregister()
 	DeferredDestroyCollision(HeightfieldRef);
 	HeightfieldRef = nullptr;
 	HeightfieldGuid = FGuid();
-	CachedHeightFieldSamples.Empty();
+
+	{
+		FWriteScopeLock Lock(CachedHeightFieldSamplesLock);
+		CachedHeightFieldSamples.Empty();
+	}
 
 	if (GetLandscapeProxy())
 	{
@@ -1598,7 +1602,11 @@ void ULandscapeHeightfieldCollisionComponent::BeginDestroy()
 	{
 		HeightfieldRef = nullptr;
 		HeightfieldGuid = FGuid();
-		CachedHeightFieldSamples.Empty();
+
+		{
+			FWriteScopeLock Lock(CachedHeightFieldSamplesLock);
+			CachedHeightFieldSamples.Empty();
+		}
 	}
 }
 
@@ -1621,7 +1629,12 @@ bool ULandscapeHeightfieldCollisionComponent::RecreateCollision()
 
 		HeightfieldRef = nullptr; // Ensure data will be recreated
 		HeightfieldGuid = FGuid();
-		CachedHeightFieldSamples.Empty();
+
+		{
+			FWriteScopeLock Lock(CachedHeightFieldSamplesLock);
+			CachedHeightFieldSamples.Empty();
+		}
+
 		RecreatePhysicsState();
 
 		// Make sure our collision isn't destroyed while we still have a physics particle active
@@ -1896,6 +1909,8 @@ bool ULandscapeHeightfieldCollisionComponent::DoCustomNavigableGeometryExport(FN
 
 void ULandscapeHeightfieldCollisionComponent::GatherGeometrySlice(FNavigableGeometryExport& GeomExport, const FBox& SliceBox) const
 {
+	FReadScopeLock Lock(CachedHeightFieldSamplesLock);
+	
 	// note that this function can get called off game thread
 	if (CachedHeightFieldSamples.IsEmpty() == false)
 	{
@@ -1914,6 +1929,8 @@ ENavDataGatheringMode ULandscapeHeightfieldCollisionComponent::GetGeometryGather
 
 void ULandscapeHeightfieldCollisionComponent::PrepareGeometryExportSync()
 {
+	FWriteScopeLock Lock(CachedHeightFieldSamplesLock);
+	
 	if(IsValidRef(HeightfieldRef) && HeightfieldRef->HeightfieldGeometry.GetReference() && CachedHeightFieldSamples.IsEmpty())
 	{
 		const UWorld* World = GetWorld();
@@ -2592,7 +2609,10 @@ void ULandscapeHeightfieldCollisionComponent::GetResourceSizeEx(FResourceSizeEx&
 		HeightfieldRef->GetResourceSizeEx(CumulativeResourceSize);
 	}
 
-	CachedHeightFieldSamples.GetResourceSizeEx(CumulativeResourceSize);
+	{
+		FReadScopeLock Lock(CachedHeightFieldSamplesLock);
+		CachedHeightFieldSamples.GetResourceSizeEx(CumulativeResourceSize);
+	}
 }
 
 TOptional<float> ALandscapeProxy::GetHeightAtLocation(FVector Location, EHeightfieldSource HeightFieldSource) const
