Change landscape group log to verbose since this can happen during Undo of a proxy deletion

#Synced-CL 49989314

diff --git a/Engine/Source/Runtime/Landscape/Private/LandscapeGroup.cpp b/Engine/Source/Runtime/Landscape/Private/LandscapeGroup.cpp
--- a/Engine/Source/Runtime/Landscape/Private/LandscapeGroup.cpp
+++ b/Engine/Source/Runtime/Landscape/Private/LandscapeGroup.cpp
@@ -111,7 +111,7 @@ void FLandscapeGroup::RegisterComponent(ULandscapeComponent* Component)
 	// Early-out if the proxy hasn't got a proper root component. To our knowledge, this only ever happens when cooking where landscape groups are useless : 
 	if (USceneComponent* RootComponent = LandscapeProxy->GetRootComponent(); (RootComponent == nullptr) || !RootComponent->IsRegistered())
 	{
-		UE_LOG(LogLandscape, Log, TEXT("Skipped registration of landscape component to its proxy %s : the proxy's root component isn't properly set or registered"), *LandscapeProxy->GetFullName())
+		UE_LOG(LogLandscape, Verbose, TEXT("Skipped registration of landscape component to its proxy %s : the proxy's root component isn't properly set or registered"), *LandscapeProxy->GetFullName())
 		return;
 	}
 
