Prevent crash when ULandscapeComponent::GetCurrentRuntimeMaterialInstanceCount, etc. gets called on a component without a World, which is currently possible if a level gets streamed out and the material quality settings get changed shorly after

#Synced-CL 50200720

diff --git a/Engine/Source/Runtime/Landscape/Private/Landscape.cpp b/Engine/Source/Runtime/Landscape/Private/Landscape.cpp
--- a/Engine/Source/Runtime/Landscape/Private/Landscape.cpp
+++ b/Engine/Source/Runtime/Landscape/Private/Landscape.cpp
@@ -374,8 +374,9 @@ UMaterialInstance* ULandscapeComponent::GetMaterialInstance(int32 InIndex, bool
 int32 ULandscapeComponent::GetCurrentRuntimeMaterialInstanceCount() const
 {
 	ALandscapeProxy* Proxy = GetLandscapeProxy();
-	const ERHIFeatureLevel::Type FeatureLevel = Proxy->GetWorld()->GetFeatureLevel();
-	if (FeatureLevel == ERHIFeatureLevel::ES3_1)
+	UWorld* World = GetWorld();
+	const bool bIsMobile = (World != nullptr) && (World->GetFeatureLevel() == ERHIFeatureLevel::ES3_1);
+	if (bIsMobile)
 	{
 		return MobileMaterialInterfaces.Num();
 	}
@@ -384,13 +385,14 @@ int32 ULandscapeComponent::GetCurrentRuntimeMaterialInstanceCount() const
 	return GetMaterialInstanceCount(bDynamic);
 }
 
-class UMaterialInterface* ULandscapeComponent::GetCurrentRuntimeMaterialInterface(int32 InIndex) const
+UMaterialInterface* ULandscapeComponent::GetCurrentRuntimeMaterialInterface(int32 InIndex) const
 {
 	ALandscapeProxy* Proxy = GetLandscapeProxy();
-	const ERHIFeatureLevel::Type FeatureLevel = GetLandscapeProxy()->GetWorld()->GetFeatureLevel();
-
-	if (FeatureLevel == ERHIFeatureLevel::ES3_1)
+	UWorld* World = GetWorld();
+	bool bIsMobile = (World != nullptr) && (World->GetFeatureLevel() == ERHIFeatureLevel::ES3_1);
+	if (bIsMobile)
 	{
+		check(MobileMaterialInterfaces.IsValidIndex(InIndex));
 		return MobileMaterialInterfaces[InIndex];
 	}
 
