Addressing bughawk feedback on Landscape EdgeFixup code
Delete should use the same type and array style as was used to allocate the memory

#Synced-CL 50300730

diff --git a/Engine/Source/Runtime/Landscape/Private/LandscapeEdgeFixup.cpp b/Engine/Source/Runtime/Landscape/Private/LandscapeEdgeFixup.cpp
--- a/Engine/Source/Runtime/Landscape/Private/LandscapeEdgeFixup.cpp
+++ b/Engine/Source/Runtime/Landscape/Private/LandscapeEdgeFixup.cpp
@@ -1170,7 +1170,7 @@ void ULandscapeHeightmapTextureEdgeFixup::PatchTextureEdge_Internal(UE::Landscap
 					}
 
 					{
-						const uint32 SourcePitch = 4 * SrcWidth;
+						const uint32 SourcePitch = sizeof(UE::Landscape::FHeightmapTexel) * SrcWidth;
 						FUpdateTextureRegion2D UpdateRegion(DstStartX, DstStartY, SrcStartX, SrcStartY, SrcWidth, SrcHeight);
 
 						FIXUP_DEBUG_LOG_RENDER(TEXT("    PatchTextureEdge_Internal [UpdateTexture2D] res:%p e:%d mip:%d reg:[%d %d %d %d %d %d]"), HeightTextureResource, EdgeIndex, MipIndex,
@@ -1180,7 +1180,10 @@ void ULandscapeHeightmapTextureEdgeFixup::PatchTextureEdge_Internal(UE::Landscap
 					}
 				}
 
-				delete Buffer;
+				// free the data buffer (in our specific case, DestView has ownership of the memory)
+				Buffer = nullptr;
+				UE::Landscape::FHeightmapTexel* DestViewData = &DestView[0];
+				delete[] DestViewData;
 			}
 		);
 	}
