Add Enable Grass toggle to give user a mechanism to disable LS grass per landscape.

#rnx

#Synced-CL 50393138

diff --git a/Engine/Source/Runtime/Landscape/Classes/Landscape.h b/Engine/Source/Runtime/Landscape/Classes/Landscape.h
--- a/Engine/Source/Runtime/Landscape/Classes/Landscape.h
+++ b/Engine/Source/Runtime/Landscape/Classes/Landscape.h
@@ -337,6 +337,12 @@ public:
 	LANDSCAPE_API void SetLODGroupKey(uint32 InLODGroupKey);
 	LANDSCAPE_API uint32 GetLODGroupKey();
 
+	UFUNCTION(BlueprintCallable, Category = "Landscape")
+	bool GetGrassEnabled() const { return bGrassEnabled; }
+
+	UFUNCTION(BlueprintCallable, Category = "Landscape")
+	LANDSCAPE_API void SetGrassEnabled(bool bEnabled, bool bFlushGrassMaps = true);
+
 	/**
 	* Render the final heightmap in the requested top-down window as one -atlased- texture in the provided render target 2D
 	*  Can be called at runtime.
@@ -637,6 +643,9 @@ private:
 	void MarkAllLandscapeRenderStateDirty();
 	bool RenderMergedTextureInternal(const FTransform& InRenderAreaWorldTransform, const FBox2D& InRenderAreaExtents, const TArray<FName>& InWeightmapLayerNames, UTextureRenderTarget* OutRenderTarget);
 
+	UPROPERTY(EditAnywhere, Category="Grass")
+	bool bGrassEnabled = true;
+
 public:
 
 #if WITH_EDITORONLY_DATA
diff --git a/Engine/Source/Runtime/Landscape/Classes/LandscapeProxy.h b/Engine/Source/Runtime/Landscape/Classes/LandscapeProxy.h
--- a/Engine/Source/Runtime/Landscape/Classes/LandscapeProxy.h
+++ b/Engine/Source/Runtime/Landscape/Classes/LandscapeProxy.h
@@ -693,6 +693,9 @@ PRAGMA_ENABLE_DEPRECATION_WARNINGS
 	inline bool GetDisableRuntimeGrassMapGeneration() { return bDisableRuntimeGrassMapGeneration; }
 	inline void SetDisableRuntimeGrassMapGenerationProxyOnly(bool bInDisableRuntimeGrassMapGeneration) { bDisableRuntimeGrassMapGeneration = bInDisableRuntimeGrassMapGeneration; }
 
+	/** Whether grass primitives should be generated for this landscape. */
+	LANDSCAPE_API bool ShouldGenerateGrass() const;
+
 	/**
 	 *	The resolution to cache lighting at, in texels/quad in one axis
 	 *  Total resolution would be changed by StaticLightingResolution*StaticLightingResolution
diff --git a/Engine/Source/Runtime/Landscape/Private/Landscape.cpp b/Engine/Source/Runtime/Landscape/Private/Landscape.cpp
--- a/Engine/Source/Runtime/Landscape/Private/Landscape.cpp
+++ b/Engine/Source/Runtime/Landscape/Private/Landscape.cpp
@@ -1791,6 +1791,19 @@ uint32 ALandscape::GetLODGroupKey()
 	return LODGroupKey;
 }
 
+void ALandscape::SetGrassEnabled(bool bEnabled, bool bFlushGrassMaps)
+{
+	if (bEnabled != bGrassEnabled)
+	{
+		bGrassEnabled = bEnabled;
+
+		if (!bGrassEnabled)
+		{
+			FlushGrassComponents(/* OnlyForComponents = */nullptr, /* bFlushGrassMaps = */bFlushGrassMaps);
+		}
+	}
+}
+
 void ALandscape::MarkAllLandscapeRenderStateDirty()
 {
 	if (ULandscapeInfo* Info = GetLandscapeInfo())
diff --git a/Engine/Source/Runtime/Landscape/Private/LandscapeEdit.cpp b/Engine/Source/Runtime/Landscape/Private/LandscapeEdit.cpp
--- a/Engine/Source/Runtime/Landscape/Private/LandscapeEdit.cpp
+++ b/Engine/Source/Runtime/Landscape/Private/LandscapeEdit.cpp
@@ -6734,7 +6734,17 @@ void ALandscape::PostEditChangeProperty(FPropertyChangedEvent& PropertyChangedEv
 	{
 		bMarkAllLandscapeRenderStateDirty = true;
 	}
-	
+	else if (MemberPropertyName == GET_MEMBER_NAME_CHECKED(ALandscape, bGrassEnabled))
+	{
+		if (!bGrassEnabled)
+		{
+			// This should not technically be required but is done to expedite the flush and mirror the default logic in ALandscape::SetGrassEnabled().
+			// This may leave the grass builder in a non-sync'd state as it thinks it is populated but the grass maps are gone, but the builder does detect that
+			// and will fix itself.
+			FlushGrassComponents(/* OnlyForComponents = */nullptr, /* bFlushGrassMaps = */true);
+		}
+	}
+
 	// If the property that has changed is overridable or inherited and not interactive (dragging a slider)
 	// synchronize the change on all landscape proxies :
 	if (IsSharedProperty(MemberPropertyName) && PropertyChangedEvent.ChangeType != EPropertyChangeType::Interactive)
diff --git a/Engine/Source/Runtime/Landscape/Private/LandscapeGrass.cpp b/Engine/Source/Runtime/Landscape/Private/LandscapeGrass.cpp
--- a/Engine/Source/Runtime/Landscape/Private/LandscapeGrass.cpp
+++ b/Engine/Source/Runtime/Landscape/Private/LandscapeGrass.cpp
@@ -1094,6 +1094,22 @@ void ALandscapeProxy::UpdateGrassTypeSummary()
 	GrassTypeSummary.LandscapeComponentCount = LandscapeComponents.Num();
 }
 
+bool ALandscapeProxy::ShouldGenerateGrass() const
+{
+	if (!GGrassEnable)
+	{
+		return false;
+	}
+
+	if (const ALandscape* LandscapeActor = GetLandscapeActor())
+	{
+		return LandscapeActor->GetGrassEnabled();
+	}
+
+	// No landscape actor available, don't generate grass
+	return false;
+}
+
 // the purpose of this class is to copy the lightmap from the terrain, and set the CoordinateScale and CoordinateBias to zero.
 // we re-use the same texture references, so the memory cost is relatively minimal.
 class FLandscapeGrassLightMap : public FLightMap2D
@@ -2548,7 +2564,7 @@ void ALandscapeProxy::BuildGrassMaps()
 int32 ALandscapeProxy::GetOutdatedGrassMapCount() const
 {
 	int32 GrassMapsNeedingRender = 0;
-	if (GGrassEnable > 0)
+	if (ShouldGenerateGrass())
 	{
 		GrassMapsNeedingRender = GetWorld()->GetSubsystem<ULandscapeSubsystem>()->GetGrassMapBuilder()->CountOutdatedGrassMaps(LandscapeComponents);
 	}
@@ -2581,7 +2597,7 @@ void ALandscapeProxy::UpdateGrass(const TArray<FVector>& Cameras, int32& InOutNu
 	const bool bIsGrassCreationPrioritized = LandscapeSubsystem ? LandscapeSubsystem->IsGrassCreationPrioritized() : false;
 	const bool bWaitAsyncTasks = bForceSync || bIsGrassCreationPrioritized;
 
-	if (GGrassEnable > 0)
+	if (ShouldGenerateGrass())
 	{
 		float GuardBand = GGuardBandMultiplier;
 		float DiscardGuardBand = GGuardBandDiscardMultiplier;
diff --git a/Engine/Source/Runtime/Landscape/Private/LandscapeGrassMapsBuilder.cpp b/Engine/Source/Runtime/Landscape/Private/LandscapeGrassMapsBuilder.cpp
--- a/Engine/Source/Runtime/Landscape/Private/LandscapeGrassMapsBuilder.cpp
+++ b/Engine/Source/Runtime/Landscape/Private/LandscapeGrassMapsBuilder.cpp
@@ -402,24 +402,31 @@ bool FLandscapeGrassMapsBuilder::UpdateTrackedComponents(const TArray<FVector>&
 			case EComponentStage::Pending:
 				continue; // next!
 			case EComponentStage::NotReady:
-				// in game, any not ready component will never become ready.
-				// in editor, check to see if the conditions changed.
-#if WITH_EDITOR
-				if (!bIsGameWorld && AmortizedUpdate.ShouldUpdate(ComponentStateIndex))
+				// check to see if the conditions changed.
+				if (AmortizedUpdate.ShouldUpdate(ComponentStateIndex))
 				{
-					if (Component->CanRenderGrassMap())
+					if (Component->GetLandscapeProxy()->ShouldGenerateGrass() && Component->CanRenderGrassMap())
 					{
 						break; // cancel and evict to restart build process
 					}
 				}
-#endif // WITH_EDITOR
 				continue; // next!
 			case EComponentStage::TextureStreaming:
+				if (!Component->GetLandscapeProxy()->ShouldGenerateGrass())
+				{
+					break; // cancel and evict to restart build process
+				}
+
 				// don't process streaming states yet -- first process rendering states to free up slots
 				StreamingStatesToProcess.Add(State);
 				continue; // next!
 
 			case EComponentStage::Rendering:
+				if (!Component->GetLandscapeProxy()->ShouldGenerateGrass())
+				{
+					break; // cancel and evict to restart build process
+				}
+
 				check(State->ActiveRender != nullptr);
 				{
 					// NOTE: on RHI platforms that don't support fences (D3D11), the IsReady() async check will never return true within a single frame (i.e. BuildGrassMapsNow)
@@ -444,6 +451,11 @@ bool FLandscapeGrassMapsBuilder::UpdateTrackedComponents(const TArray<FVector>&
 				continue; // next!
 
 			case EComponentStage::AsyncFetch:
+				if (!Component->GetLandscapeProxy()->ShouldGenerateGrass())
+				{
+					break; // cancel and evict to restart build process
+				}
+
 				{
 					FAsyncTask<FAsyncFetchTask>* Task = State->AsyncFetchTask.Get();
 					check(Task);
@@ -458,6 +470,11 @@ bool FLandscapeGrassMapsBuilder::UpdateTrackedComponents(const TArray<FVector>&
 				// only check for invalidation of populated grass maps once in a while
 				if (AmortizedUpdate.ShouldUpdate(ComponentStateIndex))
 				{
+					if (!Component->GetLandscapeProxy()->ShouldGenerateGrass()) // TODO put above the ShouldUpdate()?
+					{
+						break; // cancel and evict to restart build process
+					}
+
 					// detect if grass data has been cleared by someone manually calling Flush (i.e. when landscape edits are made)
 					if (!Component->GrassData->HasValidData())
 					{
@@ -1369,6 +1386,13 @@ bool FLandscapeGrassMapsBuilder::TryFastpathsFromPending(FComponentState& State,
 	check(State.Stage == EComponentStage::Pending);
 	ULandscapeComponent* Component = State.Component;
 
+	if (!Component->GetLandscapeProxy()->ShouldGenerateGrass())
+	{
+		GRASS_DEBUG_LOG(TEXT("Grass generation disabled for %s"), *Component->GetName());
+		PendingToNotReady(State);
+		return true;
+	}
+
 	if (World->IsGameWorld()) // including PIE
 	{
 		// if runtime grass generation is disabled, go straight to not ready
