Fix performance regression in landscape sculpt and paint.  Don't propogate the Update_All flag from the component up to the actor unless explicitly requested.

#Synced-CL 50458757

diff --git a/Engine/Source/Runtime/Landscape/Private/LandscapeEditLayers.cpp b/Engine/Source/Runtime/Landscape/Private/LandscapeEditLayers.cpp
--- a/Engine/Source/Runtime/Landscape/Private/LandscapeEditLayers.cpp
+++ b/Engine/Source/Runtime/Landscape/Private/LandscapeEditLayers.cpp
@@ -6930,15 +6930,25 @@ void ULandscapeComponent::RequestHeightmapUpdate(bool bUpdateAll, bool bUpdateCo
 
 	LayerUpdateFlagPerMode |= NewLayerUpdateFlags;
 
+	// The typical codepath for sculpt or paint will have bUpdateCollision=true and bUpdateAll=false.  It's unclear why this also causes
+	// Update_Heightmap_All on the component flags, but that's how it's been for a while.  We want either of the "_Editing" flags also
+	// passed through to the landscape actor.  Only pass through the _All flag in the case of bUpdateAll, not just from bUpdateCollision.
+	// Note that Update_Heightmap_All on the actor will cause the entire landscape to re-evaluate, not just this component.
+	uint32 ActorFlags = NewLayerUpdateFlags;
+	if (!bUpdateAll)
+	{
+		ActorFlags &= ~ELandscapeLayerUpdateMode::Update_Heightmap_All;  // Remove _All flag
+	}
+
 	if (ALandscape* LandscapeActor = GetLandscapeActor())
 	{
-		LandscapeActor->RequestLayersContentUpdate((ELandscapeLayerUpdateMode)NewLayerUpdateFlags);
+		LandscapeActor->RequestLayersContentUpdate((ELandscapeLayerUpdateMode)ActorFlags);
 	}
 	else
 	{
 		// This condition is only valid with streaming proxies, since the landscape actor will always be reachable in the non-WP Case
 		ALandscapeStreamingProxy* LandscapeStreamingProxy = CastChecked<ALandscapeStreamingProxy>(GetLandscapeProxy());
-		LandscapeStreamingProxy->RequestDeferredLayerContentUpdateModes((ELandscapeLayerUpdateMode)NewLayerUpdateFlags);
+		LandscapeStreamingProxy->RequestDeferredLayerContentUpdateModes((ELandscapeLayerUpdateMode)ActorFlags);
 	}
 }
 
@@ -6958,15 +6968,25 @@ void ULandscapeComponent::RequestWeightmapUpdate(bool bUpdateAll, bool bUpdateCo
 
 	LayerUpdateFlagPerMode |= NewLayerUpdateFlags;
 
+	// The typical codepath for sculpt or paint will have bUpdateCollision=true and bUpdateAll=false.  It's unclear why this also causes
+	// Update_Weightmap_All on the component flags, but that's how it's been for a while.  We want either of the "_Editing" flags also
+	// passed through to the landscape actor.  Only pass through the _All flag in the case of bUpdateAll, not just from bUpdateCollision.
+	// Note that Update_Weightmap_All on the actor will cause the entire landscape to re-evaluate, not just this component.
+	uint32 ActorFlags = NewLayerUpdateFlags;
+	if (!bUpdateAll)
+	{
+		ActorFlags &= ~ELandscapeLayerUpdateMode::Update_Weightmap_All;  // Remove _All flag
+	}
+
 	if (ALandscape* LandscapeActor = GetLandscapeActor())
 	{
-		LandscapeActor->RequestLayersContentUpdate((ELandscapeLayerUpdateMode)NewLayerUpdateFlags);
+		LandscapeActor->RequestLayersContentUpdate((ELandscapeLayerUpdateMode)ActorFlags);
 	}
 	else
 	{
 		// This condition is only valid with streaming proxies, since the landscape actor will always be reachable in the non-WP Case
 		ALandscapeStreamingProxy* LandscapeStreamingProxy = CastChecked<ALandscapeStreamingProxy>(GetLandscapeProxy());
-		LandscapeStreamingProxy->RequestDeferredLayerContentUpdateModes((ELandscapeLayerUpdateMode) NewLayerUpdateFlags);
+		LandscapeStreamingProxy->RequestDeferredLayerContentUpdateModes((ELandscapeLayerUpdateMode)ActorFlags);
 	}
 }
 
