* Fixed Ctrl-E not retrieving landscape materials on a given proxy when those come from a per-component material override

In detail :
* Added RetrieveAllLandscapeMaterials on ALandscapeProxy and ULandscapeComponent to be able to find all of the materials that can be used by the landscape, including per-LOD overrides, hole material, etc.
* Renamed/moved GetLayersFromMaterial into  RetrieveTargetLayerNamesFromMaterial and use the above function internally

#Synced-CL 50506770

diff --git a/Engine/Source/Editor/LandscapeEditor/Private/LandscapeEditorObject.cpp b/Engine/Source/Editor/LandscapeEditor/Private/LandscapeEditorObject.cpp
--- a/Engine/Source/Editor/LandscapeEditor/Private/LandscapeEditorObject.cpp
+++ b/Engine/Source/Editor/LandscapeEditor/Private/LandscapeEditorObject.cpp
@@ -891,7 +891,7 @@ void ULandscapeEditorObject::RefreshImportLayersList(bool bRefreshFromTarget, bo
 	else
 	{
 		Material = NewLandscape_Material.Get();
-		LayerNames = ALandscapeProxy::GetLayersFromMaterial(Material);
+		LayerNames = UE::Landscape::RetrieveTargetLayerNamesFromMaterial(Material);
 	}
 
 	TArray<FLandscapeImportLayer>& LandscapeLayers = bRefreshFromTarget ? ImportLandscape_Layers : NewLandscape_Layers;
diff --git a/Engine/Source/Editor/WorldBrowser/Private/Tiles/STiledLandscapeImportDlg.cpp b/Engine/Source/Editor/WorldBrowser/Private/Tiles/STiledLandscapeImportDlg.cpp
--- a/Engine/Source/Editor/WorldBrowser/Private/Tiles/STiledLandscapeImportDlg.cpp
+++ b/Engine/Source/Editor/WorldBrowser/Private/Tiles/STiledLandscapeImportDlg.cpp
@@ -17,6 +17,7 @@
 #include "LandscapeFileFormatInterface.h"
 #include "LandscapeEditorModule.h"
 #include "LandscapeDataAccess.h"
+#include "LandscapeUtils.h"
 
 #define LOCTEXT_NAMESPACE "WorldBrowser"
 
@@ -779,7 +780,7 @@ FText STiledLandscapeImportDlg::GenerateConfigurationText(int32 NumComponents, i
 
 void STiledLandscapeImportDlg::UpdateLandscapeLayerList()
 {
-	TArray<FName> LayerNamesList = ALandscapeProxy::GetLayersFromMaterial(ImportSettings.LandscapeMaterial.Get());
+	TArray<FName> LayerNamesList = UE::Landscape::RetrieveTargetLayerNamesFromMaterial(ImportSettings.LandscapeMaterial.Get());
 	LayerDataList.Empty();
 	for (FName LayerName : LayerNamesList)
 	{
diff --git a/Engine/Source/Editor/WorldBrowser/Private/Tiles/WorldTileCollectionModel.cpp b/Engine/Source/Editor/WorldBrowser/Private/Tiles/WorldTileCollectionModel.cpp
--- a/Engine/Source/Editor/WorldBrowser/Private/Tiles/WorldTileCollectionModel.cpp
+++ b/Engine/Source/Editor/WorldBrowser/Private/Tiles/WorldTileCollectionModel.cpp
@@ -46,6 +46,7 @@
 #include "LandscapeMeshProxyComponent.h"
 #include "LandscapeFileFormatInterface.h"
 #include "LandscapeEditorModule.h"
+#include "LandscapeUtils.h"
 #include "IMeshReductionManagerModule.h"
 #include "IMeshMergeUtilities.h"
 #include "MeshMergeModule.h"
@@ -599,7 +600,7 @@ void FWorldTileCollectionModel::FillWeightmapsSubMenu(FMenuBuilder& InMenuBuilde
 		auto TileModel = StaticCastSharedPtr<FWorldTileModel>(LevelModel);
 		if (TileModel->IsLandscapeBased())
 		{
-			TArray<FName> Layers = ALandscapeProxy::GetLayersFromMaterial(TileModel->GetLandscape()->LandscapeMaterial);
+			TArray<FName> Layers = UE::Landscape::RetrieveTargetLayerNamesFromMaterial(TileModel->GetLandscape()->LandscapeMaterial);
 			for (FName LayerName : Layers)
 			{
 				LayerNames.AddUnique(LayerName);
diff --git a/Engine/Source/Runtime/Landscape/Classes/LandscapeComponent.h b/Engine/Source/Runtime/Landscape/Classes/LandscapeComponent.h
--- a/Engine/Source/Runtime/Landscape/Classes/LandscapeComponent.h
+++ b/Engine/Source/Runtime/Landscape/Classes/LandscapeComponent.h
@@ -1081,6 +1081,12 @@ public:
 	/** Returns the components's LandscapeHoleMaterial, or the Component's OverrideLandscapeHoleMaterial if set */
 	LANDSCAPE_API UMaterialInterface* GetLandscapeHoleMaterial() const;
 
+	/**
+	 * Retrieve the list of all materials that might be used by this landscape component (LandscapeMaterial, LandscapeHoleMaterial, etc.)
+	 * @param OutMaterials list of unique materials retrieved. The materials found will be appended to the list
+	 */
+	LANDSCAPE_API void RetrieveAllLandscapeMaterials(TSet<UMaterialInterface*>& OutMaterials) const;
+
 	FBox ComputeApproximateBounds();
 
 #if WITH_EDITOR
diff --git a/Engine/Source/Runtime/Landscape/Classes/LandscapeProxy.h b/Engine/Source/Runtime/Landscape/Classes/LandscapeProxy.h
--- a/Engine/Source/Runtime/Landscape/Classes/LandscapeProxy.h
+++ b/Engine/Source/Runtime/Landscape/Classes/LandscapeProxy.h
@@ -1270,6 +1270,12 @@ PRAGMA_ENABLE_DEPRECATION_WARNINGS
 	// Get Hole Landscape Material assigned to this Landscape
 	virtual UMaterialInterface* GetLandscapeHoleMaterial() const;
 
+	/** 
+	 * Retrieve the list of all materials that might be used by this landscape proxy or its components (LandscapeMaterial, LandscapeHoleMaterial, component overrides, etc.) 
+	 * @param OutMaterials list of unique materials retrieved. The materials found will be appended to the list
+	 */
+	LANDSCAPE_API void RetrieveAllLandscapeMaterials(TSet<UMaterialInterface*>& OutMaterials) const;
+
 #if WITH_EDITOR
 	/* Serialize all hashes/guids that record the current state of this proxy */
 	void SerializeStateHashes(FArchive& Ar);
@@ -1303,7 +1309,9 @@ PRAGMA_ENABLE_DEPRECATION_WARNINGS
 	void UpdateProxyLayersWeightmapUsage();
 	void RequestProxyLayersWeightmapUsageUpdate();
 
+	UE_DEPRECATED(5.8, "Use UE::Landscape::RetrieveTargetLayerNamesFromMaterial")
 	LANDSCAPE_API static const TArray<FName> GetLayersFromMaterial(UMaterialInterface* Material);
+	UE_DEPRECATED(5.8, "Use RetrieveTargetLayerNamesFromMaterials")
 	LANDSCAPE_API const TArray<FName> GetLayersFromMaterial() const;
 
 	// Get all Layer names used by the materials assigned the proxy and components (Proxy Material, Component, Component LOD Material and Hole Material)
diff --git a/Engine/Source/Runtime/Landscape/Private/Landscape.cpp b/Engine/Source/Runtime/Landscape/Private/Landscape.cpp
--- a/Engine/Source/Runtime/Landscape/Private/Landscape.cpp
+++ b/Engine/Source/Runtime/Landscape/Private/Landscape.cpp
@@ -1123,6 +1123,20 @@ UMaterialInterface* ULandscapeComponent::GetLandscapeHoleMaterial() const
 	return nullptr;
 }
 
+void ULandscapeComponent::RetrieveAllLandscapeMaterials(TSet<UMaterialInterface*>& OutMaterials) const
+{
+	if (UMaterialInterface* Material = GetLandscapeMaterial(); 
+		Material && (Material != UMaterial::GetDefaultMaterial(MD_Surface)))
+	{
+		OutMaterials.Add(Material);
+	}
+	if (UMaterialInterface* Material = GetLandscapeHoleMaterial();
+		Material && (Material != UMaterial::GetDefaultMaterial(MD_Surface)))
+	{
+		OutMaterials.Add(Material);
+	}
+}
+
 #if WITH_EDITOR
 bool ULandscapeComponent::IsLandscapeHoleMaterialValid() const
 {
@@ -5231,6 +5245,31 @@ UMaterialInterface* ALandscapeStreamingProxy::GetLandscapeHoleMaterial() const
 	return nullptr;
 }
 
+void ALandscapeProxy::RetrieveAllLandscapeMaterials(TSet<UMaterialInterface*>& OutMaterials) const
+{
+	// Retrieve the materials specified in the proxy or its parent landscape (for a streaming proxy)
+	if (UMaterialInterface* Material = GetLandscapeMaterial();
+		Material && (Material != UMaterial::GetDefaultMaterial(MD_Surface)))
+	{
+		OutMaterials.Add(Material);
+	}
+
+	if (UMaterialInterface* Material = GetLandscapeHoleMaterial();
+		Material && (Material != UMaterial::GetDefaultMaterial(MD_Surface)))
+	{
+		OutMaterials.Add(Material);
+	}
+
+	// Also retrieve the materials from the components, as they may override the proxy's materials : 
+	for (ULandscapeComponent* Component : LandscapeComponents)
+	{
+		if (Component != nullptr)
+		{
+			Component->RetrieveAllLandscapeMaterials(OutMaterials);
+		}
+	}
+}
+
 #if WITH_EDITOR
 
 bool ALandscapeStreamingProxy::IsSharedPropertyOverridden(const FName& InPropertyName) const
diff --git a/Engine/Source/Runtime/Landscape/Private/LandscapeEdit.cpp b/Engine/Source/Runtime/Landscape/Private/LandscapeEdit.cpp
--- a/Engine/Source/Runtime/Landscape/Private/LandscapeEdit.cpp
+++ b/Engine/Source/Runtime/Landscape/Private/LandscapeEdit.cpp
@@ -1071,22 +1071,11 @@ bool ALandscapeProxy::GetReferencedContentObjects(TArray<UObject*>& Objects) con
 {
 	Super::GetReferencedContentObjects(Objects);
 
-	if (LandscapeMaterial != nullptr)
+	TSet<UMaterialInterface*> LandscapeMaterials;
+	RetrieveAllLandscapeMaterials(LandscapeMaterials);
+	for (UMaterialInterface* Material : LandscapeMaterials)
 	{
-		Objects.AddUnique(LandscapeMaterial);
-	}
-
-	for (const FLandscapePerLODMaterialOverride& LODOverrideMaterial : PerLODOverrideMaterials)
-	{
-		if (LODOverrideMaterial.Material != nullptr)
-		{
-			Objects.AddUnique(LODOverrideMaterial.Material);
-		}
-	}
-
-	if (LandscapeHoleMaterial != nullptr)
-	{
-		Objects.AddUnique(LandscapeHoleMaterial);
+		Objects.AddUnique(Material);
 	}
 
 	return true;
@@ -3083,53 +3072,27 @@ struct FHeightmapInfo
 	TArray<FColor*> HeightmapTextureMipData;
 };
 
+// Deprecated
 const TArray<FName> ALandscapeProxy::GetLayersFromMaterial(UMaterialInterface* MaterialInterface)
 {
-	if (MaterialInterface)
-	{
-		const FMaterialCachedExpressionData& CachedExpressionData = MaterialInterface->GetCachedExpressionData();
-		if (CachedExpressionData.EditorOnlyData)
-		{
-			TArray<FName> LayerNames = CachedExpressionData.EditorOnlyData->LandscapeLayerNames;
-			LayerNames.RemoveAll([](const FName& Name)
-			{
-				return Name.IsNone();
-			});
-
-			return LayerNames;
-		}
-	}
-	return FMaterialCachedExpressionEditorOnlyData::EmptyData.LandscapeLayerNames;
+	return UE::Landscape::RetrieveTargetLayerNamesFromMaterial(MaterialInterface);
 }
 
+// Deprecated
 const TArray<FName> ALandscapeProxy::GetLayersFromMaterial() const
 {
-	return GetLayersFromMaterial(LandscapeMaterial);
+	return RetrieveTargetLayerNamesFromMaterials();
 }
 
 TArray<FName> ALandscapeProxy::RetrieveTargetLayerNamesFromMaterials() const
 {
-	TSet<FName> LayerNames;
-	LayerNames.Append(GetLayersFromMaterial());
+	TSet<UMaterialInterface*> LandscapeMaterials;
+	RetrieveAllLandscapeMaterials(LandscapeMaterials);
 
-	for (int32 ComponentIndex = 0; ComponentIndex < LandscapeComponents.Num(); ComponentIndex++)
+	TSet<FName> LayerNames;
+	for (UMaterialInterface* Material : LandscapeMaterials)
 	{
-		const ULandscapeComponent* Component = LandscapeComponents[ComponentIndex];
-
-		// Add layers from per-component override materials
-		if ((Component != nullptr) && (Component->OverrideMaterial != nullptr))
-		{
-			LayerNames.Append(GetLayersFromMaterial(Component->OverrideMaterial));
-
-			// Hole Override material
-			LayerNames.Append(GetLayersFromMaterial( Component->OverrideHoleMaterial));
-				
-			// LOD Materials
-			for(const FLandscapePerLODMaterialOverride& OverrideMaterial : Component->GetPerLODOverrideMaterials())
-			{
-				LayerNames.Append(GetLayersFromMaterial(OverrideMaterial.Material));	
-			}
-		}
+		LayerNames.Append(UE::Landscape::RetrieveTargetLayerNamesFromMaterial(Material));
 	}
 	return LayerNames.Array();
 }
diff --git a/Engine/Source/Runtime/Landscape/Private/LandscapeUtils.cpp b/Engine/Source/Runtime/Landscape/Private/LandscapeUtils.cpp
--- a/Engine/Source/Runtime/Landscape/Private/LandscapeUtils.cpp
+++ b/Engine/Source/Runtime/Landscape/Private/LandscapeUtils.cpp
@@ -21,6 +21,7 @@
 #include "Editor/UnrealEdEngine.h"
 #include "EditorDirectories.h"
 #include "Engine/Texture2D.h"
+#include "MaterialCachedData.h"
 #include "MaterialDomain.h"
 #include "LandscapeComponent.h"
 #include "LandscapeMaterialInstanceConstant.h"
@@ -569,6 +570,24 @@ bool DeleteActors(const TArray<AActor*>& InActorsToDelete, UWorld* InWorld, bool
 	return bSuccess;
 }
 
+TArray<FName> RetrieveTargetLayerNamesFromMaterial(UMaterialInterface* InMaterial)
+{
+	TArray<FName> LayerNames;
+	if (InMaterial)
+	{
+		const FMaterialCachedExpressionData& CachedExpressionData = InMaterial->GetCachedExpressionData();
+		if (CachedExpressionData.EditorOnlyData)
+		{
+			LayerNames = CachedExpressionData.EditorOnlyData->LandscapeLayerNames;
+			LayerNames.RemoveAll([](const FName& Name)
+				{
+					return Name.IsNone();
+				});
+		}
+	}
+	return LayerNames;
+}
+
 #endif // WITH_EDITOR
 
 bool LandscapeProxySortPredicate(const TWeakObjectPtr<ALandscapeProxy> APtr, const TWeakObjectPtr<ALandscapeProxy> BPtr)
diff --git a/Engine/Source/Runtime/Landscape/Public/LandscapeUtils.h b/Engine/Source/Runtime/Landscape/Public/LandscapeUtils.h
--- a/Engine/Source/Runtime/Landscape/Public/LandscapeUtils.h
+++ b/Engine/Source/Runtime/Landscape/Public/LandscapeUtils.h
@@ -355,6 +355,15 @@ LANDSCAPE_API FString ConvertTargetLayerNamesToString(const TArrayView<const FNa
  */
 LANDSCAPE_API bool DeleteActors(const TArray<AActor*>& InActorsToDelete, UWorld* InWorld, bool bInAllowUI);
 
+
+/**
+ * Analyzes the material's material expressions and return the landscape target layer names declared in them
+ * @param InMaterial the material to analyze
+ * 
+ * @return the list of unique target layer names declared in that material
+ */
+LANDSCAPE_API TArray<FName> RetrieveTargetLayerNamesFromMaterial(UMaterialInterface* InMaterial);
+
 #endif // WITH_EDITOR
 
 /** Helper to sort ALandscapeProxy actors by their section base. Compares X axis values by default and Y values if needed */
