Overridable landscape grass types

* Landscape grass types are now addressed by their declared (input) name in the landscape material
* This allows to specify grass types overrides within a landscape proxy
* This also allows to export a certain grass map without specifying a landscape grass type (e.g. meant for PCG)
* This also allows multiple grass maps to be associated withe same landscape grass type, if needed

Note :
* For now, the grass types overrides property is hidden within the landscape proxy details panel. It will be made available when we customize the UX in order to help the user choose valid grass names via a combo box

#Synced-CL 50539597

#Synced-CL 50539597

#Synced-CL 50539597

diff --git a/Engine/Plugins/PCG/Shaders/Private/PCGGrassMapUnpackerCS.usf b/Engine/Plugins/PCG/Shaders/Private/PCGGrassMapUnpackerCS.usf
--- a/Engine/Plugins/PCG/Shaders/Private/PCGGrassMapUnpackerCS.usf
+++ b/Engine/Plugins/PCG/Shaders/Private/PCGGrassMapUnpackerCS.usf
@@ -56,7 +56,7 @@ RWTexture2D<float> OutUnpackedHeight;
 StructuredBuffer<int> InComponentHeightOffsets;
 StructuredBuffer<int2> InComponentGrassmapOffsetAndChannels;
 StructuredBuffer<int> InOutputIndexToExportIndex;
-int InNumExportedGrassTypes;
+int InNumExportedGrassMaps;
 
 int4 InLinearTileIndexToComponentIndex[PCG_MAX_NUM_LANDSCAPE_COMPONENTS]; // Only x component used, 16 byte alignment on arrays
 uint2 InNumTiles;
@@ -117,7 +117,7 @@ void PCGGrassMapUnpacker_CS(const uint3 DispatchThreadId : SV_DispatchThreadID)
 
 	// Coordinates into input height & grass maps (see packing above).
 	const int ExportSlice = InOutputIndexToExportIndex[OutputSlice];
-	const int2 ComponentAndGrassXOffsetAndChannel = (ExportSlice > -1) ? InComponentGrassmapOffsetAndChannels[ComponentIndex * InNumExportedGrassTypes + ExportSlice] : -1;
+	const int2 ComponentAndGrassXOffsetAndChannel = (ExportSlice > -1) ? InComponentGrassmapOffsetAndChannels[ComponentIndex * InNumExportedGrassMaps + ExportSlice] : -1;
 	if (ComponentAndGrassXOffsetAndChannel.x != -1)
 	{
 		const uint ComponentXOffset = ComponentAndGrassXOffsetAndChannel.x;
diff --git a/Engine/Plugins/PCG/Source/PCG/Private/Elements/PCGGenerateLandscapeTextures.cpp b/Engine/Plugins/PCG/Source/PCG/Private/Elements/PCGGenerateLandscapeTextures.cpp
--- a/Engine/Plugins/PCG/Source/PCG/Private/Elements/PCGGenerateLandscapeTextures.cpp
+++ b/Engine/Plugins/PCG/Source/PCG/Private/Elements/PCGGenerateLandscapeTextures.cpp
@@ -13,6 +13,7 @@
 #include "Helpers/PCGHelpers.h"
 #include "Metadata/Accessors/PCGAttributeAccessorHelpers.h"
 
+#include "Algo/Transform.h"
 #include "DrawDebugHelpers.h"
 #include "EngineDefines.h" // For UE_ENABLE_DEBUG_DRAWING
 #include "Landscape.h"
@@ -26,6 +27,7 @@
 #include "RenderGraphUtils.h"
 #include "SystemTextures.h"
 #include "Engine/World.h"
+#include "UObject/FortniteMainBranchObjectVersion.h"
 
 #include UE_INLINE_GENERATED_CPP_BY_NAME(PCGGenerateLandscapeTextures)
 
@@ -84,7 +86,7 @@ namespace PCGGenerateLandscapeTextures
 			return false;
 		}
 
-		const bool bGenerateGrassMaps = !InData.Context->GrassTypesToExport.IsEmpty();
+		const bool bGenerateGrassMaps = !InData.Context->GrassNamesToExport.IsEmpty();
 		const bool bGenerateHeightMap = InData.Settings->bGenerateHeightMap;
 
 		UE::Landscape::Grass::EGrassWeightExporterFlags Flags = UE::Landscape::Grass::EGrassWeightExporterFlags::None;
@@ -97,12 +99,32 @@ namespace PCGGenerateLandscapeTextures
 			Flags |= UE::Landscape::Grass::EGrassWeightExporterFlags::NeedsHeightmap;
 		}
 
-		TSharedPtr<FLandscapeGrassWeightExporter> Exporter = MakeShared<FLandscapeGrassWeightExporter>(
-			InData.Landscape,
-			InData.LandscapeComponents,
-			Flags,
-			InData.Context->GrassTypesToExport,
-			/*InRequestedHeightMips = */TConstArrayView<int32>{});
+		TUniquePtr<FLandscapeGrassWeightExporter> Exporter;
+		{
+			TArray<FName> GrassExporterGrassNames;
+			if (InData.Settings->bTagDataWithLandscapeGrassTypeAssetNames)
+			{
+				// If bTagDataWithLandscapeGrassTypeAssetNames, translate the grass types names to the actual grass names, because that's what the grass weight exporter knows about : 
+				GrassExporterGrassNames.Reserve(InData.Context->GrassNamesToExport.Num());
+				Algo::Transform(InData.Context->GrassNamesToExport, GrassExporterGrassNames, [Context=InData.Context](const FName& InGrassTypeName)
+				{
+					const FName* FoundGrassName = Context->GrassTypeToGrassNameMap.Find(InGrassTypeName);
+					check(FoundGrassName != nullptr);
+					return *FoundGrassName;
+				});
+			}
+			else
+			{
+				GrassExporterGrassNames = InData.Context->GrassNamesToExport;
+			}
+
+			Exporter.Reset(new FLandscapeGrassWeightExporter(
+				InData.Landscape,
+				InData.LandscapeComponents,
+				Flags,
+				GrassExporterGrassNames,
+				/*InRequestedHeightMips = */TConstArrayView<int32>{}));
+		}
 
 		const FIntPoint& OutputTextureSize = Exporter->GetTargetSize();
 		const int32 Max2DTextureDimension = GetMax2DTextureDimension();
@@ -180,7 +202,7 @@ namespace PCGGenerateLandscapeTextures
 			, DispatchSize
 			, OutputTextureData = InData.Context->OutputTextureData
 			, bLastGenerationPass = InData.bLastGenerationPass
-			, Exporter](FRHICommandListImmediate& RHICmdList)
+			, Exporter = MoveTemp(Exporter)](FRHICommandListImmediate& RHICmdList)
 		{
 			TRACE_CPUPROFILER_EVENT_SCOPE(PCGGenerateLandscapeTextures);
 			LLM_SCOPE_BYTAG(PCG);
@@ -207,7 +229,7 @@ namespace PCGGenerateLandscapeTextures
 			Exporter->RenderLandscapeComponentToTexture_RenderThread(GraphBuilder, GrassMapTexture);
 
 			// Must have at least one slice to be a valid resource.
-			const uint32 NumOutputSlices = FMath::Max(Context->GrassTypesToOutput.Num(), 1);
+			const uint32 NumOutputSlices = FMath::Max(Context->GrassNamesToOutput.Num(), 1);
 
 			FRDGTextureRef HeightMap = nullptr;
 			bool bHeightMapRequiresExport = false;
@@ -263,7 +285,7 @@ namespace PCGGenerateLandscapeTextures
 			Parameters->InTileGridOrigin = TileGridStart;
 			Parameters->InDispatchOffset = FUintVector2(DispatchMin.X, DispatchMin.Y);
 			Parameters->InDispatchSize = FUintVector2(DispatchSize.X, DispatchSize.Y);
-			Parameters->InNumExportedGrassTypes = Context->GrassTypesToExport.Num();
+			Parameters->InNumExportedGrassMaps = Context->GrassNamesToExport.Num();
 
 			// Initialize to invalid component indices.
 			FMemory::Memset(&Parameters->InLinearTileIndexToComponentIndex[0], -1, FPCGGrassMapUnpackerCS::MaxNumLandscapeComponents * sizeof(int32) * 4);
@@ -302,27 +324,37 @@ namespace PCGGenerateLandscapeTextures
 			Parameters->InComponentHeightOffsets = GraphBuilder.CreateSRV(ComponentHeightOffsetsBuffer);
 
 			FRDGBufferRef ComponentGrassmapOffsetAndChannelsBuffer;
-			if (!Context->GrassTypesToExport.IsEmpty())
+			if (!Context->GrassNamesToExport.IsEmpty())
 			{
 				// For each component, for each grass type, X offset for where to find the grass data in the grass texture (.x) and on which channel (.y)
 				//  It's possible a grass map is absent from a given component, in which case, the offset and channel will be set to -1
 				TArray<FIntVector2> ComponentGrassmapOffsetAndChannels;
-				ComponentGrassmapOffsetAndChannels.SetNumUninitialized(LandscapeComponentKeys.Num() * Context->GrassTypesToExport.Num());
+				ComponentGrassmapOffsetAndChannels.SetNumUninitialized(LandscapeComponentKeys.Num() * Context->GrassNamesToExport.Num());
 				FMemory::Memset(ComponentGrassmapOffsetAndChannels.GetData(), -1, ComponentGrassmapOffsetAndChannels.Num() * sizeof(FIntVector2));
 
 				for (int32 ComponentIndex = 0; ComponentIndex < LandscapeComponentKeys.Num(); ++ComponentIndex)
 				{
-					for (int32 GrassmapIndex = 0; GrassmapIndex < Context->GrassTypesToExport.Num(); ++GrassmapIndex)
+					for (int32 GrassmapIndex = 0; GrassmapIndex < Context->GrassNamesToExport.Num(); ++GrassmapIndex)
 					{
-						FIntVector2& ComponentGrassmapOffsetAndChannel = ComponentGrassmapOffsetAndChannels[ComponentIndex * Context->GrassTypesToExport.Num() + GrassmapIndex];
+						FIntVector2& ComponentGrassmapOffsetAndChannel = ComponentGrassmapOffsetAndChannels[ComponentIndex * Context->GrassNamesToExport.Num() + GrassmapIndex];
+
+						FName LandscapeExporterGrassName = Context->GrassNamesToExport[GrassmapIndex];
+						// If we have a grass type name -> grass name mapping, that's because bTagDataWithLandscapeGrassTypeAssetNames was used, so we need to translate the grass type name into the grass name, 
+						//  which is what the grass exporter understands :
+						if (!Context->GrassTypeToGrassNameMap.IsEmpty())
+						{
+							const FName* FoundGrassName = Context->GrassTypeToGrassNameMap.Find(LandscapeExporterGrassName);
+							check(FoundGrassName != nullptr);
+							LandscapeExporterGrassName = *FoundGrassName;
+						}
 
 						TValueOrError<TPair<FIntRect /*OutputRect*/, uint8 /*ChannelIndex*/>, FLandscapeGrassWeightExporter_RenderThread::EGrassmapTextureInfoError> TextureInfo =
-							Exporter->GetTextureInfoForGrass(LandscapeComponentKeys[ComponentIndex], Context->GrassTypesToExport[GrassmapIndex]);
+							Exporter->GetTextureInfoForGrass(LandscapeComponentKeys[ComponentIndex], LandscapeExporterGrassName);
 						if (TextureInfo.HasError())
 						{
 							// It can fail, but only if that's because the grass map is absent from this component
-							checkf(TextureInfo.GetError() == FLandscapeGrassWeightExporter_RenderThread::EGrassmapTextureInfoError::InvalidGrassType,
-								TEXT("Failed to retrieve grass map info on component %s : reason : %i"), *LandscapeComponentKeys[ComponentIndex].ToString(), TextureInfo.GetError());
+							checkf(TextureInfo.GetError() == FLandscapeGrassWeightExporter_RenderThread::EGrassmapTextureInfoError::InvalidGrassName,
+								TEXT("Failed to retrieve grass map (%s) info on component %s : reason : %i"), *LandscapeExporterGrassName.ToString() , *LandscapeComponentKeys[ComponentIndex].ToString(), TextureInfo.GetError());
 						}
 						else
 						{
@@ -332,7 +364,7 @@ namespace PCGGenerateLandscapeTextures
 						}
 					}
 				}
-				check(ComponentGrassmapOffsetAndChannels.Num() == LandscapeComponentKeys.Num() * Context->GrassTypesToExport.Num());
+				check(ComponentGrassmapOffsetAndChannels.Num() == LandscapeComponentKeys.Num() * Context->GrassNamesToExport.Num());
 				ComponentGrassmapOffsetAndChannelsBuffer = CreateStructuredBuffer(GraphBuilder, TEXT("PCGGenerateLandscapeTexturesComponentGrassmapOffsetAndChannels"), ComponentGrassmapOffsetAndChannels);
 			}
 			else
@@ -408,6 +440,20 @@ UPCGGenerateLandscapeTexturesSettings::UPCGGenerateLandscapeTexturesSettings()
 	}
 }
 
+void UPCGGenerateLandscapeTexturesSettings::Serialize(FArchive& Ar)
+{
+	Super::Serialize(Ar);
+	Ar.UsingCustomVersion(FFortniteMainBranchObjectVersion::GUID);
+#if WITH_EDITORONLY_DATA
+	if (Ar.IsLoading() && Ar.CustomVer(FFortniteMainBranchObjectVersion::GUID) < FFortniteMainBranchObjectVersion::OverridableLandscapeGrassTypes)
+	{
+		// To maintain backwards compatibility, we have to rely on landscape grass type asset names, since that used to be the way to address grass maps in the past and automatic deprecation is not possible
+		//  until the graph runs (to gain access to the landscape component's NamedGrassTypes)
+		bTagDataWithLandscapeGrassTypeAssetNames = true;
+	}
+#endif // WITH_EDITORONLY_DATA
+}
+
 #if WITH_EDITOR
 TArray<FPCGPreConfiguredSettingsInfo> UPCGGenerateLandscapeTexturesSettings::GetPreconfiguredInfo() const
 {
@@ -513,8 +559,6 @@ void FPCGGenerateLandscapeTexturesContext::AddExtraStructReferencedObjects(FRefe
 
 	Collector.AddReferencedObjects(GrassMapTextureDatas);
 	Collector.AddReferencedObjects(TexturesToStream);
-	Collector.AddReferencedObjects(GrassTypesToExport);
-	Collector.AddReferencedObjects(GrassTypesToOutput);
 }
 
 void FPCGGenerateLandscapeTexturesElement::GetDependenciesCrc(const FPCGGetDependenciesCrcParams& InParams, FPCGCrc& OutCrc) const
@@ -583,16 +627,16 @@ bool FPCGGenerateLandscapeTexturesElement::ExecuteInternal(FPCGContext* InContex
 			return true;
 		}
 
-		TArray<FString> SelectedGrassTypes;
+		TArray<FName> SelectedGrassNames;
 
 		if (!Settings->bOverrideFromInput)
 		{
-			SelectedGrassTypes = Settings->SelectedGrassTypes;
+			SelectedGrassNames = Settings->SelectedGrassTypes;
 		}
 		else
 		{
 			const TArray<FPCGTaggedData> OverrideTaggedDatas = InContext->InputData.GetInputsByPin(PCGGenerateLandscapeTextures::GrassTypeOverridesPinLabel);
-			TArray<FString> GrassTypeOverrides;
+			TArray<FName> GrassNameOverrides;
 
 			for (const FPCGTaggedData& OverrideTaggedData : OverrideTaggedDatas)
 			{
@@ -600,9 +644,9 @@ bool FPCGGenerateLandscapeTexturesElement::ExecuteInternal(FPCGContext* InContex
 				{
 					const FPCGAttributePropertyInputSelector Selector = Settings->GrassTypesAttribute.CopyAndFixLast(OverrideData);
 
-					if (PCGAttributeAccessorHelpers::ExtractAllValues(OverrideData, Selector, GrassTypeOverrides, InContext))
+					if (PCGAttributeAccessorHelpers::ExtractAllValues(OverrideData, Selector, GrassNameOverrides, InContext))
 					{
-						SelectedGrassTypes = std::move(GrassTypeOverrides);
+						SelectedGrassNames = std::move(GrassNameOverrides);
 					}
 					else
 					{
@@ -616,15 +660,15 @@ bool FPCGGenerateLandscapeTexturesElement::ExecuteInternal(FPCGContext* InContex
 		const bool bOutputDataForEachSelectedGrassType = !Settings->bExcludeSelectedGrassTypes && Settings->bOutputDataForEachSelectedGrassType;
 		if (bOutputDataForEachSelectedGrassType)
 		{
-			Context->GrassTypesToOutput.SetNum(SelectedGrassTypes.Num());
+			Context->GrassNamesToOutput.SetNum(SelectedGrassNames.Num());
 
-			Context->OutputIndexToExportIndex.SetNumUninitialized(SelectedGrassTypes.Num());
+			Context->OutputIndexToExportIndex.SetNumUninitialized(SelectedGrassNames.Num());
 			FMemory::Memset(Context->OutputIndexToExportIndex.GetData(), -1, Context->OutputIndexToExportIndex.Num() * Context->OutputIndexToExportIndex.GetTypeSize());
 		}
 
 		Context->GenerationBounds = ExecutionSource->GetExecutionState().GetBounds();
 		TArray<ULandscapeComponent*> OverlappingComponents;
-		TSet<TObjectPtr<ULandscapeGrassType>> ProcessedGrassTypes;
+		TSet<FName> ProcessedGrassNames;
 
 		for (const TSoftObjectPtr<ALandscapeProxy>& LandscapeProxyPtr : LandscapeData->Landscapes)
 		{
@@ -671,48 +715,77 @@ bool FPCGGenerateLandscapeTexturesElement::ExecuteInternal(FPCGContext* InContex
 							break;
 						}
 
-						FString GrassName; // Used in loop but hoisted for efficiency
+						LandscapeComponent->UpdateGrassTypes();
 
-						for (const TObjectPtr<ULandscapeGrassType>& GrassType : LandscapeComponent->GetGrassTypes())
+						TMap<FName, TObjectPtr<ULandscapeGrassType>> ValidNamedGrassTypes = UE::Landscape::Grass::FilterValidNamedGrassTypes(LandscapeComponent->GetNamedGrassTypes());
+						for (const TPair<FName, TObjectPtr<ULandscapeGrassType>>& ItPair : ValidNamedGrassTypes)
 						{
-							if (ProcessedGrassTypes.Contains(GrassType))
+							check(!ItPair.Key.IsNone());
+
+							FName GrassName = ItPair.Key;
+
+							// If bTagDataWithLandscapeGrassTypeAssetNames is used, the grass name will be the grass type asset's name, so we store a mapping between this and the actual grass name, so that the landscape grass exporter can work with the latter :
+							if (Settings->bTagDataWithLandscapeGrassTypeAssetNames)
+							{
+								if (ItPair.Value == nullptr)
+								{
+									continue;
+								}
+
+								FName GrassTypeName = ItPair.Value->GetFName();
+								// In this case, use the grass type's name as the actual "PCG" grass name : 
+								GrassName = GrassTypeName;
+
+								// Register the grass type name -> grass name mapping or warn if that mapping has already been registered with a different grass name
+								if (FName* FoundGrassName = Context->GrassTypeToGrassNameMap.Find(GrassTypeName))
+								{
+									if (*FoundGrassName != ItPair.Key)
+									{
+										PCGLog::LogWarningOnGraph(FText::Format(LOCTEXT("DuplicateGrassNamesForGrassTypeAsset", "Found more than one grass name associated with the same grass type asset's name. "
+											"Found grass name {0} for Landscape Grass Type asset {1} while grass name {2} was already registered. This only happens because bTagDataWithLandscapeGrassTypeAssetNames is used. Consider unchecking this option and change "
+											"the PCG graph accordingly in order to resolve the ambiguity, since only the grass name defines the actual grass map. "),
+											FText::FromString(ItPair.Key.ToString()), FText::FromString(GrassTypeName.ToString()), FText::FromString(FoundGrassName->ToString())), Context);
+									}
+								}
+								else
+								{
+									Context->GrassTypeToGrassNameMap.Add(GrassTypeName, ItPair.Key);
+								}
+							}
+
+							if (ProcessedGrassNames.Contains(GrassName))
 							{
 								continue;
 							}
 
-							ProcessedGrassTypes.Add(GrassType);
+							ProcessedGrassNames.Add(GrassName);
+							const int32 IndexInSelection = SelectedGrassNames.IndexOfByKey(GrassName);
 
-							GrassType->GetName(GrassName);
-							const int32 IndexInSelection = SelectedGrassTypes.IndexOfByKey(GrassName);
 							const bool bIsSelectedLayer = IndexInSelection != INDEX_NONE;
-
-							// Register the grass type if it is included or not excluded.
+							// Register the grass name if it is included or not excluded.
 							if (bIsSelectedLayer != Settings->bExcludeSelectedGrassTypes)
 							{
-								const int32 SizeBefore = Context->GrassTypesToExport.Num();
-								Context->GrassTypesToExport.AddUnique(GrassType);
-
-								// Update grass type lists if this grass type is newly added.
-								if (Context->GrassTypesToExport.Num() > SizeBefore)
+								const int32 SizeBefore = Context->GrassNamesToExport.Num();
+								Context->GrassNamesToExport.AddUnique(GrassName);
+								// Update grass name lists if this grass name is newly added.
+								if (Context->GrassNamesToExport.Num() > SizeBefore)
 								{
 									if (bOutputDataForEachSelectedGrassType)
 									{
-										// If we need to follow the prescribed list of grass types, fill in the entry here.
-										Context->GrassTypesToOutput[IndexInSelection] = GrassType;
-										Context->OutputIndexToExportIndex[IndexInSelection] = Context->GrassTypesToExport.Num() - 1;
+										// If we need to follow the prescribed list of grass names, fill in the entry here.
+										Context->GrassNamesToOutput[IndexInSelection] = GrassName;
+										Context->OutputIndexToExportIndex[IndexInSelection] = Context->GrassNamesToExport.Num() - 1;
 									}
 									else
 									{
-										// Otherwise build list of grass types as we encounter them, if they are not filtered out by the include/exclude list.
-										Context->GrassTypesToOutput.Add(GrassType);
-										Context->OutputIndexToExportIndex.Add(Context->GrassTypesToExport.Num() - 1);
+										// Otherwise build list of grass names as we encounter them, if they are not filtered out by the include/exclude list.
+										Context->GrassNamesToOutput.Add(GrassName);
+										Context->OutputIndexToExportIndex.Add(Context->GrassNamesToExport.Num() - 1);
 									}
 								}
 							}
 						}
 
-						LandscapeComponent->UpdateGrassTypes();
-
 						Context->LandscapeComponents.Add(LandscapeComponent);
 						OverlappingComponents.Add(LandscapeComponent);
 
@@ -755,7 +828,7 @@ bool FPCGGenerateLandscapeTexturesElement::ExecuteInternal(FPCGContext* InContex
 			}
 		}
 
-		if (!Context->Landscape.Get() || Context->LandscapeComponents.IsEmpty() || (Context->GrassTypesToOutput.IsEmpty() && !Settings->bGenerateHeightMap) || !Context->TotalLandscapeComponentBounds.IsValid)
+		if (!Context->Landscape.Get() || Context->LandscapeComponents.IsEmpty() || (Context->GrassNamesToOutput.IsEmpty() && !Settings->bGenerateHeightMap) || !Context->TotalLandscapeComponentBounds.IsValid)
 		{
 			return true;
 		}
@@ -804,7 +877,7 @@ bool FPCGGenerateLandscapeTexturesElement::ExecuteInternal(FPCGContext* InContex
 						Context->TexturesToStream.Add(HeightMap);
 					}
 
-					if (Context->GrassTypesToOutput.Num() > 0)
+					if (Context->GrassNamesToOutput.Num() > 0)
 					{
 						const ERHIFeatureLevel::Type FeatureLevel = World->GetFeatureLevel();
 						for (UTexture2D* WeightmapTexture : LandscapeComponent->GetRenderedWeightmapTexturesForFeatureLevel(FeatureLevel))
@@ -1018,8 +1091,8 @@ bool FPCGGenerateLandscapeTexturesElement::ExecuteInternal(FPCGContext* InContex
 		}
 		else
 		{
-			// Create the texture data objects if they haven't been created already. There should be one per selected grass type.
-			for (int DataIndex = 0; DataIndex < Context->GrassTypesToOutput.Num(); ++DataIndex)
+			// Create the texture data objects if they haven't been created already. There should be one per selected grass name.
+			for (int DataIndex = 0; DataIndex < Context->GrassNamesToOutput.Num(); ++DataIndex)
 			{
 				Context->GrassMapTextureDatas.Add(FPCGContext::NewObject_AnyThread<UPCGTextureData>(Context));
 			}
@@ -1052,13 +1125,13 @@ bool FPCGGenerateLandscapeTexturesElement::ExecuteInternal(FPCGContext* InContex
 		}
 		else
 		{
-			for (int DataIndex = 0; DataIndex < Context->GrassTypesToOutput.Num(); ++DataIndex)
+			for (int DataIndex = 0; DataIndex < Context->GrassNamesToOutput.Num(); ++DataIndex)
 			{
 				check(Context->GrassMapTextureDatas.IsValidIndex(DataIndex));
 				UPCGTextureData* TextureData = Context->GrassMapTextureDatas[DataIndex];
 				check(TextureData);
 
-				const int32 ExportIndex = Context->GrassTypesToOutput[DataIndex] ? Context->GrassTypesToExport.IndexOfByKey(Context->GrassTypesToOutput[DataIndex]) : INDEX_NONE;
+				const int32 ExportIndex = Context->GrassNamesToOutput[DataIndex].IsNone() ? INDEX_NONE : Context->GrassNamesToExport.IndexOfByKey(Context->GrassNamesToOutput[DataIndex]);
 
 				if (ExportIndex != INDEX_NONE)
 				{
@@ -1139,7 +1212,7 @@ bool FPCGGenerateLandscapeTexturesElement::ExecuteInternal(FPCGContext* InContex
 	{
 		FString TagString; // Used in loop but hoisted for efficiency.
 
-		for (int DataIndex = 0; DataIndex < Context->GrassTypesToOutput.Num(); ++DataIndex)
+		for (int DataIndex = 0; DataIndex < Context->GrassNamesToOutput.Num(); ++DataIndex)
 		{
 			check(Context->GrassMapTextureDatas.IsValidIndex(DataIndex));
 			UPCGTextureData* TextureData = Context->GrassMapTextureDatas[DataIndex];
@@ -1155,22 +1228,16 @@ bool FPCGGenerateLandscapeTexturesElement::ExecuteInternal(FPCGContext* InContex
 			OutTaggedData.Pin = PCGGenerateLandscapeTextures::OutputGrassMapsPinLabel;
 			OutTaggedData.Data = TextureData;
 
-			FString GrassTypeName = TEXT("");
-
-			if (ULandscapeGrassType* GrassType = Context->GrassTypesToOutput[DataIndex].Get())
-			{
-				GrassType->GetName(GrassTypeName);
-			}
-
-			if (!GrassTypeName.IsEmpty())
+			FName GrassName = Context->GrassNamesToOutput[DataIndex];
+			if (!GrassName.IsNone())
 			{
-				PCGComputeHelpers::GetPrefixedDataLabel(GrassTypeName, TagString);
+				PCGComputeHelpers::GetPrefixedDataLabel(GrassName.ToString(), TagString);
 				OutTaggedData.Tags.Add(TagString);
 			}
 			else if (!Settings->bExcludeSelectedGrassTypes && Settings->SelectedGrassTypes.IsValidIndex(DataIndex))
 			{
 				// We are using inclusion list, look up name
-				PCGComputeHelpers::GetPrefixedDataLabel(Settings->SelectedGrassTypes[DataIndex], TagString);
+				PCGComputeHelpers::GetPrefixedDataLabel(Settings->SelectedGrassTypes[DataIndex].ToString(), TagString);
 				OutTaggedData.Tags.Add(TagString);
 			}
 			else
diff --git a/Engine/Plugins/PCG/Source/PCG/Public/Elements/PCGGenerateLandscapeTextures.h b/Engine/Plugins/PCG/Source/PCG/Public/Elements/PCGGenerateLandscapeTextures.h
--- a/Engine/Plugins/PCG/Source/PCG/Public/Elements/PCGGenerateLandscapeTextures.h
+++ b/Engine/Plugins/PCG/Source/PCG/Public/Elements/PCGGenerateLandscapeTextures.h
@@ -27,6 +27,10 @@ class UPCGGenerateLandscapeTexturesSettings : public UPCGSettings
 public:
 	UPCGGenerateLandscapeTexturesSettings();
 
+	//~Begin UObject interface
+	virtual void Serialize(FArchive& Ar) override;
+	//~End UObject interface
+
 	//~Begin UPCGSettings interface
 #if WITH_EDITOR
 	virtual FName GetDefaultNodeName() const override { return FName(TEXT("GenerateLandscapeTextures")); }
@@ -53,7 +57,7 @@ protected:
 public:
 	/** Select which grass types to generate. */
 	UPROPERTY(BlueprintReadWrite, EditAnywhere, Category = Settings, meta = (EditCondition = "!bOverrideFromInput"))
-	TArray<FString> SelectedGrassTypes;
+	TArray<FName> SelectedGrassTypes;
 
 	/** Override grass types from input. */
 	UPROPERTY(BlueprintReadWrite, EditAnywhere, Category = Settings)
@@ -85,6 +89,10 @@ public:
 
 	UPROPERTY(BlueprintReadWrite, EditAnywhere, Category = Settings, meta=(ClampMin=1))
 	int32 NumLandscapeComponentsGeneratedPerFrame = 1;
+
+	/** Use the Landscape Grass Type asset's names, rather than the grass names to sample the grass maps. Mostly there for deprecation purposes : it used to be that these were the only way to address grass maps, but that is not the case anymore : the grass input's name in the Landscape Grass Output material expression is what drives the grass maps, so they should be used from now on (for example, the material can now use the same Landscape Grass Type asset for 2 different grass maps, so addressing them by asset name is not sufficient to differentiate them). */
+	UPROPERTY(BlueprintReadWrite, EditAnywhere, Category = Settings, AdvancedDisplay)
+	bool bTagDataWithLandscapeGrassTypeAssetNames = false;
 };
 
 struct FPCGGenerateLandscapeTexturesContext : public FPCGContext
@@ -109,13 +117,13 @@ public:
 	/** Exported result textures. */
 	TSharedPtr<FOutputTextureData, ESPMode::ThreadSafe> OutputTextureData;
 
-	/** The grass types passed to the grass exporter. Each of these grass types is present on at least one of the landscape components. */
-	TArray<TObjectPtr<ULandscapeGrassType>> GrassTypesToExport;
+	/** The grass names passed to the grass exporter. Each of these grass names is present on at least one of the landscape components. */
+	TArray<FName> GrassNamesToExport;
 
-	/** The grass types to output from the node. Grass types that were not present on the overlapping landscape components will emit empty/black texture data. */
-	TArray<TObjectPtr<ULandscapeGrassType>> GrassTypesToOutput;
+	/** The grass names to output from the node. Grass names that were not present on the overlapping landscape components will emit empty/black texture data. */
+	TArray<FName> GrassNamesToOutput;
 
-	/** Maps output grass type indices to exported grass type indices. */
+	/** Maps output grass name indices to exported grass name indices. */
 	TArray<int32> OutputIndexToExportIndex;
 
 	TWeakObjectPtr<ALandscape> Landscape;
@@ -157,6 +165,10 @@ public:
 
 	TObjectPtr<UTexture> BlackTexture;
 
+	/** Indirection table to support bTagDataWithLandscapeGrassTypeAssetNames, where grass names for PCG(that used to be the Landscape Grass Type asset names) don't match the actual landscape grass names
+	 (that come from the input pins of the Landscape Grass Output Material Expression) */
+	TMap<FName, FName> GrassTypeToGrassNameMap;
+
 protected:
 	virtual void AddExtraStructReferencedObjects(FReferenceCollector& Collector) override;
 
@@ -176,6 +188,10 @@ public:
 	TRefCountPtr<IPooledRenderTarget> HeightHandle = nullptr;
 	UE_DEPRECATED(5.8, "Replaced with FOutputTextureData")
 	TRefCountPtr<IPooledRenderTarget> GrassMapHandle = nullptr;
+	UE_DEPRECATED(5.8, "Replaced with GrassNamesToExport")
+	TArray<TObjectPtr<ULandscapeGrassType>> GrassTypesToExport;
+	UE_DEPRECATED(5.8, "Replaced with GrassNamesToOutput")
+	TArray<TObjectPtr<ULandscapeGrassType>> GrassTypesToOutput;
 };
 
 class FPCGGenerateLandscapeTexturesElement : public IPCGElement
diff --git a/Engine/Plugins/PCG/Source/PCGCompute/Internal/PCGGrassMapUnpackerCS.h b/Engine/Plugins/PCG/Source/PCGCompute/Internal/PCGGrassMapUnpackerCS.h
--- a/Engine/Plugins/PCG/Source/PCGCompute/Internal/PCGGrassMapUnpackerCS.h
+++ b/Engine/Plugins/PCG/Source/PCGCompute/Internal/PCGGrassMapUnpackerCS.h
@@ -29,7 +29,7 @@ public:
 		SHADER_PARAMETER(FUintVector2, InNumTiles)
 		SHADER_PARAMETER(FIntPoint, InTileGridOrigin)
 		SHADER_PARAMETER(uint32, InLandscapeComponentResolution)
-		SHADER_PARAMETER(int32, InNumExportedGrassTypes)
+		SHADER_PARAMETER(int32, InNumExportedGrassMaps)
 		SHADER_PARAMETER(float, InLandscapeScaleZ)
 		SHADER_PARAMETER(float, InLandscapeRootZ)
 		SHADER_PARAMETER(uint32, InOutputHeight)
diff --git a/Engine/Source/Editor/DetailCustomizations/Private/DetailCustomizations.cpp b/Engine/Source/Editor/DetailCustomizations/Private/DetailCustomizations.cpp
--- a/Engine/Source/Editor/DetailCustomizations/Private/DetailCustomizations.cpp
+++ b/Engine/Source/Editor/DetailCustomizations/Private/DetailCustomizations.cpp
@@ -88,7 +88,6 @@
 #include "LocalLightComponentDetails.h"
 #include "MarginCustomization.h"
 #include "MaterialAttributePropertyDetails.h"
-#include "MaterialExpressionLandscapeGrassCustomization.h"
 #include "MaterialExpressionTextureBaseDetails.h"
 #include "MaterialInstanceDynamicDetails.h"
 #include "MaterialParameterCollectionDetails.h"
@@ -380,7 +379,6 @@ void FDetailCustomizationsModule::RegisterPropertyTypeCustomizations()
 	RegisterCustomPropertyTypeLayout("PerQualityLevelInt", FOnGetPropertyTypeCustomizationInstance::CreateStatic(&FPerQualityLevelPropertyCustomization<FPerQualityLevelInt>::MakeInstance));
 	RegisterCustomPropertyTypeLayout("PerQualityLevelFloat", FOnGetPropertyTypeCustomizationInstance::CreateStatic(&FPerQualityLevelPropertyCustomization<FPerQualityLevelFloat>::MakeInstance));
 	RegisterCustomPropertyTypeLayout("SkeletalMeshOptimizationSettings", FOnGetPropertyTypeCustomizationInstance::CreateStatic(&FSkeletalMeshReductionSettingsDetails::MakeInstance));
-	RegisterCustomPropertyTypeLayout("GrassInput", FOnGetPropertyTypeCustomizationInstance::CreateStatic(&FMaterialExpressionLandscapeGrassInputCustomization::MakeInstance));
 	RegisterCustomPropertyTypeLayout("ComponentReference", FOnGetPropertyTypeCustomizationInstance::CreateStatic(&FComponentReferenceCustomization::MakeInstance));
 	RegisterCustomPropertyTypeLayout("SoftComponentReference", FOnGetPropertyTypeCustomizationInstance::CreateStatic(&FComponentReferenceCustomization::MakeInstance));
 	RegisterCustomPropertyTypeLayout("EMaterialShadingModel", FOnGetPropertyTypeCustomizationInstance::CreateStatic(&FMaterialShadingModelCustomization::MakeInstance));
diff --git a/Engine/Source/Editor/DetailCustomizations/Private/MaterialExpressionLandscapeGrassCustomization.cpp b/Engine/Source/Editor/DetailCustomizations/Private/MaterialExpressionLandscapeGrassCustomization.cpp
deleted file mode 100644
diff --git a/Engine/Source/Editor/DetailCustomizations/Private/MaterialExpressionLandscapeGrassCustomization.h b/Engine/Source/Editor/DetailCustomizations/Private/MaterialExpressionLandscapeGrassCustomization.h
deleted file mode 100644
diff --git a/Engine/Source/Runtime/Core/Private/Misc/MapErrors.cpp b/Engine/Source/Runtime/Core/Private/Misc/MapErrors.cpp
--- a/Engine/Source/Runtime/Core/Private/Misc/MapErrors.cpp
+++ b/Engine/Source/Runtime/Core/Private/Misc/MapErrors.cpp
@@ -42,6 +42,7 @@ FLazyName FMapErrors::LandscapeComponentPostLoad_Warning(TEXT("LandscapeComponen
 FLazyName FMapErrors::LandscapeLazyObjectPtrDeprecation_Warning(TEXT("LandscapeLazyObjectPtrDeprecation_Warning"));
 FLazyName FMapErrors::LandscapeProxyFixupSharedData_Warning(TEXT("LandscapeProxyFixupSharedData_Warning"));
 FLazyName FMapErrors::LandscapeProxyInvalidPropertyOverride_Warning(TEXT("LandscapeProxyInvalidPropertyOverride_Warning"));
+FLazyName FMapErrors::LandscapeProxyInvalidGrassTypesOverride_Warning(TEXT("LandscapeProxyInvalidGrassTypesOverride_Warning"));
 FLazyName FMapErrors::DuplicateLevelInfo(TEXT("DuplicateLevelInfo"));
 FLazyName FMapErrors::NoKillZ(TEXT("NoKillZ"));
 FLazyName FMapErrors::LightComponentNull(TEXT("LightComponentNull"));
diff --git a/Engine/Source/Runtime/Core/Public/Misc/MapErrors.h b/Engine/Source/Runtime/Core/Public/Misc/MapErrors.h
--- a/Engine/Source/Runtime/Core/Public/Misc/MapErrors.h
+++ b/Engine/Source/Runtime/Core/Public/Misc/MapErrors.h
@@ -138,6 +138,9 @@ struct FMapErrors
 	/** The landscape proxy property has been overridden but is invalid.  */
 	static CORE_API FLazyName LandscapeProxyInvalidPropertyOverride_Warning;
 
+	/** At least one of the landscape grass type overrides is invalid.  */
+	static CORE_API FLazyName LandscapeProxyInvalidGrassTypesOverride_Warning;
+
     /**  Level */
 
 	/** Duplicate level info: Two WorldInfos somehow exist... */
diff --git a/Engine/Source/Runtime/Core/Public/UObject/FortniteMainBranchObjectVersions.inl b/Engine/Source/Runtime/Core/Public/UObject/FortniteMainBranchObjectVersions.inl
--- a/Engine/Source/Runtime/Core/Public/UObject/FortniteMainBranchObjectVersions.inl
+++ b/Engine/Source/Runtime/Core/Public/UObject/FortniteMainBranchObjectVersions.inl
@@ -769,5 +769,8 @@ FORTNITE_MAIN_VERSION(PCG_PropertyBagHierarchySupport_Outdated, 251)
 // PCG Graph's Instanced Property Bags now support Hierarchy Editor, post loading existing data into new format
 FORTNITE_MAIN_VERSION(PCG_PropertyBagHierarchySupport, 252)
 
+// Introduce per-proxy overridable landscape grass 
+FORTNITE_MAIN_VERSION(OverridableLandscapeGrassTypes, 253)
+
 // -----<new versions can be added above this line>-------------------------------------------------
 #undef FORTNITE_MAIN_VERSION
diff --git a/Engine/Source/Runtime/Engine/Private/Materials/MaterialCachedData.cpp b/Engine/Source/Runtime/Engine/Private/Materials/MaterialCachedData.cpp
--- a/Engine/Source/Runtime/Engine/Private/Materials/MaterialCachedData.cpp
+++ b/Engine/Source/Runtime/Engine/Private/Materials/MaterialCachedData.cpp
@@ -113,7 +113,10 @@ void FMaterialCachedExpressionData::AddReferencedObjects(FReferenceCollector& Co
  	Collector.AddStableReferenceArray(&EditorOnlyData->EditorReferencedDefaultTextures);
 #endif
 	Collector.AddStableReferenceArray(&ReferencedTextureCollections);
-	Collector.AddStableReferenceArray(&GrassTypes);
+	for (TPair<FName, TObjectPtr<ULandscapeGrassType>>& ItPair : NamedGrassTypes)
+	{
+		Collector.AddStableReference(&ItPair.Value);
+	}
 	Collector.AddStableReferenceArray(&MaterialCacheTags);
 	Collector.AddStableReferenceArray(&MaterialLayers.Layers);
 	Collector.AddStableReferenceArray(&MaterialLayers.Blends);
@@ -627,9 +630,9 @@ void FMaterialCachedExpressionData::UpdateForExpressions(const FMaterialCachedEx
 		}
 		else if (UMaterialExpressionLandscapeGrassOutput* ExpressionGrassOutput = Cast<UMaterialExpressionLandscapeGrassOutput>(Expression))
 		{
-			for (const auto& Type : ExpressionGrassOutput->GrassTypes)
+			for (const FGrassInput& Input : ExpressionGrassOutput->GrassTypes)
 			{
-				GrassTypes.Add(Type.GrassType);
+				NamedGrassTypes.Add(Input.Name, Input.GrassType);
 			}
 		}
 		else if (UMaterialExpressionQualitySwitch* QualitySwitchNode = Cast<UMaterialExpressionQualitySwitch>(Expression))
@@ -1017,6 +1020,28 @@ void FMaterialCachedExpressionData::PostSerialize(const FArchive& Ar)
 		{
 			PropertyConnectedMask = uint64(PropertyConnectedBitmask_DEPRECATED);
 		}
+
+#if WITH_EDITORONLY_DATA
+PRAGMA_DISABLE_DEPRECATION_WARNINGS
+		// Handle loading cached expression data where grass types were still serialized as an array. This could happen when loading a material with bIsCookedForEditor,
+		//  in which case we wouldn't have access to the grass name associated with the ULandscapeGrassType asset, since it's part of UMaterialExpressionLandscapeGrassOutput and is only 
+		//  updated when UpdateCachedExpressionData is called (which is not the case for bIsCookedForEditor stuff)
+		if (!GrassTypes_DEPRECATED.IsEmpty())
+		{
+			check(NamedGrassTypes.IsEmpty()); // This situation should only happen when loading cooked data that doesn't have the new, named, grass types yet
+			for (TObjectPtr<ULandscapeGrassType> GrassType : GrassTypes_DEPRECATED)
+			{
+				if (GrassType != nullptr)
+				{
+					// Just use the grass type's fallback name as a fallback :
+					check(!NamedGrassTypes.Contains(GrassType->GetGrassNameFallback())); // Should be unique as the old GrassTypes were forced unique by the UX code
+					NamedGrassTypes.Add(GrassType->GetGrassNameFallback(), GrassType);
+				}
+			}
+			GrassTypes_DEPRECATED.Empty();
+		}
+PRAGMA_ENABLE_DEPRECATION_WARNINGS
+#endif // WITH_EDITORONLY_DATA
 	}
 
 #if WITH_EDITORONLY_DATA
diff --git a/Engine/Source/Runtime/Engine/Public/MaterialCachedData.h b/Engine/Source/Runtime/Engine/Public/MaterialCachedData.h
--- a/Engine/Source/Runtime/Engine/Public/MaterialCachedData.h
+++ b/Engine/Source/Runtime/Engine/Public/MaterialCachedData.h
@@ -369,9 +369,14 @@ struct FMaterialCachedExpressionData
 	UPROPERTY()
 	TArray<FMaterialParameterCollectionInfo> ParameterCollectionInfos;
 
+#if WITH_EDITORONLY_DATA
 	UPROPERTY()
-	TArray<TObjectPtr<ULandscapeGrassType>> GrassTypes;
-	
+	TArray<TObjectPtr<ULandscapeGrassType>> GrassTypes_DEPRECATED;
+#endif // WITH_EDITORONLY_DATA	
+
+	UPROPERTY()
+	TMap<FName, TObjectPtr<ULandscapeGrassType>> NamedGrassTypes;
+
 	UPROPERTY()
 	TArray<TObjectPtr<UMaterialCacheVirtualTextureTag>> MaterialCacheTags;
 
diff --git a/Engine/Source/Runtime/Landscape/Classes/LandscapeComponent.h b/Engine/Source/Runtime/Landscape/Classes/LandscapeComponent.h
--- a/Engine/Source/Runtime/Landscape/Classes/LandscapeComponent.h
+++ b/Engine/Source/Runtime/Landscape/Classes/LandscapeComponent.h
@@ -9,8 +9,9 @@
 #include "Engine/TextureStreamingTypes.h"
 #include "Components/PrimitiveComponent.h"
 #include "UObject/PerPlatformProperties.h"
-#include "LandscapePhysicalMaterial.h"
+#include "LandscapeGrassType.h"
 #include "LandscapeInfo.h"
+#include "LandscapePhysicalMaterial.h"
 #include "LandscapeWeightmapUsage.h"
 #include "Containers/ArrayView.h"
 #include "Engine/StreamableRenderAsset.h"
@@ -30,8 +31,8 @@ class FMeshMapBuildData;
 class FPrimitiveSceneProxy;
 class ITargetPlatform;
 class ULandscapeComponent;
-class ULandscapeGrassType;
 class ULandscapeHeightfieldCollisionComponent;
+class ULandscapeHeightmapTextureEdgeFixup;
 class ULandscapeInfo;
 class ULandscapeLayerInfoObject;
 class ULightComponent;
@@ -40,12 +41,11 @@ class UMaterialInterface;
 class UTexture2D;
 struct FConvexVolume;
 struct FEngineShowFlags;
-struct FLandscapeTextureDataInfo;
-struct FStaticLightingPrimitiveInfo;
 struct FLandscapeEditDataInterface;
-struct FLandscapeMobileRenderData;
-class ULandscapeHeightmapTextureEdgeFixup;
 struct FLandscapeGroup;
+struct FLandscapeMobileRenderData;
+struct FLandscapeTextureDataInfo;
+struct FStaticLightingPrimitiveInfo;
 enum class ELandscapeTextureUsage : uint8;
 enum class ELandscapeTextureType : uint8;
 
@@ -232,7 +232,7 @@ struct FLandscapeComponentGrassData
 	//  therefore empty) or not known yet (== UnknownNumElements)
 	int32 NumElements = UnknownNumElements;
 	// Serialized in one block to prevent Slack waste
-	TMap<TObjectPtr<ULandscapeGrassType>, int32> WeightOffsets;
+	TMap<FName, int32> WeightOffsets;
 	TArray<uint8> HeightWeightData;
 
 	// Note: We need to explicitly disable warnings on these constructors/operators for clang to be happy with deprecated variables
@@ -252,18 +252,22 @@ struct FLandscapeComponentGrassData
 	// Returns whether the data is completely empty (e.g. all-zero weightmap data). Returns false if the data just wasn't computed yet :
 	bool HasData() const;
 
+#if WITH_EDITOR
+	// Initialize function meant for deserializing from old data : the new format doesn't take ULandscapeGrassType* but FName :
 	void InitializeFrom(const TArray<uint16>& HeightData, const TMap<ULandscapeGrassType*, TArray<uint8>>& WeightData);
-	void InitializeFrom(IBuffer2DView<uint16>* HeightData, TMap<ULandscapeGrassType*, IBuffer2DView<uint8>*>& WeightData, bool bStripEmptyWeights);
+	void FixupFallbackGrassNames(const TMap<FName, TObjectPtr<ULandscapeGrassType>>& InNamedGrassTypes);
+#endif // WITH_EDITOR
+	void InitializeFrom(IBuffer2DView<uint16>* HeightData, TMap<FName, IBuffer2DView<uint8>*>& WeightData, bool bStripEmptyWeights);
 
 	bool HasWeightData() const;
-	TArrayView<uint8> GetWeightData(const ULandscapeGrassType* GrassType);
-	bool Contains(ULandscapeGrassType* GrassType) const;
+	TArrayView<uint8> GetWeightData(const FName& GrassName);
+	bool Contains(const FName& GrassName) const;
 	TArrayView<uint16> GetHeightData();
 
 	SIZE_T GetAllocatedSize() const;
 
 	// Check whether we can discard any data not needed with current scalability settings
-	void ConditionalDiscardDataOnLoad();
+	void ConditionalDiscardDataOnLoad(const TMap<FName, TObjectPtr<ULandscapeGrassType>>& InNamedGrassTypes);
 
 	friend FArchive& operator<<(FArchive& Ar, FLandscapeComponentGrassData& Data);
 };
@@ -563,8 +567,11 @@ private:
 	TArray<FLandscapePerLODMaterialOverride> PerLODOverrideMaterials;
 
 #if WITH_EDITORONLY_DATA
-	/** The value of the landscape material AllStateCRC the last time the GrassTypes array was updated from it */
-	uint32 LastLandscapeMaterialAllStateCRCWhenGrassTypesBuilt = 0;
+	/** A CRC of everything that can lead NamedGrassTypes to be changed (e.g. the material changing, or grass type overrides changing) */
+	uint32 GrassTypesCRC = 0;
+
+	UPROPERTY()
+	TArray<TObjectPtr<ULandscapeGrassType>> GrassTypes_DEPRECATED;
 #endif // WITH_EDITORONLY_DATA
 
 	/** Cached list of grass types supported by the component's material.
@@ -572,7 +579,7 @@ private:
 	* on the cooked material.
 	* Call UpdateGrassTypes() to ensure this array is up to date */
 	UPROPERTY()
-	TArray<TObjectPtr<ULandscapeGrassType>> GrassTypes;
+	TMap<FName, TObjectPtr<ULandscapeGrassType>> NamedGrassTypes;
 
 public:
 	// Non-serialized runtime cache of values derived from the assigned grass types.
@@ -584,7 +591,7 @@ public:
 		double MaxInstanceDiscardDistance = DBL_MAX;
 	};
 	FGrassTypeSummary GrassTypeSummary;
-	inline bool IsGrassTypeSummaryValid() { return GrassTypeSummary.bInvalid; }
+	inline bool IsGrassTypeSummaryValid() { return !GrassTypeSummary.bInvalid; }
 
 	/** Invalidate the grass type summary.  Call whenever grass types are changed to indicate that the summary values are out of date. */
 	LANDSCAPE_API void InvalidateGrassTypeSummary();
@@ -746,8 +753,6 @@ public:
 	static void DeclareConstructClasses(TArray<FTopLevelAssetPath>& OutConstructClasses, const UClass* SpecificSubclass);
 #endif
 
-	static void AddReferencedObjects(UObject* InThis, FReferenceCollector& Collector);
-
 #if WITH_EDITOR
 	virtual void BeginCacheForCookedPlatformData(const ITargetPlatform* TargetPlatform) override;
 	virtual void PreEditUndo() override;
@@ -892,19 +897,22 @@ public:
 	LANDSCAPE_API ULandscapeInfo* GetLandscapeInfo() const;
 
 	UFUNCTION(BlueprintCallable, Category = "Landscape|Grass", meta=(DisplayName = "GetGrassTypes", Tooltip = "Returns the grass types used by the landscape material on this component"))
-	const TArray<ULandscapeGrassType*>& GetGrassTypesBP() const { return ObjectPtrDecay(GrassTypes); }
+	TArray<ULandscapeGrassType*> GetGrassTypesBP() const;
 
 	/** Returns the array of grass types used by the landscape material. Call UpdateGrassTypes first to ensure this array is up to date. */
-	const TArray<TObjectPtr<ULandscapeGrassType>>& GetGrassTypes() const { return GrassTypes; }
+	UE_DEPRECATED(5.8, "Use GetNamedGrassTypes")
+	TArray<TObjectPtr<ULandscapeGrassType>> GetGrassTypes() const;
+
+	const TMap<FName, TObjectPtr<ULandscapeGrassType>>& GetNamedGrassTypes() const { return NamedGrassTypes; }
 
 	/** Temporarily sets the grass type for this component. Any call to UpdateGrassTypes may override what has been set using this method. */
-	void SetGrassTypes(const TArray<TObjectPtr<ULandscapeGrassType>>& InGrassTypes)
+	void SetNamedGrassTypes(const TMap<FName, TObjectPtr<ULandscapeGrassType>>& InNamedGrassTypes)
 	{
-		GrassTypes = InGrassTypes;
+		NamedGrassTypes = InNamedGrassTypes;
 		InvalidateGrassTypeSummary();
 	}
 	
-	bool MaterialHasGrass() const { return !GetGrassTypes().IsEmpty(); }
+	bool MaterialHasGrass() const { return !GetNamedGrassTypes().IsEmpty(); }
 
 	double GetGrassTypesMaxDiscardDistance() const { return GrassTypeSummary.MaxInstanceDiscardDistance; }
 	void SetGrassTypesMaxDiscardDistance(const float InGrassTypesMaxDiscardDistance) { GrassTypeSummary.MaxInstanceDiscardDistance = InGrassTypesMaxDiscardDistance; GrassTypeSummary.bInvalid = false; }
diff --git a/Engine/Source/Runtime/Landscape/Classes/LandscapeGrassType.h b/Engine/Source/Runtime/Landscape/Classes/LandscapeGrassType.h
--- a/Engine/Source/Runtime/Landscape/Classes/LandscapeGrassType.h
+++ b/Engine/Source/Runtime/Landscape/Classes/LandscapeGrassType.h
@@ -208,8 +208,17 @@ class ULandscapeGrassType : public UObject
 	//~ End UObject Interface
 
 	uint32 ComputeStateHash();
-};
-
 
+	/** 
+	 * Return a FName that corresponds to this grass type. It used to be that grass types were unique for spawning grass on a given component. Nowadays, 
+	 *  we use named "slots" for each grass type (the GrassName being the slot's name in the landscape material). This function here returns a fallback
+	 *  name for when serializing old data that was grass type-based only. 
+	*/
+	inline FName GetGrassNameFallback() const
+	{
+		// Simply use the UObject's name as the fallback since it's unique on a given landscape component (because grass types were unique on a given component already)
+		return GetFName();
+	}
+};
 
 #undef UE_API
diff --git a/Engine/Source/Runtime/Landscape/Classes/LandscapeProxy.h b/Engine/Source/Runtime/Landscape/Classes/LandscapeProxy.h
--- a/Engine/Source/Runtime/Landscape/Classes/LandscapeProxy.h
+++ b/Engine/Source/Runtime/Landscape/Classes/LandscapeProxy.h
@@ -266,7 +266,10 @@ struct FCachedLandscapeFoliage
 	struct FGrassCompKey
 	{
 		TWeakObjectPtr<ULandscapeComponent> BasedOn;
+		// We store both the GrassName and GrassType in the key, since they're not entirely synonymous anymore : 
+		//  A given grass type can be used by 2 different grass names now. They will be treated as 2 separate entries in the foliage cache
 		TWeakObjectPtr<ULandscapeGrassType> GrassType;
+		FName GrassName;
 		int32 SqrtSubsections;
 		int32 CachedMaxInstancesPerComponent;
 		int32 SubsectionX;
@@ -292,13 +295,14 @@ struct FCachedLandscapeFoliage
 				SubsectionY == Other.SubsectionY &&
 				BasedOn == Other.BasedOn &&
 				GrassType == Other.GrassType &&
+				GrassName == Other.GrassName &&
 				NumVarieties == Other.NumVarieties &&
 				VarietyIndex == Other.VarietyIndex;
 		}
 
 		friend uint32 GetTypeHash(const FGrassCompKey& Key)
 		{
-			return GetTypeHash(Key.BasedOn) ^ GetTypeHash(Key.GrassType) ^ Key.SqrtSubsections ^ Key.CachedMaxInstancesPerComponent ^ (Key.SubsectionX << 16) ^ (Key.SubsectionY << 24) ^ (Key.NumVarieties << 3) ^ (Key.VarietyIndex << 13);
+			return GetTypeHash(Key.BasedOn) ^ GetTypeHash(Key.GrassType)  ^ GetTypeHash(Key.GrassName) ^ Key.SqrtSubsections ^ Key.CachedMaxInstancesPerComponent ^ (Key.SubsectionX << 16) ^ (Key.SubsectionY << 24) ^ (Key.NumVarieties << 3) ^ (Key.VarietyIndex << 13);
 		}
 
 	};
@@ -804,6 +808,10 @@ PRAGMA_ENABLE_DEPRECATION_WARNINGS
 	UPROPERTY(EditAnywhere, Category=Lightmass, meta = (LandscapeInherited))
 	FLightmassPrimitiveSettings LightmassSettings;
 
+	// TODO [jonathan.bard] Hidden for now because we want a UX solution to help the user choose appropriate grass names from the assigned materials. Add this when ready : EditAnywhere, Category=Grass
+	UPROPERTY(meta = (LandscapeOverridable))
+	TMap<FName, TObjectPtr<ULandscapeGrassType>> GrassTypesOverrides;
+
 	// Landscape LOD to use for collision tests. Higher numbers use less memory and process faster, but are much less accurate
 	UPROPERTY(EditAnywhere, Category=Collision, meta = (LandscapeOverridable))
 	int32 CollisionMipLevel;
@@ -1183,10 +1191,13 @@ PRAGMA_ENABLE_DEPRECATION_WARNINGS
 	FDelegateHandle FeatureLevelChangedDelegateHandle;
 
 	FGuid GetNaniteContentId() const;
+
+	void ValidateGrassTypeOverrides();
 #endif // WITH_EDITOR
 
 	//~ Begin UObject Interface.
 #if WITH_EDITOR
+	virtual void CheckForErrors() override;
 	virtual void BeginCacheForCookedPlatformData(const ITargetPlatform* TargetPlatform) override;
 	virtual void OnCookEvent(UE::Cook::ECookEvent CookEvent, UE::Cook::FCookEventContext& CookContext) override;
 #endif // WITH_EDITOR
@@ -1276,6 +1287,11 @@ PRAGMA_ENABLE_DEPRECATION_WARNINGS
 	 */
 	LANDSCAPE_API void RetrieveAllLandscapeMaterials(TSet<UMaterialInterface*>& OutMaterials) const;
 
+	LANDSCAPE_API TMap<FName, TObjectPtr<ULandscapeGrassType>> GetGrassTypeOverrides() const
+	{
+		return GrassTypesOverrides;
+	}
+
 #if WITH_EDITOR
 	/* Serialize all hashes/guids that record the current state of this proxy */
 	void SerializeStateHashes(FArchive& Ar);
diff --git a/Engine/Source/Runtime/Landscape/Internal/LandscapeGrassWeightExporter.h b/Engine/Source/Runtime/Landscape/Internal/LandscapeGrassWeightExporter.h
--- a/Engine/Source/Runtime/Landscape/Internal/LandscapeGrassWeightExporter.h
+++ b/Engine/Source/Runtime/Landscape/Internal/LandscapeGrassWeightExporter.h
@@ -21,12 +21,10 @@ namespace UE::Landscape
 	LANDSCAPE_API bool CanRenderGrassMap(ULandscapeComponent* Component);
 	UE_DEPRECATED(5.8, "Renamed to SupportsRuntimeGrassMapGeneration (LandscapeUtils.h)")
 	LANDSCAPE_API bool IsRuntimeGrassMapGenerationSupported();
-} // UE::Landscape
+} // namespace UE::Landscape
 
 namespace UE::Landscape::Grass
 {
-	void AddGrassWeightShaderTypes(FMaterialShaderTypes& InOutShaderTypes);
-
 	enum class EGrassWeightExporterFlags
 	{
 		None = 0,
@@ -95,18 +93,17 @@ public:
 	{
 		InvalidComponentKey, 
 		InvalidFlags,
-		InvalidGrassType,
+		InvalidGrassName,
 	};
 	/**
 	 * Return the location of the grass data for the requested component in the output texture
 	 * @param InComponentKey Key of the component being requested
-	 * @param InGrassType 
+	 * @param InGrassName Name of the grass map being requested 
 	 * 
 	 * @return If the component / grass data is found, a FIntRect that contains the pixels containing the requested grass data in the generated texture, along with the channel index. 
 	 *  Returns an error otherwise.
 	 */
-	// TODO [jonathan.bard] : replace grass type with FName
-	LANDSCAPE_API TValueOrError<TPair<FIntRect /*OutputRect*/, uint8 /*ChannelIndex*/>, EGrassmapTextureInfoError> GetTextureInfoForGrass(const FIntPoint& InComponentKey, ULandscapeGrassType* InGrassType);
+	LANDSCAPE_API TValueOrError<TPair<FIntRect /*OutputRect*/, uint8 /*ChannelIndex*/>, EGrassmapTextureInfoError> GetTextureInfoForGrass(const FIntPoint& InComponentKey, const FName& InGrassName);
 
 private:
 	/** Creates a texture and renders the component to it, and then and triggers a readback of the texture. */
@@ -116,7 +113,7 @@ private:
 	struct FComponentInfo
 	{
 		FComponentInfo(ULandscapeComponent* InComponent, UE::Landscape::Grass::EGrassWeightExporterFlags InFlags,
-			TConstArrayView<TObjectPtr<ULandscapeGrassType>> InRequestedGrassTypes, TConstArrayView<int32> InRequestedHeightMips);
+			TConstArrayView<FName> InRequestedGrassNames, TConstArrayView<int32> InRequestedHeightMips);
 
 		int32 GetNumPasses() const
 		{
@@ -146,8 +143,7 @@ private:
 
 		// Identifier for this component (use its key because it's a render-thread object so we don't want to access TObjectPtr there)
 		FIntPoint ComponentKey = FIntPoint(MAX_int32, MAX_int32);
-		// TODO [jonathan.bard] : this is a render thread object, this should not have a TObjectPtr (replace with FName)
-		TMap<TObjectPtr<ULandscapeGrassType>, FPassChannelIndex> PerGrassTypePassChannelIndices;
+		TMap<FName, FPassChannelIndex> PerGrassNamePassChannelIndices;
 		TArray<FIntVector4> PerPassPerChannelGrassWeightIndices;
 		FVector2D ViewOffset = FVector2D::ZeroVector;
 		int32 PixelOffsetX = 0;
@@ -189,11 +185,11 @@ public:
 	 * @param InLandscape - Parent landscape. All requested component must share the same parent landscape
 	 * @param InRequestedComponents - List of components for which to export weight (/height) maps
 	 * @param InFlags - Flags for specializing what to export and how to read back the results
-	 * @param InRequestedGrassTypes - (Optional) List of grass types requested for readback. If empty, all grass types are exported
+	 * @param InRequestedGrassNames - (Optional) List of grass map names requested for readback. If empty, all grass maps are exported
 	 * @param InRequestedHeightMips - (Optional) List of mip levels requested for reading back heightmaps
 	 */
 	LANDSCAPE_API FLandscapeGrassWeightExporter(ALandscape* InLandscape, TArrayView<ULandscapeComponent* const> InRequestedComponents, UE::Landscape::Grass::EGrassWeightExporterFlags InFlags,
-		TConstArrayView<TObjectPtr<ULandscapeGrassType>> InRequestedGrassTypes = {}, TConstArrayView<int32> InRequestedHeightMips = {});
+		TConstArrayView<FName> InRequestedGrassNames = {}, TConstArrayView<int32> InRequestedHeightMips = {});
 
 	// If using the async readback path, check its status and update if needed. Return true when the AsyncReadbackResults are available.
 	// You must call this periodically, or the async readback may not complete.
diff --git a/Engine/Source/Runtime/Landscape/Private/Landscape.cpp b/Engine/Source/Runtime/Landscape/Private/Landscape.cpp
--- a/Engine/Source/Runtime/Landscape/Private/Landscape.cpp
+++ b/Engine/Source/Runtime/Landscape/Private/Landscape.cpp
@@ -9,6 +9,7 @@ Landscape.cpp: Terrain rendering
 #include "Algo/Count.h"
 #include "Algo/ForEach.h"
 #include "Algo/IsSorted.h"
+#include "Algo/NoneOf.h"
 #include "Algo/Transform.h"
 #include "Serialization/MemoryWriter.h"
 #include "Serialization/BufferArchive.h"
@@ -46,6 +47,7 @@ Landscape.cpp: Terrain rendering
 #include "DerivedDataCacheInterface.h"
 #include "Interfaces/ITargetPlatform.h"
 #include "LandscapeMeshCollisionComponent.h"
+#include "MaterialCachedData.h"
 #include "MaterialDomain.h"
 #include "Materials/Material.h"
 #include "LandscapeMaterialInstanceConstant.h"
@@ -696,6 +698,61 @@ FGuid ALandscapeProxy::GetNaniteContentId() const
 	return FGuid(Hash[0] ^ Hash[4], Hash[1], Hash[2], Hash[3]);
 }
 
+void ALandscapeProxy::CheckForErrors()
+{
+	Super::CheckForErrors();
+
+	ValidateGrassTypeOverrides();
+}
+
+void ALandscapeProxy::ValidateGrassTypeOverrides() 
+{
+	if (GrassTypesOverrides.IsEmpty())
+	{
+		return;
+	}
+
+ 	TSet<UMaterialInterface*> LandscapeMaterials;
+	RetrieveAllLandscapeMaterials(LandscapeMaterials);
+
+	TArray<FString> MissingGrassOverrides;
+	for (const auto& ItPair : GrassTypesOverrides)
+	{
+		FName GrassName = ItPair.Key;
+		if (!GrassName.IsNone())
+		{
+			// If that grass name is unknown to all of our landscape materials, report : 
+			if (Algo::NoneOf(LandscapeMaterials, [&GrassName](UMaterialInterface* InMaterial) 
+				{ 
+					return InMaterial->GetCachedExpressionData().NamedGrassTypes.Contains(GrassName);
+				}))
+			{ 
+				MissingGrassOverrides.AddUnique(GrassName.ToString());
+			}
+		}
+	}
+
+	if (!MissingGrassOverrides.IsEmpty())
+	{
+		TArray<FString> LandscapeMaterialStrings;
+		LandscapeMaterialStrings.Reserve(LandscapeMaterials.Num());
+		Algo::Transform(LandscapeMaterials, LandscapeMaterialStrings, [](UMaterialInterface* InMaterial) { return InMaterial->GetName(); });
+
+		FFormatNamedArguments Arguments;
+		Arguments.Add(TEXT("NumMissingGrassOverrides"), MissingGrassOverrides.Num());
+		Arguments.Add(TEXT("MissingGrassOverrides"), FText::FromString(FString::Join(MissingGrassOverrides, TEXT(","))));
+		Arguments.Add(TEXT("NumLandscapeMaterials"), LandscapeMaterialStrings.Num());
+		Arguments.Add(TEXT("LandscapeMaterials"), FText::FromString(FString::Join(LandscapeMaterialStrings, TEXT(","))));
+
+		FMessageLog MessageLog("MapCheck");
+		MessageLog.Warning()->AddToken(FUObjectToken::Create(this, FText::FromString(GetActorNameOrLabel())))
+			->AddToken(FTextToken::Create(FText::Format(LOCTEXT("MapCheck_Message_InvalidLandscapeGrassTypesOverride", 
+				"None of the following grass type overrides '{MissingGrassOverrides}' can be found in the proxy landscape "
+				"{NumLandscapeMaterials}|plural(one=material,other=materials) '{LandscapeMaterials}'. {NumMissingGrassOverrides}|plural(one=This override,other=These overrides) will be ineffective."), Arguments)))
+			->AddToken(FMapErrorToken::Create(FMapErrors::LandscapeProxyInvalidGrassTypesOverride_Warning));
+	}
+}
+
 void ULandscapeComponent::CheckGenerateMobilePlatformData(bool bIsCooking, const ITargetPlatform* TargetPlatform)
 {
 	TRACE_CPUPROFILER_EVENT_SCOPE(ULandscapeComponent::CheckGenerateMobilePlatformData);
@@ -989,9 +1046,9 @@ PRAGMA_ENABLE_DEPRECATION_WARNINGS
 	}
 
 #if WITH_EDITORONLY_DATA
+PRAGMA_DISABLE_DEPRECATION_WARNINGS
 	if (Ar.IsLoading() && Ar.CustomVer(FFortniteMainBranchObjectVersion::GUID) < FFortniteMainBranchObjectVersion::NewLandscapeMaterialPerLOD)
 	{
-PRAGMA_DISABLE_DEPRECATION_WARNINGS
 		if (MobileMaterialInterface_DEPRECATED != nullptr)
 		{
 			MobileMaterialInterfaces.AddUnique(MobileMaterialInterface_DEPRECATED);
@@ -1001,8 +1058,26 @@ PRAGMA_DISABLE_DEPRECATION_WARNINGS
 		{
 			MobileCombinationMaterialInstances.AddUnique(MobileCombinationMaterialInstance_DEPRECATED);
 		}
-PRAGMA_ENABLE_DEPRECATION_WARNINGS
 	}
+
+	if (Ar.IsLoading())
+	{
+		if (!GrassTypes_DEPRECATED.IsEmpty())
+		{
+			check(NamedGrassTypes.IsEmpty()); 
+			for (TObjectPtr<ULandscapeGrassType> GrassType : GrassTypes_DEPRECATED)
+			{
+				if (GrassType != nullptr)
+				{
+					// Just use the grass type's fallback name for backwards compatibility since we do the same in FMaterialCachedExpressionData::PostSerialize, 
+					//  this will be updated by UpdateGrassTypes to support the actual remapping feature that comes with NamedGrassTypes
+					check(!NamedGrassTypes.Contains(GrassType->GetGrassNameFallback())); // Should be unique as the old GrassTypes were forced unique by the UX code
+					NamedGrassTypes.Emplace(GrassType->GetGrassNameFallback(), GrassType);
+				}
+			}
+		}
+	}
+PRAGMA_ENABLE_DEPRECATION_WARNINGS
 #endif // WITH_EDITORONLY_DATA
 
 	if (Ar.UEVer() >= VER_UE4_SERIALIZE_LANDSCAPE_GRASS_DATA)
@@ -1589,9 +1664,22 @@ PRAGMA_ENABLE_DEPRECATION_WARNINGS
 	}
 #endif // !WITH_EDITORONLY_DATA
 
+	if (GIsEditor && GetLinkerCustomVersion(FFortniteMainBranchObjectVersion::GUID) < FFortniteMainBranchObjectVersion::OverridableLandscapeGrassTypes)
+	{
+		// Make sure grass types are up-to-date (i.e. the grass names are updated based on the material). We need to make sure the material has performed its post-load :
+		if (UMaterialInterface* Material = GetLandscapeMaterial())
+		{
+			Material->ConditionalPostLoad();
+		}
+		UpdateGrassTypes();
+		// Prior to this version, GrassData contained ULandscapeGrassType* entries. These have turned into GrassName (FName), as declared by the landscape material. We can do the fixup here
+		//  since we have access to the material-declared grass types here (while FLandscapeComponentGrassData's operator<< cannot) :
+		GrassData->FixupFallbackGrassNames(GetNamedGrassTypes());
+	}
+
 #endif // WITH_EDITOR
 
-	GrassData->ConditionalDiscardDataOnLoad();
+	GrassData->ConditionalDiscardDataOnLoad(GetNamedGrassTypes());
 }
 
 void ULandscapeComponent::CollectPSOPrecacheData(const FPSOPrecacheParams& BasePrecachePSOParams, FMaterialInterfacePSOPrecacheParamsList& OutParams)
@@ -1645,14 +1733,6 @@ void ULandscapeComponent::DeclareConstructClasses(TArray<FTopLevelAssetPath>& Ou
 }
 #endif
 
-void ULandscapeComponent::AddReferencedObjects(UObject* InThis, FReferenceCollector& Collector)
-{
-	Super::AddReferencedObjects(InThis, Collector);
-
-	ThisClass* const TypedThis = Cast<ThisClass>(InThis);
-	Collector.AddReferencedObjects(TypedThis->GrassData->WeightOffsets, TypedThis);
-}
-
 #if WITH_EDITORONLY_DATA
 TArray<ALandscapeProxy*> ALandscapeProxy::LandscapeProxies;
 #endif
@@ -2113,6 +2193,25 @@ ULandscapeInfo* ULandscapeComponent::GetLandscapeInfo() const
 	return GetLandscapeProxy()->GetLandscapeInfo();
 }
 
+TArray<ULandscapeGrassType*> ULandscapeComponent::GetGrassTypesBP() const
+{
+	TArray<ULandscapeGrassType*> GrassTypes;
+	GrassTypes.Reserve(NamedGrassTypes.Num());
+	for (const TPair<FName, TObjectPtr<ULandscapeGrassType>>& ItPair : NamedGrassTypes)
+	{
+		if (ItPair.Value != nullptr)
+		{
+			GrassTypes.Add(ItPair.Value);
+		}
+	}
+	return GrassTypes;
+}
+
+TArray<TObjectPtr<ULandscapeGrassType>> ULandscapeComponent::GetGrassTypes() const
+{
+	return ObjectPtrWrap(GetGrassTypesBP());
+}
+
 void ULandscapeComponent::BeginDestroy()
 {
 	Super::BeginDestroy();
diff --git a/Engine/Source/Runtime/Landscape/Private/LandscapeEdit.cpp b/Engine/Source/Runtime/Landscape/Private/LandscapeEdit.cpp
--- a/Engine/Source/Runtime/Landscape/Private/LandscapeEdit.cpp
+++ b/Engine/Source/Runtime/Landscape/Private/LandscapeEdit.cpp
@@ -6184,6 +6184,16 @@ void ALandscapeProxy::PostEditChangeProperty(FPropertyChangedEvent& PropertyChan
 			UpdateRenderingMethod();
 		}
 	}
+	else if (MemberPropertyName == GET_MEMBER_NAME_CHECKED(ALandscapeProxy, GrassTypesOverrides))
+	{
+		for (ULandscapeComponent* Component : LandscapeComponents)
+		{
+			if (Component != nullptr)
+			{
+				Component->UpdateGrassTypes();
+			}
+		}
+	}
 
 	// Remove any null landscape components
 	LandscapeComponents.RemoveAll([](const ULandscapeComponent* Component) { return Component == nullptr; });
diff --git a/Engine/Source/Runtime/Landscape/Private/LandscapeGrass.cpp b/Engine/Source/Runtime/Landscape/Private/LandscapeGrass.cpp
--- a/Engine/Source/Runtime/Landscape/Private/LandscapeGrass.cpp
+++ b/Engine/Source/Runtime/Landscape/Private/LandscapeGrass.cpp
@@ -29,6 +29,7 @@
 #include "LandscapePrivate.h"
 #include "LandscapeProxy.h"
 #include "LandscapeInfo.h"
+#include "LandscapeUtils.h"
 #include "LandscapeUtilsPrivate.h"
 #include "LightMap.h"
 #include "Engine/MapBuildDataRegistry.h"
@@ -189,12 +190,6 @@ static FAutoConsoleVariableRef CVarGrassCreationPrioritizedMultipler(
 	GGrassCreationPrioritizedMultipler,
 	TEXT("Multiplier applied to MaxCreatePerFrame and MaxAsyncTasks when grass creation is prioritized."));
 
-static int32 GGrassUpdateAllOnRebuild = 0;
-static FAutoConsoleVariableRef CVarUpdateAllOnRebuild(
-	TEXT("grass.UpdateAllOnRebuild"),
-	GGrassUpdateAllOnRebuild,
-	TEXT(""));
-
 static TAutoConsoleVariable<int32> CVarRayTracingLandscapeGrass(
 	TEXT("r.RayTracing.Geometry.LandscapeGrass"),
 	0,
@@ -233,6 +228,7 @@ static FAutoConsoleCommand ConsoleCommandDumpLandscapeGrassData(
 		bool bDetailed = false;
 		bool bByProxy = false;
 		bool bByComponent = false;
+		bool bByGrassName = false;
 		bool bByGrassType = false;
 		bool bFull = false;
 		for (const FString& Arg : Args)
@@ -261,6 +257,10 @@ static FAutoConsoleCommand ConsoleCommandDumpLandscapeGrassData(
 			{
 				bByComponent = true;
 			}
+			if (FParse::Param(*Arg, TEXT("bygrassname")))
+			{
+				bByGrassName = true;
+			}
 			if (FParse::Param(*Arg, TEXT("bygrasstype")))
 			{
 				bByGrassType = true;
@@ -270,7 +270,7 @@ static FAutoConsoleCommand ConsoleCommandDumpLandscapeGrassData(
 				bFull = true;
 			}
 		}
-		if (!bDetailed && !bByProxy && !bByComponent && !bByGrassType && !bFull)
+		if (!bDetailed && !bByProxy && !bByComponent && !bByGrassName && !bByGrassType && !bFull)
 		{
 			// If nothing specified, assume full report :
 			bFull = true;
@@ -278,24 +278,29 @@ static FAutoConsoleCommand ConsoleCommandDumpLandscapeGrassData(
 
 		struct FSortDataEntry
 		{
-			FSortDataEntry(const UWorld* InWorld, const ALandscape* InActor, const ALandscapeProxy* InProxy, const ULandscapeComponent* InComponent, const ULandscapeGrassType* InGrassType, int32 InOffset)
+			FSortDataEntry(const UWorld* InWorld, const ALandscape* InActor, const ALandscapeProxy* InProxy, const ULandscapeComponent* InComponent, const FName& InGrassName, const ULandscapeGrassType* InGrassType, int32 InOffset, bool bInIsHeightData)
 				: World(InWorld)
 				, Actor(InActor)
 				, Proxy(InProxy)
 				, Component(InComponent)
+				, GrassName(InGrassName)
 				, GrassType(InGrassType)
 				, Offset(InOffset)
+				, bIsHeightData(bInIsHeightData)
 			{}
 			const UWorld* World = nullptr;
 			const ALandscape* Actor = nullptr;
 			const ALandscapeProxy* Proxy = nullptr;
 			const ULandscapeComponent* Component = nullptr;
-			const ULandscapeGrassType* GrassType = nullptr; // If nullptr, this entry corresponds to height data 
+			FName GrassName; // Can be NAME_None for height data 
+			const ULandscapeGrassType* GrassType = nullptr; // Can be nullptr if this entry corresponds to height data or if we have some grass data in the grass map but no grass type associated with it
 			int32 Offset = 0;
 			int32 Size = 0;
+			bool bIsHeightData = false;
 		};
 		TArray<FSortDataEntry> SortData;
 		TSet<const ULandscapeGrassType*> UniqueGrassTypes;
+		TSet<FName> UniqueGrassNames;
 
 		// Gather all data first : 
 		for (TObjectIterator<const ULandscapeComponent> It(RF_ClassDefaultObject, true, EInternalObjectFlags::Garbage); It; ++It)
@@ -312,17 +317,27 @@ static FAutoConsoleCommand ConsoleCommandDumpLandscapeGrassData(
 							int32 MinWeightOffset = MAX_int32;
 							for (auto ItPair : Component->GrassData->WeightOffsets)
 							{
-								const ULandscapeGrassType* GrassType = ItPair.Key;
-								int32 GrassTypeOffset = ItPair.Value;
-								SortData.Add(FSortDataEntry(World, Actor, Proxy, Component, GrassType, GrassTypeOffset));
-								MinWeightOffset = FMath::Min(GrassTypeOffset, MinWeightOffset);
+								const FName& GrassName = ItPair.Key;
+								const ULandscapeGrassType* GrassType = nullptr;
+								check(!GrassName.IsNone());
+								if (const TObjectPtr<ULandscapeGrassType>* FoundGrassType = Component->GetNamedGrassTypes().Find(GrassName))
+								{
+									GrassType = *FoundGrassType;
+									int32 GrassTypeOffset = ItPair.Value;
+									SortData.Add(FSortDataEntry(World, Actor, Proxy, Component, GrassName, GrassType, GrassTypeOffset, /*bInIsHeightData = */false));
+									MinWeightOffset = FMath::Min(GrassTypeOffset, MinWeightOffset);
 
-								UniqueGrassTypes.Add(GrassType);
+									UniqueGrassNames.Add(GrassName);
+									if (GrassType != nullptr)
+									{
+										UniqueGrassTypes.Add(GrassType);
+									}
+								}
 							}
 							// Height data, if any, is at the beginning so if one the weight offsets is at index 0, then there's no height data on this component at all
 							int32 HeightOffset = (MinWeightOffset == 0) ? INDEX_NONE : 0;
 							// Always add an entry for height, even if no height data was found (makes reporting easier) : 
-							SortData.Add(FSortDataEntry(World, Actor, Proxy, Component, /* InGrassType = */nullptr, HeightOffset));
+							SortData.Add(FSortDataEntry(World, Actor, Proxy, Component, /*InGrassName = */NAME_None, /* InGrassType = */nullptr, HeightOffset, /*bInIsHeightData = */true));
 						}
 					}
 				}
@@ -374,7 +389,7 @@ static FAutoConsoleCommand ConsoleCommandDumpLandscapeGrassData(
 			SortData[NumEntries - 1].Size = CurrentComponent->GrassData->HeightWeightData.Num() - CurrentOffsetInComponent;
 		}
 
-		// Sort by World, then Actor, then Proxy, then Component, then GrassType : 
+		// Sort by World, then Actor, then Proxy, then Component, then grass name : 
 		SortData.Sort([](const FSortDataEntry& InLHS, const FSortDataEntry& InRHS)
 		{
 			if (InLHS.World == InRHS.World)
@@ -385,7 +400,7 @@ static FAutoConsoleCommand ConsoleCommandDumpLandscapeGrassData(
 					{
 						if (InLHS.Component == InRHS.Component)
 						{
-							return InLHS.GrassType < InRHS.GrassType;
+							return InLHS.GrassName.LexicalLess(InRHS.GrassName);
 						}
 						else
 						{
@@ -410,7 +425,8 @@ static FAutoConsoleCommand ConsoleCommandDumpLandscapeGrassData(
 
 		auto GetFormattedObjectName = [bUseFullNames](const UObject* InObject) -> FString { return bUseFullNames ? InObject->GetFullName() : InObject->GetName(); };
 		auto GetFormattedActorName = [bUseFullNames](const AActor* InActor) -> FString { return bUseFullNames ? InActor->GetFullName() : InActor->GetActorNameOrLabel(); };
-		auto GetFormattedGrassTypeName = [bUseFullNames](const ULandscapeGrassType* InGrassType) -> FString { return (InGrassType != nullptr) ? (bUseFullNames ? InGrassType->GetFullName() : InGrassType->GetName()) : FString("(HeightData)"); };
+		auto GetFormattedGrassName = [](const FName& InGrassName) -> FString { check(!InGrassName.IsNone()); return InGrassName.ToString(); };
+		auto GetFormattedGrassTypeName = [bUseFullNames](const ULandscapeGrassType* InGrassType) -> FString { return (InGrassType != nullptr) ? (bUseFullNames ? InGrassType->GetFullName() : InGrassType->GetName()) : FString(TEXT("(null)")); };
 		auto LogWorldActorLine = [GetFormattedObjectName, GetFormattedActorName](const UWorld* InWorld, const ALandscape* InActor) { GLog->Logf(TEXT("-- World: %s - Landscape Actor: %s --"), *GetFormattedObjectName(InWorld), *GetFormattedActorName(InActor)); };
 
 		// List of lines to display in a sub-report (one element per entry):
@@ -432,16 +448,28 @@ static FAutoConsoleCommand ConsoleCommandDumpLandscapeGrassData(
 
 			if (bCSV)
 			{
-				GLog->Logf(TEXT("World,Landscape Actor,Landscape Proxy,Landscape Component,Landscape Grass Type,Size (KB)"));
+				GLog->Logf(TEXT("World,Landscape Actor,Landscape Proxy,Landscape Component,Landscape Grass Name,Landscape Grass Type,Size (KB)"));
 				for (const FSortDataEntry& Entry : SortData)
 				{
-					if ((Entry.GrassType != nullptr) || (bIncludeEmpty || (Entry.Size > 0)))
+					if (Entry.bIsHeightData)
+					{
+						FString Line = FString::Printf(TEXT("%s,%s,%s,%s,(HeightData),(HeightData),%.2f"),
+							*GetFormattedObjectName(Entry.World),
+							*GetFormattedActorName(Entry.Actor),
+							*GetFormattedActorName(Entry.Proxy),
+							*Entry.Component->GetName(),
+							Entry.Size / 1024.0f);
+						// Use the line itself as the sort key since it's already formatted the way we want: 
+						AddLineToSubReport(Line, Line);
+					}
+					else if ((Entry.GrassType != nullptr) || (bIncludeEmpty || (Entry.Size > 0)))
 					{
-						FString Line = FString::Printf(TEXT("%s,%s,%s,%s,%s,%.2f"),
+						FString Line = FString::Printf(TEXT("%s,%s,%s,%s,%s,%s,%.2f"),
 							*GetFormattedObjectName(Entry.World),
 							*GetFormattedActorName(Entry.Actor),
 							*GetFormattedActorName(Entry.Proxy),
 							*Entry.Component->GetName(),
+							*GetFormattedGrassName(Entry.GrassName),
 							*GetFormattedGrassTypeName(Entry.GrassType),
 							Entry.Size / 1024.0f);
 						// Use the line itself as the sort key since it's already formatted the way we want: 
@@ -519,7 +547,7 @@ static FAutoConsoleCommand ConsoleCommandDumpLandscapeGrassData(
 						bCurrentActorHasGrassData = true;
 						bCurrentProxyHasGrassData = true;
 						bCurrentComponentHasGrassData = true;
-						if (Entry.GrassType == nullptr)
+						if (Entry.bIsHeightData)
 						{
 							TotalHeightDataSize += Entry.Size;
 						}
@@ -552,6 +580,7 @@ static FAutoConsoleCommand ConsoleCommandDumpLandscapeGrassData(
 				GLog->Logf(TEXT("Num landscape actors with grass data: %i"), NumActorsWithGrassData);
 				GLog->Logf(TEXT("Num landscape proxies with grass data: %i"), NumProxiesWithGrassData);
 				GLog->Logf(TEXT("Num landscape components with grass data: %i"), NumComponentsWithGrassData);
+				GLog->Logf(TEXT("Num grass names used: %i"), UniqueGrassNames.Num());
 				GLog->Logf(TEXT("Num grass types used: %i"), UniqueGrassTypes.Num());
 				GLog->Logf(TEXT("Total height data size (KB): %.2f"), TotalHeightDataSize / 1024.0f);
 				GLog->Logf(TEXT("Total weight data size (KB): %.2f"), TotalWeightDataSize / 1024.0f);
@@ -570,13 +599,14 @@ static FAutoConsoleCommand ConsoleCommandDumpLandscapeGrassData(
 			LogWorldActorLine(SortData[0].World, SortData[0].Actor);
 			const ALandscape* CurrentActor = SortData[0].Actor;
 
-			auto LogCSVHeader = [bCSV]() { if (bCSV) { GLog->Logf(TEXT("Landscape Proxy,Num Grass Types,Total Height Data Size (KB), Total Weight Data Size (KB), Total Data Size (KB)")); } };
-			auto LogLine = [bCSV, AddLineToSubReport](const TCHAR* InProxyName, int32 InNumGrassTypes, int32 InTotalHeightDataSizeBytes, int32 InTotalWeightDataSizeBytes)
+			auto LogCSVHeader = [bCSV]() { if (bCSV) { GLog->Logf(TEXT("Landscape Proxy,Num Grass Names,Num Grass Types,Total Height Data Size (KB), Total Weight Data Size (KB), Total Data Size (KB)")); } };
+			auto LogLine = [bCSV, AddLineToSubReport](const TCHAR* InProxyName, int32 InNumGrassNames, int32 InNumGrassTypes, int32 InTotalHeightDataSizeBytes, int32 InTotalWeightDataSizeBytes)
 			{
 				if (bCSV)
 				{
-					AddLineToSubReport(FString(InProxyName), FString::Printf(TEXT("%s,%i,%.2f,%.2f,%.2f"),
+					AddLineToSubReport(FString(InProxyName), FString::Printf(TEXT("%s,%i,%i,%.2f,%.2f,%.2f"),
 						InProxyName,
+						InNumGrassNames,
 						InNumGrassTypes,
 						InTotalHeightDataSizeBytes / 1024.0f,
 						InTotalWeightDataSizeBytes / 1024.0f,
@@ -584,8 +614,9 @@ static FAutoConsoleCommand ConsoleCommandDumpLandscapeGrassData(
 				}
 				else
 				{
-					AddLineToSubReport(FString(InProxyName), FString::Printf(TEXT("%s: Num grass types = %i, total height data size = %.2f KB, total weight data size = %.2f KB, total data size = %.2f KB"),
+					AddLineToSubReport(FString(InProxyName), FString::Printf(TEXT("%s: Num grass names = %i, Num grass types = %i, total height data size = %.2f KB, total weight data size = %.2f KB, total data size = %.2f KB"),
 						InProxyName,
+						InNumGrassNames,
 						InNumGrassTypes,
 						InTotalHeightDataSizeBytes / 1024.0f,
 						InTotalWeightDataSizeBytes / 1024.0f,
@@ -597,6 +628,7 @@ static FAutoConsoleCommand ConsoleCommandDumpLandscapeGrassData(
 			const ALandscapeProxy* CurrentProxy = SortData[0].Proxy;
 			int32 HeightDataSizeForCurrentProxy = 0;
 			int32 WeightDataSizeForCurrentProxy = 0;
+			TSet<FName> UniqueGrassNamesForCurrentProxy;
 			TSet<const ULandscapeGrassType*> UniqueGrassTypesForCurrentProxy;
 			for (int32 Index = 0; Index < NumEntries; ++Index)
 			{
@@ -604,12 +636,13 @@ static FAutoConsoleCommand ConsoleCommandDumpLandscapeGrassData(
 				if (Entry.Proxy != CurrentProxy)
 				{
 					// The info from CurrentProxy is now complete, we can dump it :
-					LogLine(*GetFormattedActorName(CurrentProxy), UniqueGrassTypesForCurrentProxy.Num(), HeightDataSizeForCurrentProxy, WeightDataSizeForCurrentProxy);
+					LogLine(*GetFormattedActorName(CurrentProxy), UniqueGrassNamesForCurrentProxy.Num(), UniqueGrassTypesForCurrentProxy.Num(), HeightDataSizeForCurrentProxy, WeightDataSizeForCurrentProxy);
 
 					// Start gathering stats for the next proxy :
 					CurrentProxy = Entry.Proxy;
 					HeightDataSizeForCurrentProxy = 0;
 					WeightDataSizeForCurrentProxy = 0;
+					UniqueGrassNamesForCurrentProxy.Empty();
 					UniqueGrassTypesForCurrentProxy.Empty();
 				}
 
@@ -625,20 +658,24 @@ static FAutoConsoleCommand ConsoleCommandDumpLandscapeGrassData(
 				// Accumulate data about this proxy :
 				if (Entry.Size > 0)
 				{
-					if (Entry.GrassType != nullptr)
+					if (Entry.bIsHeightData)
 					{
-						WeightDataSizeForCurrentProxy += Entry.Size;
-						UniqueGrassTypesForCurrentProxy.Add(Entry.GrassType);
+						HeightDataSizeForCurrentProxy += Entry.Size;
 					}
 					else
 					{
-						HeightDataSizeForCurrentProxy += Entry.Size;
+						WeightDataSizeForCurrentProxy += Entry.Size;
+						UniqueGrassNamesForCurrentProxy.Add(Entry.GrassName);
+						if (Entry.GrassType != nullptr)
+						{
+							UniqueGrassTypesForCurrentProxy.Add(Entry.GrassType);
+						}
 					}
 				}
 			}
 
 			// And dump one more time, for the last proxy, which hasn't been reported yet :
-			LogLine(*GetFormattedActorName(CurrentProxy), UniqueGrassTypesForCurrentProxy.Num(), HeightDataSizeForCurrentProxy, WeightDataSizeForCurrentProxy);
+			LogLine(*GetFormattedActorName(CurrentProxy), UniqueGrassNamesForCurrentProxy.Num(), UniqueGrassTypesForCurrentProxy.Num(), HeightDataSizeForCurrentProxy, WeightDataSizeForCurrentProxy);
 			FinalizeSubReport();
 		}
 
@@ -653,14 +690,15 @@ static FAutoConsoleCommand ConsoleCommandDumpLandscapeGrassData(
 			LogWorldActorLine(SortData[0].World, SortData[0].Actor);
 			const ALandscape* CurrentActor = SortData[0].Actor;
 
-			auto LogCSVHeader = [bCSV]() { if (bCSV) { GLog->Logf(TEXT("Landscape Proxy,Landscape Component,Num Grass Types,Total Height Data Size (KB), Total Weight Data Size (KB), Total Data Size (KB)")); } };
-			auto LogLine = [bCSV, AddLineToSubReport](const TCHAR* InProxyName, const TCHAR* InComponentName, int32 InNumGrassTypes, int32 InTotalHeightDataSizeBytes, int32 InTotalWeightDataSizeBytes)
+			auto LogCSVHeader = [bCSV]() { if (bCSV) { GLog->Logf(TEXT("Landscape Proxy,Landscape Component,Num Grass Names,Num Grass Types,Total Height Data Size (KB), Total Weight Data Size (KB), Total Data Size (KB)")); } };
+			auto LogLine = [bCSV, AddLineToSubReport](const TCHAR* InProxyName, const TCHAR* InComponentName, int32 InNumGrassNames, int32 InNumGrassTypes, int32 InTotalHeightDataSizeBytes, int32 InTotalWeightDataSizeBytes)
 			{
 				if (bCSV)
 				{
-					AddLineToSubReport(FString(InComponentName), FString::Printf(TEXT("%s,%s,%i,%.2f,%.2f,%.2f"),
+					AddLineToSubReport(FString(InComponentName), FString::Printf(TEXT("%s,%s,%i,%i,%.2f,%.2f,%.2f"),
 						InProxyName,
 						InComponentName,
+						InNumGrassNames,
 						InNumGrassTypes,
 						InTotalHeightDataSizeBytes / 1024.0f,
 						InTotalWeightDataSizeBytes / 1024.0f,
@@ -668,9 +706,10 @@ static FAutoConsoleCommand ConsoleCommandDumpLandscapeGrassData(
 				}
 				else
 				{
-					AddLineToSubReport(FString(InComponentName), FString::Printf(TEXT("%s (%s): Num grass types = %i, total height data size = %.2f KB, total weight data size = %.2f KB, total data size = %.2f KB"),
+					AddLineToSubReport(FString(InComponentName), FString::Printf(TEXT("%s (%s): Num grass names = %i, Num grass types = %i, total height data size = %.2f KB, total weight data size = %.2f KB, total data size = %.2f KB"),
 						InProxyName,
 						InComponentName,
+						InNumGrassNames,
 						InNumGrassTypes,
 						InTotalHeightDataSizeBytes / 1024.0f,
 						InTotalWeightDataSizeBytes / 1024.0f,
@@ -680,22 +719,24 @@ static FAutoConsoleCommand ConsoleCommandDumpLandscapeGrassData(
 
 			LogCSVHeader();
 			const ULandscapeComponent* CurrentComponent = SortData[0].Component;
-			int32 NumGrassTypesForCurrentComponent = 0;
 			int32 HeightDataSizeForCurrentComponent = 0;
 			int32 WeightDataSizeForCurrentComponent = 0;
+			TSet<FName> UniqueGrassNamesForCurrentComponent;
+			TSet<const ULandscapeGrassType*> UniqueGrassTypesForCurrentComponent;
 			for (int32 Index = 0; Index < NumEntries; ++Index)
 			{
 				const FSortDataEntry& Entry = SortData[Index];
 				if (Entry.Component != CurrentComponent)
 				{
 					// The info from CurrentComponent is now complete, we can dump it :
-					LogLine(*GetFormattedActorName(CurrentComponent->GetLandscapeProxy()), *CurrentComponent->GetName(), NumGrassTypesForCurrentComponent, HeightDataSizeForCurrentComponent, WeightDataSizeForCurrentComponent);
+					LogLine(*GetFormattedActorName(CurrentComponent->GetLandscapeProxy()), *CurrentComponent->GetName(), UniqueGrassNamesForCurrentComponent.Num(), UniqueGrassTypesForCurrentComponent.Num(), HeightDataSizeForCurrentComponent, WeightDataSizeForCurrentComponent);
 
 					// Start gathering stats for the next component :
 					CurrentComponent = Entry.Component;
-					NumGrassTypesForCurrentComponent = 0;
 					HeightDataSizeForCurrentComponent = 0;
 					WeightDataSizeForCurrentComponent = 0;
+					UniqueGrassNamesForCurrentComponent.Empty();
+					UniqueGrassTypesForCurrentComponent.Empty();
 				}
 
 				if (Entry.Actor != CurrentActor)
@@ -710,33 +751,135 @@ static FAutoConsoleCommand ConsoleCommandDumpLandscapeGrassData(
 				// Accumulate data about this component :
 				if (Entry.Size > 0)
 				{
-					if (Entry.GrassType != nullptr)
+					if (Entry.bIsHeightData)
 					{
-						++NumGrassTypesForCurrentComponent;
-						WeightDataSizeForCurrentComponent += Entry.Size;
+						HeightDataSizeForCurrentComponent += Entry.Size;
 					}
 					else
 					{
-						HeightDataSizeForCurrentComponent += Entry.Size;
+						WeightDataSizeForCurrentComponent += Entry.Size;
+						UniqueGrassNamesForCurrentComponent.Add(Entry.GrassName);
+						if (Entry.GrassType != nullptr)
+						{
+							UniqueGrassTypesForCurrentComponent.Add(Entry.GrassType);
+						}
 					}
 				}
 			}
 
 			// And dump one more time, for the last component, which hasn't been reported yet :
-			LogLine(*GetFormattedActorName(CurrentComponent->GetLandscapeProxy()), *CurrentComponent->GetName(), NumGrassTypesForCurrentComponent, HeightDataSizeForCurrentComponent, WeightDataSizeForCurrentComponent);
+			LogLine(*GetFormattedActorName(CurrentComponent->GetLandscapeProxy()), *CurrentComponent->GetName(), UniqueGrassNamesForCurrentComponent.Num(), UniqueGrassTypesForCurrentComponent.Num(), HeightDataSizeForCurrentComponent, WeightDataSizeForCurrentComponent);
+			FinalizeSubReport();
+		}
+
+		if (bFull || bByGrassName)
+		{
+			// Sort by Grass Name, then Component: 
+			SortData.Sort([](const FSortDataEntry& InLHS, const FSortDataEntry& InRHS)
+				{
+					if (InLHS.bIsHeightData == InRHS.bIsHeightData)
+					{
+						if (InLHS.GrassName == InRHS.GrassName)
+						{
+							return (InLHS.Component < InRHS.Component);
+						}
+						return InLHS.GrassName.LexicalLess(InRHS.GrassName);
+					}
+					// Put the height data entries first :
+					return InLHS.bIsHeightData;
+				});
+
+			// Start a new sub-report :
+			GLog->Logf(TEXT("By grass name report:"));
+			ON_SCOPE_EXIT{ GLog->Logf(TEXT("--------")); };
+
+			auto LogCSVHeader = [bCSV]() { if (bCSV) { GLog->Logf(TEXT("Grass Name,Num Components,Total Data Size (KB)")); } };
+			auto LogLine = [bCSV, AddLineToSubReport](const TCHAR* InGrassName, int32 InNumComponents, int32 InTotalDataSizeBytes)
+				{
+					if (bCSV)
+					{
+						AddLineToSubReport(FString(InGrassName), FString::Printf(TEXT("%s,%i,%.2f"),
+							InGrassName,
+							InNumComponents,
+							InTotalDataSizeBytes / 1024.0f));
+					}
+					else
+					{
+						AddLineToSubReport(FString(InGrassName), FString::Printf(TEXT("%s: Num components = %i, total data size = %.2f KB"),
+							InGrassName,
+							InNumComponents,
+							InTotalDataSizeBytes / 1024.0f));
+					}
+				};
+
+			LogCSVHeader();
+
+			int32 Index = 0;
+			int32 HeightDataSize = 0;
+			int32 NumComponentsWithHeightData = 0;
+			// Display all the height data first : 
+			for (; (Index < NumEntries) && SortData[Index].bIsHeightData; ++Index)
+			{
+				const FSortDataEntry& Entry = SortData[Index];
+				check(Entry.bIsHeightData);
+				HeightDataSize += Entry.Size;
+				++NumComponentsWithHeightData;
+			}
+			if (NumComponentsWithHeightData > 0)
+			{
+				// Dump info for height :
+				LogLine(TEXT("(HeightData)"), NumComponentsWithHeightData, HeightDataSize);
+			}
+
+			FName CurrentGrassName = (Index < NumEntries) ? SortData[Index].GrassName : NAME_None;
+			int32 DataSizeForCurrentGrassName = 0;
+			int32 NumComponentsForCurrentGrassName = 0;
+			for (; Index < NumEntries; ++Index)
+			{
+				const FSortDataEntry& Entry = SortData[Index];
+				check(!Entry.bIsHeightData);
+				if (Entry.GrassName != CurrentGrassName)
+				{
+					// The info from CurrentGrassName is now complete, we can dump it :
+					LogLine(*GetFormattedGrassName(CurrentGrassName), NumComponentsForCurrentGrassName, DataSizeForCurrentGrassName);
+
+					// Start gathering stats for the next grass type :
+					CurrentGrassName = Entry.GrassName;
+					DataSizeForCurrentGrassName = 0;
+					NumComponentsForCurrentGrassName = 0;
+				}
+
+				// Accumulate data about this grass type :
+				if (Entry.Size > 0)
+				{
+					DataSizeForCurrentGrassName += Entry.Size;
+					++NumComponentsForCurrentGrassName;
+				}
+			}
+
+			// And dump one more time, for the last grass type, which hasn't been reported yet :
+			if (!CurrentGrassName.IsNone())
+			{
+				LogLine(*GetFormattedGrassName(CurrentGrassName), NumComponentsForCurrentGrassName, DataSizeForCurrentGrassName);
+			}
 			FinalizeSubReport();
 		}
 
 		if (bFull || bByGrassType)
 		{
-			// Sort by Grass Type, then Component: 
+			// Sort by HeightData first, then Grass Type, then Component: 
 			SortData.Sort([](const FSortDataEntry& InLHS, const FSortDataEntry& InRHS)
 			{
-				if (InLHS.GrassType == InRHS.GrassType)
+				if (InLHS.bIsHeightData == InRHS.bIsHeightData)
 				{
-					return (InLHS.Component < InRHS.Component);
+					if (InLHS.GrassType == InRHS.GrassType)
+					{
+						return (InLHS.Component < InRHS.Component);
+					}
+					return (InLHS.GrassType < InRHS.GrassType);
 				}
-				return (InLHS.GrassType < InRHS.GrassType);
+				// Put the height data entries first :
+				return InLHS.bIsHeightData;
 			});
 
 			// Start a new sub-report :
@@ -763,12 +906,31 @@ static FAutoConsoleCommand ConsoleCommandDumpLandscapeGrassData(
 			};
 
 			LogCSVHeader();
-			const ULandscapeGrassType* CurrentGrassType = SortData[0].GrassType;
+
+			int32 Index = 0;
+			int32 HeightDataSize = 0;
+			int32 NumComponentsWithHeightData = 0;
+			// Display all the height data first : 
+			for (; (Index < NumEntries) && SortData[Index].bIsHeightData; ++Index)
+			{
+				const FSortDataEntry& Entry = SortData[Index];
+				check(Entry.bIsHeightData);
+				HeightDataSize += Entry.Size;
+				++NumComponentsWithHeightData;
+			}
+			if (NumComponentsWithHeightData > 0)
+			{
+				// Dump info for height :
+				LogLine(TEXT("(HeightData)"), NumComponentsWithHeightData, HeightDataSize);
+			}
+			
+			const ULandscapeGrassType* CurrentGrassType = (Index < NumEntries) ? SortData[Index].GrassType : nullptr;
 			int32 DataSizeForCurrentGrassType = 0;
 			int32 NumComponentsForCurrentGrassType = 0;
-			for (int32 Index = 0; Index < NumEntries; ++Index)
+			for (; Index < NumEntries; ++Index)
 			{
 				const FSortDataEntry& Entry = SortData[Index];
+				check(!Entry.bIsHeightData);
 				if (Entry.GrassType != CurrentGrassType)
 				{
 					// The info from CurrentGrassType is now complete, we can dump it :
@@ -789,7 +951,10 @@ static FAutoConsoleCommand ConsoleCommandDumpLandscapeGrassData(
 			}
 
 			// And dump one more time, for the last grass type, which hasn't been reported yet :
-			LogLine(*GetFormattedGrassTypeName(CurrentGrassType), NumComponentsForCurrentGrassType, DataSizeForCurrentGrassType);
+			if (CurrentGrassType != nullptr)
+			{
+				LogLine(*GetFormattedGrassTypeName(CurrentGrassType), NumComponentsForCurrentGrassType, DataSizeForCurrentGrassType);
+			}
 			FinalizeSubReport();
 		}
 	}));
@@ -881,15 +1046,35 @@ void ULandscapeComponent::InvalidateGrassTypeSummary()
 bool ULandscapeComponent::UpdateGrassTypes(bool bForceUpdate)
 {
 	bool bChanged = false;
+
+	auto BuildNamedGrassTypes = [](UMaterialInterface* InMaterial, ALandscapeProxy* InProxy) 
+		{
+			// Build the final list of grass types for this component : start by using the material's list (we won't ever be able to support anything that is not declared in that list anyway) : 
+			TMap<FName, TObjectPtr<ULandscapeGrassType>> LocalNamedGrassTypes = InMaterial->GetMaterial()->GetCachedExpressionData().NamedGrassTypes;
+			const TMap<FName, TObjectPtr<ULandscapeGrassType>>& GrassTypeOverrides = InProxy->GetGrassTypeOverrides();
+
+			// And apply eventual overrides onto it : 
+			for (TPair<FName, TObjectPtr<ULandscapeGrassType>>& ItPair : LocalNamedGrassTypes)
+			{
+				if (const TObjectPtr<ULandscapeGrassType>* GrassTypeOverride = GrassTypeOverrides.Find(ItPair.Key))
+				{
+					ItPair.Value = *GrassTypeOverride;
+				}
+			}
+
+			return LocalNamedGrassTypes;
+		};
+
 #if WITH_EDITOR
 	if (UMaterialInterface* Material = GetLandscapeMaterial())
 	{
-		uint32 CurMaterialAllStateCRC = Material->ComputeAllStateCRC();
-		if (bForceUpdate || (LastLandscapeMaterialAllStateCRCWhenGrassTypesBuilt != CurMaterialAllStateCRC))
+		TMap<FName, TObjectPtr<ULandscapeGrassType>> LocalNamedGrassTypes = BuildNamedGrassTypes(Material, GetLandscapeProxy());
+		uint32 NewGrassTypesCRC = HashCombine(Material->ComputeAllStateCRC(), UE::Landscape::Grass::Private::ComputeNamedGrassTypesHash(LocalNamedGrassTypes));
+		if (bForceUpdate || (GrassTypesCRC != NewGrassTypesCRC))
 		{
-			LastLandscapeMaterialAllStateCRCWhenGrassTypesBuilt = CurMaterialAllStateCRC;
+			GrassTypesCRC = NewGrassTypesCRC;
 
-			SetGrassTypes(Material->GetMaterial()->GetCachedExpressionData().GrassTypes);
+			SetNamedGrassTypes(LocalNamedGrassTypes);
 			bChanged = true;
 		}
 	}
@@ -900,7 +1085,8 @@ bool ULandscapeComponent::UpdateGrassTypes(bool bForceUpdate)
 	{
 		if (UMaterialInterface* Material = GetLandscapeMaterial())
 		{
-			SetGrassTypes(Material->GetMaterial()->GetCachedExpressionData().GrassTypes);
+			TMap<FName, TObjectPtr<ULandscapeGrassType>> LocalNamedGrassTypes = BuildNamedGrassTypes(Material, GetLandscapeProxy());
+			SetNamedGrassTypes(LocalNamedGrassTypes);
 			bChanged = true;
 		}
 	}
@@ -908,8 +1094,15 @@ bool ULandscapeComponent::UpdateGrassTypes(bool bForceUpdate)
 	return bChanged;
 }
 
+namespace UE::Landscape::Grass::Private
+{
+	void AddGrassWeightShaderTypes(FMaterialShaderTypes& InOutShaderTypes);
+} // namespace UE::Landscape::Grass::Private
+
 bool ULandscapeComponent::CanRenderGrassMap() const
 {
+	using namespace UE::Landscape::Grass::Private;
+
 	// Check we can render
 	const UWorld* World = GetWorld();
 	if (GUsingNullRHI || !World || !SceneProxy)
@@ -948,7 +1141,7 @@ bool ULandscapeComponent::CanRenderGrassMap() const
 
 	// We only need the GrassWeight shaders on the fixed grid vertex factory to render grass maps : 
 	FMaterialShaderTypes ShaderTypes;
-	UE::Landscape::Grass::AddGrassWeightShaderTypes(ShaderTypes);
+	AddGrassWeightShaderTypes(ShaderTypes);
 
 	const FVertexFactoryType* LandscapeGrassVF = FindVertexFactoryType(FName(TEXT("FLandscapeFixedGridVertexFactory"), FNAME_Find));
 	if (!MaterialResource->HasShaders(ShaderTypes, LandscapeGrassVF))
@@ -1061,18 +1254,18 @@ void ALandscapeProxy::UpdateGrassTypeSummary()
 		{
 			// Update the component grass type summary
 			ComponentMaxGrassCullDistance = 0.0;
-			for (ULandscapeGrassType* GrassType : Component->GetGrassTypes())
+			TMap<FName, TObjectPtr<ULandscapeGrassType>>  ValidNamedGrassTypes = UE::Landscape::Grass::FilterValidNamedGrassTypes(Component->GetNamedGrassTypes());
+			for (const TPair<FName, TObjectPtr<ULandscapeGrassType>>& ItPair : ValidNamedGrassTypes)
 			{
-				if (GrassType != nullptr)
+				check(!ItPair.Key.IsNone() && (ItPair.Value != nullptr));
+
+				for (FGrassVariety& GrassVariety : ItPair.Value->GrassVarieties)
 				{
-					for (FGrassVariety& GrassVariety : GrassType->GrassVarieties)
+					if (GrassVariety.GrassMesh != nullptr)
 					{
-						if (GrassVariety.GrassMesh != nullptr)
-						{
-							int32 EndCullDistance = bGrassQualityLevelEnabled ? GrassVariety.EndCullDistanceQuality.GetValue(GGrassQualityLevel) : GrassVariety.EndCullDistance.GetValue();
-							ComponentMaxGrassCullDistance = FMath::Max(ComponentMaxGrassCullDistance, EndCullDistance);
-							ComponentGrassVarietyCount++;
-						}
+						int32 EndCullDistance = bGrassQualityLevelEnabled ? GrassVariety.EndCullDistanceQuality.GetValue(GGrassQualityLevel) : GrassVariety.EndCullDistance.GetValue();
+						ComponentMaxGrassCullDistance = FMath::Max(ComponentMaxGrassCullDistance, EndCullDistance);
+						ComponentGrassVarietyCount++;
 					}
 				}
 			}
@@ -1403,7 +1596,7 @@ void ULandscapeGrassType::PostEditChangeProperty(FPropertyChangedEvent& Property
 			{
 				for (ULandscapeComponent* Component : Proxy->LandscapeComponents)
 				{
-					if (Component->GetGrassTypes().Contains(this))
+					if (const FName* FoundGrassNameForType = Component->GetNamedGrassTypes().FindKey(this))
 					{
 						Component->InvalidateGrassTypeSummary();
 					}
@@ -1436,6 +1629,7 @@ uint32 ULandscapeGrassType::ComputeStateHash()
 	return Hash;
 }
 
+
 //
 // FLandscapeComponentGrassData
 //
@@ -1474,11 +1668,11 @@ bool FLandscapeComponentGrassData::HasData() const
 	return false;
 }
 
-TArrayView<uint8> FLandscapeComponentGrassData::GetWeightData(const ULandscapeGrassType* GrassType)
+TArrayView<uint8> FLandscapeComponentGrassData::GetWeightData(const FName& GrassName)
 {
 	if (!HeightWeightData.IsEmpty())
 	{
-		if (int32* OffsetPtr = WeightOffsets.Find(GrassType))
+		if (int32* OffsetPtr = WeightOffsets.Find(GrassName))
 		{
 			int32 Offset = *OffsetPtr;
 			check(Offset + NumElements <= HeightWeightData.Num());
@@ -1490,9 +1684,9 @@ TArrayView<uint8> FLandscapeComponentGrassData::GetWeightData(const ULandscapeGr
 	return TArrayView<uint8>();
 }
 
-bool FLandscapeComponentGrassData::Contains(ULandscapeGrassType* GrassType) const
+bool FLandscapeComponentGrassData::Contains(const FName& GrassName) const
 {
-	return WeightOffsets.Contains(GrassType);
+	return WeightOffsets.Contains(GrassName);
 }
 
 TArrayView<uint16> FLandscapeComponentGrassData::GetHeightData()
@@ -1506,6 +1700,8 @@ TArrayView<uint16> FLandscapeComponentGrassData::GetHeightData()
 	return MakeArrayView<uint16>((uint16*)&HeightWeightData[0], NumElements);
 }
 
+#if WITH_EDITOR
+// Initialize function meant for deserializing from old data : the new format doesn't take ULandscapeGrassType* but FName :
 void FLandscapeComponentGrassData::InitializeFrom(const TArray<uint16>& HeightData, const TMap<ULandscapeGrassType*, TArray<uint8>>& WeightData)
 {
 	WeightOffsets.Empty(WeightData.Num());
@@ -1534,7 +1730,11 @@ void FLandscapeComponentGrassData::InitializeFrom(const TArray<uint16>& HeightDa
 		
 	for (const TPair<ULandscapeGrassType*,TArray<uint8>>& Pair : WeightData)
 	{
-		WeightOffsets.Add(Pair.Key, CopyOffset);
+		// At the time this data was serialized, weights were stored per ULandscapeGrassType. Nowadays, we use FNames instead, but for old data, we 
+		//  use the ULandscapeGrassType's name as its grass name : 
+		ULandscapeGrassType* GrassType = Pair.Key;
+		FName GrassName = GrassType->GetGrassNameFallback();
+		WeightOffsets.Add(GrassName, CopyOffset);
 		check(Pair.Value.Num() == NumElements);
 		check((CopyOffset + CopySize) <= HeightWeightData.Num());
 
@@ -1543,8 +1743,45 @@ void FLandscapeComponentGrassData::InitializeFrom(const TArray<uint16>& HeightDa
 	}
 }
 
+void FLandscapeComponentGrassData::FixupFallbackGrassNames(const TMap<FName, TObjectPtr<ULandscapeGrassType>>& InNamedGrassTypes)
+{
+	TMap<FName, int32> FixedUpWeightOffsets;
+	FixedUpWeightOffsets.Reserve(WeightOffsets.Num());
+	for (const TPair<FName, int32>& ItPair : WeightOffsets)
+	{
+		const FName& GrassNameToFixup = ItPair.Key;
+		const int32 Offset = ItPair.Value;
+		// If the name is already included in the official list of available grass names, keep it as is (failsafe, in case the function is run on already fixed-up data) :
+		if (InNamedGrassTypes.Contains(GrassNameToFixup))
+		{
+			check(!FixedUpWeightOffsets.Contains(GrassNameToFixup));
+			FixedUpWeightOffsets.Add(GrassNameToFixup, Offset);
+		}
+		else
+		{
+			if (const TPair<FName, TObjectPtr<ULandscapeGrassType>>* FoundPair = Algo::FindByPredicate(InNamedGrassTypes, [GrassNameToFixup](const TPair<FName, TObjectPtr<ULandscapeGrassType>>& InPair)
+				{
+					const ULandscapeGrassType* GrassType = InPair.Value;
+					return (GrassType != nullptr) && (GrassType->GetGrassNameFallback() == GrassNameToFixup);
+				}))
+			{ 
+				// We've found the new, proper, grass name corresponding to the grass type's fallback name, let's use this
+				FixedUpWeightOffsets.Add(FoundPair->Key, Offset);
+			}
+			else
+			{
+				UE_LOG(LogLandscape, Log, TEXT("Couldn't fixup grass name %s : this grass type will be unavailable for spawning grass until generating new grass data"), *GrassNameToFixup.ToString());
+			}
+		}
+	}
+
+	// Update the weight offsets with the new names : 
+	Swap(FixedUpWeightOffsets, WeightOffsets);
+}
+#endif // WITH_EDITOR
+
 // Note that this is a relatively expensive operation (in the runtime sense), as it's making a copy of the buffers
-void FLandscapeComponentGrassData::InitializeFrom(IBuffer2DView<uint16>* HeightData, TMap<ULandscapeGrassType*, IBuffer2DView<uint8>*>& WeightData, bool bStripEmptyWeights)
+void FLandscapeComponentGrassData::InitializeFrom(IBuffer2DView<uint16>* HeightData, TMap<FName, IBuffer2DView<uint8>*>& WeightData, bool bStripEmptyWeights)
 {
 	const int32 OriginalGrassTypeCount = WeightData.Num();
 	WeightOffsets.Empty(OriginalGrassTypeCount);
@@ -1573,7 +1810,7 @@ void FLandscapeComponentGrassData::InitializeFrom(IBuffer2DView<uint16>* HeightD
 	CopyOffset += CopySize;
 	CopySize = NumElements * sizeof(uint8);
 
-	for (const TPair<ULandscapeGrassType*, IBuffer2DView<uint8>*>& Pair : WeightData)
+	for (const TPair<FName, IBuffer2DView<uint8>*>& Pair : WeightData)
 	{
 		check(Pair.Value->Num() == NumElements);
 		check((CopyOffset + CopySize) <= HeightWeightData.Num());
@@ -1684,8 +1921,26 @@ PRAGMA_ENABLE_DEPRECATION_WARNINGS
 
 		Data.InitializeFrom(DeprecatedHeightData, DeprecatedWeightData);
 	}
+	else if ((Ar.CustomVer(FFortniteMainBranchObjectVersion::GUID) < FFortniteMainBranchObjectVersion::OverridableLandscapeGrassTypes) && Ar.IsLoading())
+	{
+		Ar << Data.NumElements;
+
+		TMap<ULandscapeGrassType*, int32> DeprecatedWeightOffsets;
+		Ar << DeprecatedWeightOffsets;
+
+		// Convert from ULandscapeGrassType* to fallback FNames for the new format (this will eventually get fixed up afterwards to proper grass names, extracted from the material, in ULandscapeComponent::PostLoad)
+		for (const auto& ItPair : DeprecatedWeightOffsets)
+		{
+			if (ItPair.Key != nullptr)
+			{
+				Data.WeightOffsets.FindOrAdd(ItPair.Key->GetGrassNameFallback(), ItPair.Value);
+			}
+		}
+
+		Ar << Data.HeightWeightData;
+	}
 	else
-#endif
+#endif // WITH_EDITORONLY_DATA
 	{
 		Ar << Data.NumElements;
 		Ar << Data.WeightOffsets;
@@ -1695,18 +1950,19 @@ PRAGMA_ENABLE_DEPRECATION_WARNINGS
 	return Ar;
 }
 
-void FLandscapeComponentGrassData::ConditionalDiscardDataOnLoad()
+void FLandscapeComponentGrassData::ConditionalDiscardDataOnLoad(const TMap<FName, TObjectPtr<ULandscapeGrassType>>& InNamedGrassTypes)
 {
 	//check(HasValidData());
 	if (!GIsEditor && GGrassDiscardDataOnLoad)
 	{
 		bool bRemoved = false;
 		// Remove data for grass types which have scalability enabled
-		for (auto GrassTypeIt = WeightOffsets.CreateIterator(); GrassTypeIt; ++GrassTypeIt)
+		for (auto It = WeightOffsets.CreateIterator(); It; ++It)
 		{
-			if (!GrassTypeIt.Key() || GrassTypeIt.Key()->bEnableDensityScaling)
+			const TObjectPtr<ULandscapeGrassType>* FoundGrassType = InNamedGrassTypes.Find(It->Key);
+			if ((FoundGrassType == nullptr) || (*FoundGrassType == nullptr) || (*FoundGrassType)->bEnableDensityScaling)
 			{
-				GrassTypeIt.RemoveCurrent();
+				It.RemoveCurrent();
 				bRemoved = true;
 			}
 		}
@@ -1721,7 +1977,7 @@ void FLandscapeComponentGrassData::ConditionalDiscardDataOnLoad()
 		}
 		else if (bRemoved) 
 		{
-			TMap<ULandscapeGrassType*, int32> PreviousOffsets(ObjectPtrDecay(MoveTemp(WeightOffsets)));
+			TMap<FName, int32> PreviousOffsets(ObjectPtrDecay(MoveTemp(WeightOffsets)));
 			TArray<uint8> PreviousHeightWeightData(MoveTemp(HeightWeightData));
 			HeightWeightData.SetNumUninitialized(NumElements * sizeof(uint16) + NumElements * PreviousOffsets.Num() * sizeof(uint8), EAllowShrinking::Yes);
 
@@ -1737,7 +1993,7 @@ void FLandscapeComponentGrassData::ConditionalDiscardDataOnLoad()
 			CopyOffset += CopySize;
 			CopySize = NumElements * sizeof(uint8);
 
-			for (const TPair<ULandscapeGrassType*, int32>& Pair : PreviousOffsets)
+			for (const TPair<FName, int32>& Pair : PreviousOffsets)
 			{
 				int32 PreviousOffset = Pair.Value;
 				WeightOffsets.Add(Pair.Key, CopyOffset);
@@ -1820,10 +2076,10 @@ struct FGrassBuilderBase
 // FLandscapeComponentGrassAccess - accessor wrapper for data for one GrassType from one Component
 struct FLandscapeComponentGrassAccess
 {
-	FLandscapeComponentGrassAccess(const ULandscapeComponent* InComponent, const ULandscapeGrassType* GrassType)
+	FLandscapeComponentGrassAccess(const ULandscapeComponent* InComponent, const FName& InGrassName)
 	: GrassData(InComponent->GrassData)
 	, HeightData(InComponent->GrassData->GetHeightData())
-	, WeightData(InComponent->GrassData->GetWeightData(GrassType))
+	, WeightData(InComponent->GrassData->GetWeightData(InGrassName))
 	, Stride(InComponent->ComponentSizeQuads + 1)
 	{}
 
@@ -1893,9 +2149,10 @@ struct FAsyncGrassBuilder : public FGrassBuilderBase
 	TArray<FClusterNode> ClusterTree;
 	int32 OutOcclusionLayerNum;
 
-	FAsyncGrassBuilder(ALandscapeProxy* Landscape, ULandscapeComponent* Component, const ULandscapeGrassType* GrassType, const FGrassVariety& GrassVariety, ERHIFeatureLevel::Type FeatureLevel, UGrassInstancedStaticMeshComponent* GrassInstancedStaticMeshComponent, int32 SqrtSubsections, int32 SubX, int32 SubY, uint32 InHaltonBaseIndex, TArray<FBox>& InExcludedBoxes)
+	FAsyncGrassBuilder(ALandscapeProxy* Landscape, ULandscapeComponent* Component, const FName& GrassName, const ULandscapeGrassType* GrassType, const FGrassVariety& GrassVariety, 
+		ERHIFeatureLevel::Type FeatureLevel, UGrassInstancedStaticMeshComponent* GrassInstancedStaticMeshComponent, int32 SqrtSubsections, int32 SubX, int32 SubY, uint32 InHaltonBaseIndex, TArray<FBox>& InExcludedBoxes)
 		: FGrassBuilderBase(Landscape, Component, GrassVariety, FeatureLevel, SqrtSubsections, SubX, SubY, GrassType->bEnableDensityScaling)
-		, GrassData(Component, GrassType)
+		, GrassData(Component, GrassName)
 		, Scaling(GrassVariety.Scaling)
 		, ScaleX(GrassVariety.ScaleX)
 		, ScaleY(GrassVariety.ScaleY)
@@ -2714,290 +2971,293 @@ void ALandscapeProxy::UpdateGrass(const TArray<FVector>& Cameras, int32& InOutNu
 					Component->ChangeTag = GGrassExclusionChangeTag;
 				}
 
-				const TArray<ULandscapeGrassType*>& LandscapeGrassTypes = Component->GetGrassTypes();
-				for (ULandscapeGrassType* GrassType : LandscapeGrassTypes)
+				const TMap<FName, TObjectPtr<ULandscapeGrassType>> ValidNamedGrassTypes = UE::Landscape::Grass::FilterValidNamedGrassTypes(Component->GetNamedGrassTypes());
+				for (const TPair<FName, TObjectPtr<ULandscapeGrassType>>& ItPair : ValidNamedGrassTypes)
 				{
-					if (GrassType)
+					const FName GrassName = ItPair.Key;
+					TObjectPtr<ULandscapeGrassType> GrassType = ItPair.Value;
+					check(!GrassName.IsNone() && (GrassType != nullptr));
+
+					if (World->IsGameWorld() && !Component->GrassData->Contains(GrassName))		// TODO [chris.tchou] why do we only do this in game worlds?
 					{
-						if (World->IsGameWorld() && !Component->GrassData->Contains(GrassType))		// TODO [chris.tchou] why do we only do this in game worlds?
-						{
-							continue;
-						}
+						continue;
+					}
 
-						int32 GrassVarietyIndex = -1;
-						uint32 HaltonBaseIndex = 1;
-						for (auto& GrassVariety : GrassType->GrassVarieties)
+					int32 GrassVarietyIndex = -1;
+					uint32 HaltonBaseIndex = 1;
+					for (auto& GrassVariety : GrassType->GrassVarieties)
+					{
+						GrassVarietyIndex++;
+						int32 EndCullDistance = GrassVariety.GetEndCullDistance();
+						if (GrassVariety.GrassMesh && GrassVariety.GetDensity() > 0.0f && EndCullDistance > 0)
 						{
-							GrassVarietyIndex++;
-							int32 EndCullDistance = GrassVariety.GetEndCullDistance();
-							if (GrassVariety.GrassMesh && GrassVariety.GetDensity() > 0.0f && EndCullDistance > 0)
-							{
-								float MustHaveDistance = GuardBand * (float)EndCullDistance * CullDistanceScale;
-								float DiscardDistance = DiscardGuardBand * (float)EndCullDistance * CullDistanceScale;
+							float MustHaveDistance = GuardBand * (float)EndCullDistance * CullDistanceScale;
+							float DiscardDistance = DiscardGuardBand * (float)EndCullDistance * CullDistanceScale;
 
-								bool bUseHalton = !GrassVariety.bUseGrid;
+							bool bUseHalton = !GrassVariety.bUseGrid;
 
-								if (!bUseHalton && MinDistanceToComp > DiscardDistance)
-								{
-									continue; // we can early out if we're not using halton sequence (otherwise we have to at least update the HaltonBaseIndex)
-								}
+							if (!bUseHalton && MinDistanceToComp > DiscardDistance)
+							{
+								continue; // we can early out if we're not using halton sequence (otherwise we have to at least update the HaltonBaseIndex)
+							}
 
-								FGrassBuilderBase ForSubsectionMath(this, Component, GrassVariety, FeatureLevel);
-								int32 SqrtSubsections = 1;
-								if (ForSubsectionMath.bHaveValidData && ForSubsectionMath.SqrtMaxInstances > 0)
-								{
-									SqrtSubsections = FMath::Clamp<int32>(FMath::CeilToInt(float(ForSubsectionMath.SqrtMaxInstances) / FMath::Sqrt((float)MaxInstancesPerComponent)), 1, 16);
-								}
-								int32 MaxInstancesSub = FMath::Square(ForSubsectionMath.SqrtMaxInstances / SqrtSubsections);
+							FGrassBuilderBase ForSubsectionMath(this, Component, GrassVariety, FeatureLevel);
+							int32 SqrtSubsections = 1;
+							if (ForSubsectionMath.bHaveValidData && ForSubsectionMath.SqrtMaxInstances > 0)
+							{
+								SqrtSubsections = FMath::Clamp<int32>(FMath::CeilToInt(float(ForSubsectionMath.SqrtMaxInstances) / FMath::Sqrt((float)MaxInstancesPerComponent)), 1, 16);
+							}
+							int32 MaxInstancesSub = FMath::Square(ForSubsectionMath.SqrtMaxInstances / SqrtSubsections);
 
-								if (bUseHalton && MinDistanceToComp > DiscardDistance)
-								{
-									HaltonBaseIndex += MaxInstancesSub * SqrtSubsections * SqrtSubsections;
-									continue; // now that we have updated the halton base index, it's ok to skip the rest
-								}
+							if (bUseHalton && MinDistanceToComp > DiscardDistance)
+							{
+								HaltonBaseIndex += MaxInstancesSub * SqrtSubsections * SqrtSubsections;
+								continue; // now that we have updated the halton base index, it's ok to skip the rest
+							}
 
-								FBox LocalBox = Component->CachedLocalBox;
-								FVector LocalExtentDiv = (LocalBox.Max - LocalBox.Min) * FVector(1.0f / float(SqrtSubsections), 1.0f / float(SqrtSubsections), 1.0f);
-								for (int32 SubX = 0; SubX < SqrtSubsections; SubX++)
+							FBox LocalBox = Component->CachedLocalBox;
+							FVector LocalExtentDiv = (LocalBox.Max - LocalBox.Min) * FVector(1.0f / float(SqrtSubsections), 1.0f / float(SqrtSubsections), 1.0f);
+							for (int32 SubX = 0; SubX < SqrtSubsections; SubX++)
+							{
+								for (int32 SubY = 0; SubY < SqrtSubsections; SubY++)
 								{
-									for (int32 SubY = 0; SubY < SqrtSubsections; SubY++)
-									{
-										float MinDistanceToSubComp = MinDistanceToComp;
+									float MinDistanceToSubComp = MinDistanceToComp;
 
-										FBox WorldSubBox;
+									FBox WorldSubBox;
 
-										if ((bCullSubsections && SqrtSubsections > 1) || Component->ActiveExcludedBoxes.Num())
-										{
-											FVector BoxMin;
-											BoxMin.X = LocalBox.Min.X + LocalExtentDiv.X * float(SubX);
-											BoxMin.Y = LocalBox.Min.Y + LocalExtentDiv.Y * float(SubY);
-											BoxMin.Z = LocalBox.Min.Z;
+									if ((bCullSubsections && SqrtSubsections > 1) || Component->ActiveExcludedBoxes.Num())
+									{
+										FVector BoxMin;
+										BoxMin.X = LocalBox.Min.X + LocalExtentDiv.X * float(SubX);
+										BoxMin.Y = LocalBox.Min.Y + LocalExtentDiv.Y * float(SubY);
+										BoxMin.Z = LocalBox.Min.Z;
 
-											FVector BoxMax;
-											BoxMax.X = LocalBox.Min.X + LocalExtentDiv.X * float(SubX + 1);
-											BoxMax.Y = LocalBox.Min.Y + LocalExtentDiv.Y * float(SubY + 1);
-											BoxMax.Z = LocalBox.Max.Z;
+										FVector BoxMax;
+										BoxMax.X = LocalBox.Min.X + LocalExtentDiv.X * float(SubX + 1);
+										BoxMax.Y = LocalBox.Min.Y + LocalExtentDiv.Y * float(SubY + 1);
+										BoxMax.Z = LocalBox.Max.Z;
 
-											FBox LocalSubBox(BoxMin, BoxMax);
-											WorldSubBox = LocalSubBox.TransformBy(Component->GetComponentTransform());
+										FBox LocalSubBox(BoxMin, BoxMax);
+										WorldSubBox = LocalSubBox.TransformBy(Component->GetComponentTransform());
 
-											if (bCullSubsections && SqrtSubsections > 1)
+										if (bCullSubsections && SqrtSubsections > 1)
+										{
+											MinDistanceToSubComp = Cameras.Num() ? MAX_flt : 0.0f;
+											for (auto& Pos : Cameras)
 											{
-												MinDistanceToSubComp = Cameras.Num() ? MAX_flt : 0.0f;
-												for (auto& Pos : Cameras)
-												{
-													MinDistanceToSubComp = FMath::Min<float>(MinDistanceToSubComp, static_cast<float>(ComputeSquaredDistanceFromBoxToPoint(WorldSubBox.Min, WorldSubBox.Max, Pos)));
-												}
-												MinDistanceToSubComp = FMath::Sqrt(MinDistanceToSubComp);
+												MinDistanceToSubComp = FMath::Min<float>(MinDistanceToSubComp, static_cast<float>(ComputeSquaredDistanceFromBoxToPoint(WorldSubBox.Min, WorldSubBox.Max, Pos)));
 											}
+											MinDistanceToSubComp = FMath::Sqrt(MinDistanceToSubComp);
 										}
+									}
 
-										if (bUseHalton)
-										{
-											HaltonBaseIndex += MaxInstancesSub;  // we are going to pre-increment this for all of the continues...however we need to subtract later if we actually do this sub
-										}
+									if (bUseHalton)
+									{
+										HaltonBaseIndex += MaxInstancesSub;  // we are going to pre-increment this for all of the continues...however we need to subtract later if we actually do this sub
+									}
 
-										if (MinDistanceToSubComp > DiscardDistance)
-										{
-											continue;
-										}
+									if (MinDistanceToSubComp > DiscardDistance)
+									{
+										continue;
+									}
 
-										FCachedLandscapeFoliage::FGrassComp NewComp;
-										NewComp.Key.BasedOn = Component;
-										NewComp.Key.GrassType = GrassType;
-										NewComp.Key.SqrtSubsections = SqrtSubsections;
-										NewComp.Key.CachedMaxInstancesPerComponent = MaxInstancesPerComponent;
-										NewComp.Key.SubsectionX = SubX;
-										NewComp.Key.SubsectionY = SubY;
-										NewComp.Key.NumVarieties = GrassType->GrassVarieties.Num();
-										NewComp.Key.VarietyIndex = GrassVarietyIndex;
+									FCachedLandscapeFoliage::FGrassComp NewComp;
+									NewComp.Key.BasedOn = Component;
+									NewComp.Key.GrassType = GrassType;
+									NewComp.Key.GrassName = GrassName;
+									NewComp.Key.SqrtSubsections = SqrtSubsections;
+									NewComp.Key.CachedMaxInstancesPerComponent = MaxInstancesPerComponent;
+									NewComp.Key.SubsectionX = SubX;
+									NewComp.Key.SubsectionY = SubY;
+									NewComp.Key.NumVarieties = GrassType->GrassVarieties.Num();
+									NewComp.Key.VarietyIndex = GrassVarietyIndex;
 
-										bool bRebuildForBoxes = false;
+									bool bRebuildForBoxes = false;
 
+									{
+										FCachedLandscapeFoliage::FGrassComp* Existing = FoliageCache.CachedGrassComps.Find(NewComp.Key);
+										if (Existing && !Existing->PreviousFoliage.IsValid() && Existing->ExclusionChangeTag != GGrassExclusionChangeTag && !Existing->PendingRemovalRebuild && !Existing->Pending)
 										{
-											FCachedLandscapeFoliage::FGrassComp* Existing = FoliageCache.CachedGrassComps.Find(NewComp.Key);
-											if (Existing && !Existing->PreviousFoliage.IsValid() && Existing->ExclusionChangeTag != GGrassExclusionChangeTag && !Existing->PendingRemovalRebuild && !Existing->Pending)
+											for (const ULandscapeComponent::FExcludeBox& BoxWrapper : Component->ActiveExcludedBoxes)
 											{
-												for (const ULandscapeComponent::FExcludeBox& BoxWrapper : Component->ActiveExcludedBoxes)
-												{
-													if (BoxWrapper.Box.Intersect(WorldSubBox))
-													{
-														NewComp.ExcludedBoxes.Add(BoxWrapper.Box);
-													}
-												}
-												NewComp.ExcludedBoxes.Shrink();
-												if (NewComp.ExcludedBoxes != Existing->ExcludedBoxes)
-												{
-													bRebuildForBoxes = true;
-													NewComp.PreviousFoliage = Existing->Foliage;
-													Existing->PendingRemovalRebuild = true;
-												}
-												else
+												if (BoxWrapper.Box.Intersect(WorldSubBox))
 												{
-													Existing->ExclusionChangeTag = GGrassExclusionChangeTag;
+													NewComp.ExcludedBoxes.Add(BoxWrapper.Box);
 												}
 											}
-
-											if (Existing || MinDistanceToSubComp > MustHaveDistance)
+											NewComp.ExcludedBoxes.Shrink();
+											if (NewComp.ExcludedBoxes != Existing->ExcludedBoxes)
 											{
-												if (Existing)
-												{
-													Existing->Touch();
-												}
-												if (!bRebuildForBoxes)
-												{
-													continue;
-												}
+												bRebuildForBoxes = true;
+												NewComp.PreviousFoliage = Existing->Foliage;
+												Existing->PendingRemovalRebuild = true;
+											}
+											else
+											{
+												Existing->ExclusionChangeTag = GGrassExclusionChangeTag;
 											}
 										}
 
-										if (!bRebuildForBoxes && !bForceSync && (InOutNumCompsCreated >= GrassMaxCreatePerFrame || AsyncFoliageTasks.Num() >= MaxTasks))
+										if (Existing || MinDistanceToSubComp > MustHaveDistance)
 										{
-											continue; // one per frame, but we still want to touch the existing ones and we must do the rebuilds because we changed the tag
+											if (Existing)
+											{
+												Existing->Touch();
+											}
+											if (!bRebuildForBoxes)
+											{
+												continue;
+											}
 										}
-										if (!bRebuildForBoxes)
+									}
+
+									if (!bRebuildForBoxes && !bForceSync && (InOutNumCompsCreated >= GrassMaxCreatePerFrame || AsyncFoliageTasks.Num() >= MaxTasks))
+									{
+										continue; // one per frame, but we still want to touch the existing ones and we must do the rebuilds because we changed the tag
+									}
+									if (!bRebuildForBoxes)
+									{
+										for (const ULandscapeComponent::FExcludeBox& BoxWrapper : Component->ActiveExcludedBoxes)
 										{
-											for (const ULandscapeComponent::FExcludeBox& BoxWrapper : Component->ActiveExcludedBoxes)
+											if (BoxWrapper.Box.Intersect(WorldSubBox))
 											{
-												if (BoxWrapper.Box.Intersect(WorldSubBox))
-												{
-													NewComp.ExcludedBoxes.Add(BoxWrapper.Box);
-												}
+												NewComp.ExcludedBoxes.Add(BoxWrapper.Box);
 											}
-											NewComp.ExcludedBoxes.Shrink();
 										}
-										NewComp.ExclusionChangeTag = GGrassExclusionChangeTag;
+										NewComp.ExcludedBoxes.Shrink();
+									}
+									NewComp.ExclusionChangeTag = GGrassExclusionChangeTag;
 
-										check(Component->GrassData->HasValidData());
+									check(Component->GrassData->HasValidData());
 
-										InOutNumCompsCreated++;
+									InOutNumCompsCreated++;
 
-										SCOPE_CYCLE_COUNTER(STAT_FoliageGrassStartComp);
+									SCOPE_CYCLE_COUNTER(STAT_FoliageGrassStartComp);
 
-										// To guarantee consistency across platforms, we force the string to be lowercase and always treat it as an ANSI string.
-										int32 FolSeed = FCrc::StrCrc32( StringCast<ANSICHAR>( *FString::Printf( TEXT("%s%s%d %d %d"), *GrassType->GetName().ToLower(), *Component->GetName().ToLower(), SubX, SubY, GrassVarietyIndex)).Get() );
-										if (FolSeed == 0)
-										{
-											FolSeed++;
-										}
-
-										// Do not record the transaction of creating temp component for visualizations
-										ClearFlags(RF_Transactional);
-										bool PreviousPackageDirtyFlag = GetOutermost()->IsDirty();
+									// To guarantee consistency across platforms, we force the string to be lowercase and always treat it as an ANSI string.
+									int32 FolSeed = FCrc::StrCrc32(StringCast<ANSICHAR>(*FString::Printf(TEXT("%s_%s_%d_%d_%d_%d_%d"),
+										*GrassName.ToString().ToLower(), *GrassType->GetName().ToLower(), Component->GetSectionBase().X, Component->GetSectionBase().Y, SubX, SubY, GrassVarietyIndex)).Get());
+									if (FolSeed == 0)
+									{
+										FolSeed++;
+									}
 
-										UGrassInstancedStaticMeshComponent* GrassInstancedStaticMeshComponent;
-										{
-											QUICK_SCOPE_CYCLE_COUNTER(STAT_GrassCreateComp);
-											GrassInstancedStaticMeshComponent = NewObject<UGrassInstancedStaticMeshComponent>(this, NAME_None, RF_Transient);
-										}
-										NewComp.Foliage = GrassInstancedStaticMeshComponent;
-										FoliageCache.CachedGrassComps.Add(NewComp);
-
-										GrassInstancedStaticMeshComponent->Mobility = EComponentMobility::Static;
-										GrassInstancedStaticMeshComponent->SetStaticMesh(GrassVariety.GrassMesh);
-										GrassInstancedStaticMeshComponent->MinLOD = GrassVariety.MinLOD;
-										GrassInstancedStaticMeshComponent->bSelectable = false;
-										GrassInstancedStaticMeshComponent->bHasPerInstanceHitProxies = false;
-										GrassInstancedStaticMeshComponent->bReceivesDecals = GrassVariety.bReceivesDecals;
-										static FName NoCollision("NoCollision");
-										GrassInstancedStaticMeshComponent->SetCollisionProfileName(NoCollision);
-										GrassInstancedStaticMeshComponent->bDisableCollision = true;
-										GrassInstancedStaticMeshComponent->SetCanEverAffectNavigation(false);
-										GrassInstancedStaticMeshComponent->InstancingRandomSeed = FolSeed;
-										GrassInstancedStaticMeshComponent->LightingChannels = GrassVariety.LightingChannels;
-										GrassInstancedStaticMeshComponent->bHoldout = bHoldout;
-										GrassInstancedStaticMeshComponent->bCastStaticShadow = false;
-										GrassInstancedStaticMeshComponent->CastShadow = (GrassVariety.bCastDynamicShadow || GrassVariety.bCastContactShadow) && !bDisableDynamicShadows;
-										GrassInstancedStaticMeshComponent->bAffectDistanceFieldLighting = GrassVariety.bAffectDistanceFieldLighting;
-										GrassInstancedStaticMeshComponent->bCastDynamicShadow = GrassVariety.bCastDynamicShadow && !bDisableDynamicShadows;
-										GrassInstancedStaticMeshComponent->bCastContactShadow = GrassVariety.bCastContactShadow && !bDisableDynamicShadows;
-										GrassInstancedStaticMeshComponent->OverrideMaterials = GrassVariety.OverrideMaterials;
-										GrassInstancedStaticMeshComponent->bEvaluateWorldPositionOffset = true;
-										GrassInstancedStaticMeshComponent->WorldPositionOffsetDisableDistance = GrassVariety.InstanceWorldPositionOffsetDisableDistance;
-										GrassInstancedStaticMeshComponent->ShadowCacheInvalidationBehavior = GrassVariety.ShadowCacheInvalidationBehavior;
-										
-										GrassInstancedStaticMeshComponent->PrecachePSOs();
-
-										const FMeshMapBuildData* MeshMapBuildData = Component->GetMeshMapBuildData();
-
-										if (GrassVariety.bUseLandscapeLightmap
-											&& GrassVariety.GrassMesh->GetNumLODs() > 0
-											&& MeshMapBuildData
-											&& MeshMapBuildData->LightMap)
-										{
-											GrassInstancedStaticMeshComponent->SetLODDataCount(GrassVariety.GrassMesh->GetNumLODs(), GrassVariety.GrassMesh->GetNumLODs());
+									// Do not record the transaction of creating temp component for visualizations
+									ClearFlags(RF_Transactional);
+									bool PreviousPackageDirtyFlag = GetOutermost()->IsDirty();
 
-											FLightMapRef GrassLightMap = new FLandscapeGrassLightMap(*MeshMapBuildData->LightMap->GetLightMap2D());
-											FShadowMapRef GrassShadowMap = MeshMapBuildData->ShadowMap ? new FLandscapeGrassShadowMap(*MeshMapBuildData->ShadowMap->GetShadowMap2D()) : nullptr;
+									UGrassInstancedStaticMeshComponent* GrassInstancedStaticMeshComponent;
+									{
+										QUICK_SCOPE_CYCLE_COUNTER(STAT_GrassCreateComp);
+										GrassInstancedStaticMeshComponent = NewObject<UGrassInstancedStaticMeshComponent>(this, NAME_None, RF_Transient);
+									}
+									NewComp.Foliage = GrassInstancedStaticMeshComponent;
+									FoliageCache.CachedGrassComps.Add(NewComp);
+
+									GrassInstancedStaticMeshComponent->Mobility = EComponentMobility::Static;
+									GrassInstancedStaticMeshComponent->SetStaticMesh(GrassVariety.GrassMesh);
+									GrassInstancedStaticMeshComponent->MinLOD = GrassVariety.MinLOD;
+									GrassInstancedStaticMeshComponent->bSelectable = false;
+									GrassInstancedStaticMeshComponent->bHasPerInstanceHitProxies = false;
+									GrassInstancedStaticMeshComponent->bReceivesDecals = GrassVariety.bReceivesDecals;
+									static FName NoCollision("NoCollision");
+									GrassInstancedStaticMeshComponent->SetCollisionProfileName(NoCollision);
+									GrassInstancedStaticMeshComponent->bDisableCollision = true;
+									GrassInstancedStaticMeshComponent->SetCanEverAffectNavigation(false);
+									GrassInstancedStaticMeshComponent->InstancingRandomSeed = FolSeed;
+									GrassInstancedStaticMeshComponent->LightingChannels = GrassVariety.LightingChannels;
+									GrassInstancedStaticMeshComponent->bHoldout = bHoldout;
+									GrassInstancedStaticMeshComponent->bCastStaticShadow = false;
+									GrassInstancedStaticMeshComponent->CastShadow = (GrassVariety.bCastDynamicShadow || GrassVariety.bCastContactShadow) && !bDisableDynamicShadows;
+									GrassInstancedStaticMeshComponent->bAffectDistanceFieldLighting = GrassVariety.bAffectDistanceFieldLighting;
+									GrassInstancedStaticMeshComponent->bCastDynamicShadow = GrassVariety.bCastDynamicShadow && !bDisableDynamicShadows;
+									GrassInstancedStaticMeshComponent->bCastContactShadow = GrassVariety.bCastContactShadow && !bDisableDynamicShadows;
+									GrassInstancedStaticMeshComponent->OverrideMaterials = GrassVariety.OverrideMaterials;
+									GrassInstancedStaticMeshComponent->bEvaluateWorldPositionOffset = true;
+									GrassInstancedStaticMeshComponent->WorldPositionOffsetDisableDistance = GrassVariety.InstanceWorldPositionOffsetDisableDistance;
+									GrassInstancedStaticMeshComponent->ShadowCacheInvalidationBehavior = GrassVariety.ShadowCacheInvalidationBehavior;
+
+									GrassInstancedStaticMeshComponent->PrecachePSOs();
+
+									const FMeshMapBuildData* MeshMapBuildData = Component->GetMeshMapBuildData();
+
+									if (GrassVariety.bUseLandscapeLightmap
+										&& GrassVariety.GrassMesh->GetNumLODs() > 0
+										&& MeshMapBuildData
+										&& MeshMapBuildData->LightMap)
+									{
+										GrassInstancedStaticMeshComponent->SetLODDataCount(GrassVariety.GrassMesh->GetNumLODs(), GrassVariety.GrassMesh->GetNumLODs());
 
-											for (auto& LOD : GrassInstancedStaticMeshComponent->LODData)
-											{
-												// This trasient OverrideMapBuildData will be cleaned up by UMapBuildDataRegistry::CleanupTransientOverrideMapBuildData() if the underlying MeshMapBuildData is gone
-												LOD.OverrideMapBuildData = MakeUnique<FMeshMapBuildData>();
-												LOD.OverrideMapBuildData->LightMap = GrassLightMap;
-												LOD.OverrideMapBuildData->ShadowMap = GrassShadowMap;
-												LOD.OverrideMapBuildData->ResourceCluster = MeshMapBuildData->ResourceCluster;
-											}
-										}
+										FLightMapRef GrassLightMap = new FLandscapeGrassLightMap(*MeshMapBuildData->LightMap->GetLightMap2D());
+										FShadowMapRef GrassShadowMap = MeshMapBuildData->ShadowMap ? new FLandscapeGrassShadowMap(*MeshMapBuildData->ShadowMap->GetShadowMap2D()) : nullptr;
 
-										if (!Cameras.Num() || bDisableGPUCull)
-										{
-											// if we don't have any cameras, then we are rendering landscape LOD materials or somesuch and we want to disable culling
-											GrassInstancedStaticMeshComponent->InstanceStartCullDistance = 0;
-											GrassInstancedStaticMeshComponent->InstanceEndCullDistance = 0;
-										}
-										else
+										for (auto& LOD : GrassInstancedStaticMeshComponent->LODData)
 										{
-											GrassInstancedStaticMeshComponent->InstanceStartCullDistance = static_cast<int32>(GrassVariety.GetStartCullDistance() * CullDistanceScale);
-											GrassInstancedStaticMeshComponent->InstanceEndCullDistance = static_cast<int32>(GrassVariety.GetEndCullDistance() * CullDistanceScale);
+											// This trasient OverrideMapBuildData will be cleaned up by UMapBuildDataRegistry::CleanupTransientOverrideMapBuildData() if the underlying MeshMapBuildData is gone
+											LOD.OverrideMapBuildData = MakeUnique<FMeshMapBuildData>();
+											LOD.OverrideMapBuildData->LightMap = GrassLightMap;
+											LOD.OverrideMapBuildData->ShadowMap = GrassShadowMap;
+											LOD.OverrideMapBuildData->ResourceCluster = MeshMapBuildData->ResourceCluster;
 										}
+									}
 
-										{
-											QUICK_SCOPE_CYCLE_COUNTER(STAT_GrassAttachComp);
+									if (!Cameras.Num() || bDisableGPUCull)
+									{
+										// if we don't have any cameras, then we are rendering landscape LOD materials or somesuch and we want to disable culling
+										GrassInstancedStaticMeshComponent->InstanceStartCullDistance = 0;
+										GrassInstancedStaticMeshComponent->InstanceEndCullDistance = 0;
+									}
+									else
+									{
+										GrassInstancedStaticMeshComponent->InstanceStartCullDistance = static_cast<int32>(GrassVariety.GetStartCullDistance() * CullDistanceScale);
+										GrassInstancedStaticMeshComponent->InstanceEndCullDistance = static_cast<int32>(GrassVariety.GetEndCullDistance() * CullDistanceScale);
+									}
 
-											GrassInstancedStaticMeshComponent->AttachToComponent(GetRootComponent(), FAttachmentTransformRules::KeepRelativeTransform);
-											FTransform DesiredTransform = GetRootComponent()->GetComponentTransform();
-											DesiredTransform.RemoveScaling();
-											GrassInstancedStaticMeshComponent->SetWorldTransform(DesiredTransform);
+									{
+										QUICK_SCOPE_CYCLE_COUNTER(STAT_GrassAttachComp);
 
-											FoliageComponents.Add(GrassInstancedStaticMeshComponent);
-										}
+										GrassInstancedStaticMeshComponent->AttachToComponent(GetRootComponent(), FAttachmentTransformRules::KeepRelativeTransform);
+										FTransform DesiredTransform = GetRootComponent()->GetComponentTransform();
+										DesiredTransform.RemoveScaling();
+										GrassInstancedStaticMeshComponent->SetWorldTransform(DesiredTransform);
 
-										FAsyncGrassBuilder* Builder;
+										FoliageComponents.Add(GrassInstancedStaticMeshComponent);
+									}
 
-										{
-											QUICK_SCOPE_CYCLE_COUNTER(STAT_GrassCreateBuilder);
+									FAsyncGrassBuilder* Builder;
 
-											uint32 HaltonIndexForSub = 0;
-											if (bUseHalton)
-											{
-												check(HaltonBaseIndex > (uint32)MaxInstancesSub);
-												HaltonIndexForSub = HaltonBaseIndex - (uint32)MaxInstancesSub;
-											}
-											Builder = new FAsyncGrassBuilder(this, Component, GrassType, GrassVariety, FeatureLevel, GrassInstancedStaticMeshComponent, SqrtSubsections, SubX, SubY, HaltonIndexForSub, NewComp.ExcludedBoxes);
-										}
+									{
+										QUICK_SCOPE_CYCLE_COUNTER(STAT_GrassCreateBuilder);
 
-										if (Builder->bHaveValidData)
+										uint32 HaltonIndexForSub = 0;
+										if (bUseHalton)
 										{
-											FAsyncTask<FAsyncGrassTask>* Task = new FAsyncTask<FAsyncGrassTask>(Builder, NewComp.Key, GrassInstancedStaticMeshComponent);
+											check(HaltonBaseIndex > (uint32)MaxInstancesSub);
+											HaltonIndexForSub = HaltonBaseIndex - (uint32)MaxInstancesSub;
+										}
+										Builder = new FAsyncGrassBuilder(this, Component, GrassName, GrassType, GrassVariety, FeatureLevel, GrassInstancedStaticMeshComponent, SqrtSubsections, SubX, SubY, HaltonIndexForSub, NewComp.ExcludedBoxes);
+									}
 
-											Task->StartBackgroundTask();
+									if (Builder->bHaveValidData)
+									{
+										FAsyncTask<FAsyncGrassTask>* Task = new FAsyncTask<FAsyncGrassTask>(Builder, NewComp.Key, GrassInstancedStaticMeshComponent);
 
-											AsyncFoliageTasks.Add(Task);
-										}
-										else
-										{
-											delete Builder;
-										}
-										{
-											QUICK_SCOPE_CYCLE_COUNTER(STAT_GrassRegisterComp);
+										Task->StartBackgroundTask();
 
-											GrassInstancedStaticMeshComponent->RegisterComponent();
-										}
+										AsyncFoliageTasks.Add(Task);
+									}
+									else
+									{
+										delete Builder;
+									}
+									{
+										QUICK_SCOPE_CYCLE_COUNTER(STAT_GrassRegisterComp);
 
-										SetFlags(RF_Transactional);
-										GetOutermost()->SetDirtyFlag(PreviousPackageDirtyFlag);
+										GrassInstancedStaticMeshComponent->RegisterComponent();
 									}
+
+									SetFlags(RF_Transactional);
+									GetOutermost()->SetDirtyFlag(PreviousPackageDirtyFlag);
 								}
 							}
 						}
diff --git a/Engine/Source/Runtime/Landscape/Private/LandscapeGrassMapsBuilder.cpp b/Engine/Source/Runtime/Landscape/Private/LandscapeGrassMapsBuilder.cpp
--- a/Engine/Source/Runtime/Landscape/Private/LandscapeGrassMapsBuilder.cpp
+++ b/Engine/Source/Runtime/Landscape/Private/LandscapeGrassMapsBuilder.cpp
@@ -181,19 +181,12 @@ uint32 ComputeGrassMapGenerationHash(const ULandscapeComponent* Component, UMate
 	return Hash;
 }
 
-uint32 ComputeGrassInstanceGenerationHash(uint32 GrassMapGenerationHash, const TArray<TObjectPtr<ULandscapeGrassType>>& GrassTypes)
+uint32 ComputeGrassInstanceGenerationHash(uint32 GrassMapGenerationHash, const TMap<FName, TObjectPtr<ULandscapeGrassType>>& NamedGrassTypes)
 {
 	TRACE_CPUPROFILER_EVENT_SCOPE(ComputeGrassInstanceGenerationHash);
-	// grass maps are input to grass instances
-	uint32 Hash = GrassMapGenerationHash;
-
-	// If anything changes in the grass types, that affects the grass instances
-	for (ULandscapeGrassType* GrassType : GrassTypes)
-	{
-		Hash = FCrc::TypeCrc32(GrassType ? GrassType->StateHash : 0, Hash);
-	}
 
-	return Hash;
+	// grass maps are input to grass instances
+	return HashCombine(GrassMapGenerationHash, UE::Landscape::Grass::Private::ComputeNamedGrassTypesHash(NamedGrassTypes));
 }
 
 // Deprecated
@@ -208,19 +201,19 @@ namespace UE::Landscape::Grass::Private
 {
 
 #if WITH_EDITOR
-void CompileGrassMapShader(const ULandscapeComponent* Component)
-{
-	if (Component->GetMaterialInstanceCount(false) > 0)
+	void CompileGrassMapShader(const ULandscapeComponent* Component)
 	{
-		if (UMaterialInstance* MaterialInstance = Component->GetMaterialInstance(0))
+		if (Component->GetMaterialInstanceCount(false) > 0)
 		{
-			if (FMaterialResource* MaterialResource = MaterialInstance->GetMaterialResource(GetFeatureLevelShaderPlatform_Checked(Component->GetWorld()->GetFeatureLevel())))
+			if (UMaterialInstance* MaterialInstance = Component->GetMaterialInstance(0))
 			{
-				MaterialResource->FinishCompilation();
+				if (FMaterialResource* MaterialResource = MaterialInstance->GetMaterialResource(GetFeatureLevelShaderPlatform_Checked(Component->GetWorld()->GetFeatureLevel())))
+				{
+					MaterialResource->FinishCompilation();
+				}
 			}
 		}
 	}
-}
 #endif // WITH_EDITOR
 
 // calculates the minimum distance between any cameras and the specified Worldbounds.
@@ -244,7 +237,7 @@ namespace UE::Landscape
 {
 
 // Deprecated
-bool CanRenderGrassMap(ULandscapeComponent* Component)
+bool CanRenderGrassMap(ULandscapeComponent *Component)
 {
 	return Component->CanRenderGrassMap();
 }
@@ -338,8 +331,7 @@ FLandscapeGrassMapsBuilder::FComponentState::FComponentState(ULandscapeComponent
 	if (Material)
 	{
 		GrassMapGenerationHash = UE::Landscape::Grass::ComputeGrassMapGenerationHash(Component, Material);
-		const TArray<TObjectPtr<ULandscapeGrassType>>& GrassTypes = Material->GetCachedExpressionData().GrassTypes;
-		GrassInstanceGenerationHash = UE::Landscape::Grass::ComputeGrassInstanceGenerationHash(GrassMapGenerationHash, GrassTypes);
+		GrassInstanceGenerationHash = UE::Landscape::Grass::ComputeGrassInstanceGenerationHash(GrassMapGenerationHash, Material->GetCachedExpressionData().NamedGrassTypes);
 	}
 #endif // WITH_EDITOR
 }
@@ -506,8 +498,7 @@ bool FLandscapeGrassMapsBuilder::UpdateTrackedComponents(const TArray<FVector>&
 						}
 
 						// check if any grass types have changed -- this invalidates foliage instances but not the grass maps
-						const TArray<TObjectPtr<ULandscapeGrassType>>& GrassTypes = Component->GetGrassTypes();
-						const uint32 CurGrassInstanceGenerationHash = UE::Landscape::Grass::ComputeGrassInstanceGenerationHash(CurGrassMapGenerationHash, GrassTypes);
+						const uint32 CurGrassInstanceGenerationHash = UE::Landscape::Grass::ComputeGrassInstanceGenerationHash(CurGrassMapGenerationHash, Component->GetNamedGrassTypes());
 						if (State->GrassInstanceGenerationHash != CurGrassInstanceGenerationHash)
 						{
 							ComponentsToRemoveFoliageInstances.Add(Component);
@@ -1412,8 +1403,7 @@ bool FLandscapeGrassMapsBuilder::TryFastpathsFromPending(FComponentState& State,
 			if (UMaterialInterface* Material = Component->GetLandscapeMaterial())
 			{
 				State.GrassMapGenerationHash = UE::Landscape::Grass::ComputeGrassMapGenerationHash(Component, Material);
-				const TArray<TObjectPtr<ULandscapeGrassType>>& GrassTypes = Material->GetCachedExpressionData().GrassTypes;
-				State.GrassInstanceGenerationHash = UE::Landscape::Grass::ComputeGrassInstanceGenerationHash(State.GrassMapGenerationHash, GrassTypes);
+				State.GrassInstanceGenerationHash = UE::Landscape::Grass::ComputeGrassInstanceGenerationHash(State.GrassMapGenerationHash, Material->GetCachedExpressionData().NamedGrassTypes);
 			}
 		}
 
@@ -1629,7 +1619,7 @@ void FLandscapeGrassMapsBuilder::KickOffRenderAndReadback(FComponentState& State
 	check(State.ActiveRender == nullptr);
 	constexpr bool bInNeedsGrassmap = true;
 	constexpr bool bInNeedsHeightmap = true;
-	State.ActiveRender.Reset(new FLandscapeGrassWeightExporter(Component->GetLandscapeActor(), { Component }, UE::Landscape::Grass::EGrassWeightExporterFlags::Default, /*InRequestedGrassTypes = */{}, MakeConstArrayView(HeightMips)));
+	State.ActiveRender.Reset(new FLandscapeGrassWeightExporter(Component->GetLandscapeActor(), { Component }, UE::Landscape::Grass::EGrassWeightExporterFlags::Default, /*InRequestedGrassNames = */{}, MakeConstArrayView(HeightMips)));
 	check(State.ActiveRender != nullptr);
 
 	State.Stage = EComponentStage::Rendering;
diff --git a/Engine/Source/Runtime/Landscape/Private/LandscapeGrassWeightExporter.cpp b/Engine/Source/Runtime/Landscape/Private/LandscapeGrassWeightExporter.cpp
--- a/Engine/Source/Runtime/Landscape/Private/LandscapeGrassWeightExporter.cpp
+++ b/Engine/Source/Runtime/Landscape/Private/LandscapeGrassWeightExporter.cpp
@@ -217,12 +217,16 @@ IMPLEMENT_MATERIAL_SHADER_TYPE(, FLandscapeGrassWeightPS, TEXT("/Engine/Private/
 
 // ----------------------------------------------------------------------------------
 
-void UE::Landscape::Grass::AddGrassWeightShaderTypes(FMaterialShaderTypes& InOutShaderTypes)
+namespace UE::Landscape::Grass::Private
+{
+
+void AddGrassWeightShaderTypes(FMaterialShaderTypes& InOutShaderTypes)
 {
 	InOutShaderTypes.AddShaderType<FLandscapeGrassWeightVS>();
 	InOutShaderTypes.AddShaderType<FLandscapeGrassWeightPS>();
 }
 
+} // namespace UE::Landscape::Grass::Private
 
 // ----------------------------------------------------------------------------------
 
@@ -607,7 +611,7 @@ TValueOrError<FIntRect /*OutputRect*/, FLandscapeGrassWeightExporter_RenderThrea
 	return MakeValue(OutputRect);
 }
 
-TValueOrError<TPair<FIntRect /*OutputRect*/, uint8 /*ChannelIndex*/>, FLandscapeGrassWeightExporter_RenderThread::EGrassmapTextureInfoError> FLandscapeGrassWeightExporter_RenderThread::GetTextureInfoForGrass(const FIntPoint& InComponentKey, ULandscapeGrassType* InGrassType)
+TValueOrError<TPair<FIntRect /*OutputRect*/, uint8 /*ChannelIndex*/>, FLandscapeGrassWeightExporter_RenderThread::EGrassmapTextureInfoError> FLandscapeGrassWeightExporter_RenderThread::GetTextureInfoForGrass(const FIntPoint& InComponentKey, const FName& InGrassName)
 {
 	using namespace UE::Landscape::Grass;
 	using namespace UE::Landscape::Grass::Private;
@@ -623,10 +627,10 @@ TValueOrError<TPair<FIntRect /*OutputRect*/, uint8 /*ChannelIndex*/>, FLandscape
 		return MakeError(EGrassmapTextureInfoError::InvalidFlags);
 	}
 
-	const FComponentInfo::FPassChannelIndex* FoundPassChannelIndex = FoundComponentInfo->PerGrassTypePassChannelIndices.Find(InGrassType);
+	const FComponentInfo::FPassChannelIndex* FoundPassChannelIndex = FoundComponentInfo->PerGrassNamePassChannelIndices.Find(InGrassName);
 	if (FoundPassChannelIndex == nullptr)
 	{
-		return MakeError(EGrassmapTextureInfoError::InvalidGrassType);
+		return MakeError(EGrassmapTextureInfoError::InvalidGrassName);
 	}
 
 	check(FoundComponentInfo->PerPassPerChannelGrassWeightIndices.IsValidIndex(FoundPassChannelIndex->PassIndex));
@@ -673,13 +677,13 @@ FLandscapeGrassWeightExporter::FLandscapeGrassWeightExporter(ALandscapeProxy* In
 		InLandscapeProxy->GetLandscapeActor(),
 		InLandscapeComponents,
 		UE::Landscape::Private::BuildLandscapeGrassWeightExporterConstructionFlags(bInNeedsGrassmap, bInNeedsHeightmap, bInRenderImmediately, bInReadbackToCPU),
-		/*InRequestedGrassTypes = */{},
+		/*InRequestedGrassNames = */{},
 		MakeArrayView(InHeightMips))
 {
 }
 
 FLandscapeGrassWeightExporter::FLandscapeGrassWeightExporter(ALandscape* InLandscape, TArrayView<ULandscapeComponent* const> InRequestedComponents, UE::Landscape::Grass::EGrassWeightExporterFlags InFlags,
-	TConstArrayView<TObjectPtr<ULandscapeGrassType>> InRequestedGrassTypes, TConstArrayView<int32> InRequestedHeightMips)
+	TConstArrayView<FName> InRequestedGrassNames, TConstArrayView<int32> InRequestedHeightMips)
 	: FLandscapeGrassWeightExporter_RenderThread(/*ComponentSizeVerts = */InLandscape->ComponentSizeQuads + 1, InLandscape->SubsectionSizeQuads, InLandscape->NumSubsections, InRequestedHeightMips, InFlags)
 	, Landscape(InLandscape)
 {
@@ -701,7 +705,7 @@ FLandscapeGrassWeightExporter::FLandscapeGrassWeightExporter(ALandscape* InLands
 		MinSectionBase = MinSectionBase.ComponentMin(Component->GetSectionBase());
 		ensure(Component->SceneProxy);
 
-		FComponentInfo& ComponentInfo = ComponentInfos.Emplace_GetRef(Component, InFlags, InRequestedGrassTypes, InRequestedHeightMips);
+		FComponentInfo& ComponentInfo = ComponentInfos.Emplace_GetRef(Component, InFlags, InRequestedGrassNames, InRequestedHeightMips);
 		ComponentInfo.PixelOffsetX = CurrentPixelOffsetX;
 
 		CurrentPixelOffsetX += ComponentInfo.GetNumPasses() * SingleTileWidth;
@@ -920,7 +924,7 @@ TMap<ULandscapeComponent*, TUniquePtr<FLandscapeComponentGrassData>, TInlineSetA
 
 	Results.Reserve(ComponentInfos.Num());
 	FHeightBuffer2DView HeightData;
-	TMap<ULandscapeGrassType*, FByteBuffer2DView> WeightData;
+	TMap<FName, FByteBuffer2DView> WeightData;
 
 	ULandscapeInfo* LandscapeInfo = Landscape->GetLandscapeInfo();
 	check(LandscapeInfo != nullptr);
@@ -945,24 +949,24 @@ TMap<ULandscapeComponent*, TUniquePtr<FLandscapeComponentGrassData>, TInlineSetA
 #endif // WITH_EDITORONLY_DATA
 
 		WeightData.Empty();
-		WeightData.Reserve(ComponentInfo.PerGrassTypePassChannelIndices.Num());
+		WeightData.Reserve(ComponentInfo.PerGrassNamePassChannelIndices.Num());
 
-		for (auto& ItPair : ComponentInfo.PerGrassTypePassChannelIndices)
+		for (auto& ItPair : ComponentInfo.PerGrassNamePassChannelIndices)
 		{
 			// Note: WeightData points directly at the elements of GrassWeightArrays (DO NOT REALLOCATE GRASSWEIGHTARRAYS)
 			WeightData.Add(ItPair.Key, FByteBuffer2DView(/*InByteStrideX = */4, /*InByteStrideY = */TargetSize.X * 4, /*InNumX = */ComponentSizeVerts, /*InNumY = */ComponentSizeVerts));
 		}
 
-		auto FindGrassTypeForPassAndGrassWeight = [&](const FComponentInfo::FPassChannelIndex& InPassChannelIndex) -> ULandscapeGrassType*
+		auto FindGrassNameForPassAndGrassWeight = [&](const FComponentInfo::FPassChannelIndex& InPassChannelIndex) -> FName
 		{
-			for (auto& ItPair : ComponentInfo.PerGrassTypePassChannelIndices)
+			for (auto& ItPair : ComponentInfo.PerGrassNamePassChannelIndices)
 			{
 				if (ItPair.Value == InPassChannelIndex)
 				{
 					return ItPair.Key;
 				}
 			}
-			return nullptr;
+			return NAME_None;
 		};
 
 #if WITH_EDITOR
@@ -1000,9 +1004,9 @@ TMap<ULandscapeComponent*, TUniquePtr<FLandscapeComponentGrassData>, TInlineSetA
 								}
 								else 
 								{
-									ULandscapeGrassType* GrassType = FindGrassTypeForPassAndGrassWeight(InPassChannelIndex);
-									check(GrassType != nullptr);
-									Result = GrassType->GetName();
+									const FName GrassName = FindGrassNameForPassAndGrassWeight(InPassChannelIndex);
+									check(!GrassName.IsNone());
+									Result = GrassName.ToString();
 								}
 							}
 							else
@@ -1030,7 +1034,7 @@ TMap<ULandscapeComponent*, TUniquePtr<FLandscapeComponentGrassData>, TInlineSetA
 				StrBuilder.Appendf(TEXT("A : %s\n"), *GetChannelString(FComponentInfo::FPassChannelIndex { PassIndex, 3 }));
 			}
 			FFileHelper::SaveStringToFile(StrBuilder.ToView(), *(FilePath + ".desc.txt"));
-			FFileHelper::CreateBitmap(*(FilePath + ".bmp"), TargetSize.X, TargetSize.Y, Samples.GetData(), /*SubRegion = */nullptr, &IFileManager::Get(), /*OutFilename = */nullptr, /*bInWriteAlpha = */!ComponentInfo.PerGrassTypePassChannelIndices.IsEmpty());
+			FFileHelper::CreateBitmap(*(FilePath + ".bmp"), TargetSize.X, TargetSize.Y, Samples.GetData(), /*SubRegion = */nullptr, &IFileManager::Get(), /*OutFilename = */nullptr, /*bInWriteAlpha = */!ComponentInfo.PerGrassNamePassChannelIndices.IsEmpty());
 		}
 #endif // WITH_EDITOR
 	
@@ -1052,11 +1056,11 @@ TMap<ULandscapeComponent*, TUniquePtr<FLandscapeComponentGrassData>, TInlineSetA
 			HeightMipsRects.Add(MipIndex, TextureInfo.GetValue());
 		}
 
-		TMap<ULandscapeGrassType*, TPair<FIntRect, int32>> GrassRectsAndChannelIndices;
-		for (auto& ItPair : ComponentInfo.PerGrassTypePassChannelIndices)
+		TMap<FName, TPair<FIntRect, int32>> GrassRectsAndChannelIndices;
+		for (auto& ItPair : ComponentInfo.PerGrassNamePassChannelIndices)
 		{
 			TValueOrError<TPair<FIntRect /*OutputRect*/, uint8 /*ChannelIndex*/>, EGrassmapTextureInfoError> TextureInfo = GetTextureInfoForGrass(ComponentInfo.ComponentKey, ItPair.Key);
-			checkf(!TextureInfo.HasError(), TEXT("GetTextureInfoForGrass failed for grass type %s (error : %i)"), *ItPair.Key->GetName(), static_cast<int32>(TextureInfo.GetError()));
+			checkf(!TextureInfo.HasError(), TEXT("GetTextureInfoForGrass failed for grass name %s (error : %i)"), *ItPair.Key.ToString(), static_cast<int32>(TextureInfo.GetError()));
 			GrassRectsAndChannelIndices.Add(ItPair.Key, TextureInfo.GetValue());
 		}
 #endif // GRASS_ENABLE_EXTRAVALIDATION
@@ -1085,18 +1089,19 @@ TMap<ULandscapeComponent*, TUniquePtr<FLandscapeComponentGrassData>, TInlineSetA
 				{
 					if (PerChannelGrassWeightIndices[ChannelIndex] != GrassWeightInvalidIndex)
 					{
-						ULandscapeGrassType* GrassType = FindGrassTypeForPassAndGrassWeight(FComponentInfo::FPassChannelIndex { PassIdx, ChannelIndex });
+						const FName GrassName = FindGrassNameForPassAndGrassWeight(FComponentInfo::FPassChannelIndex { PassIdx, ChannelIndex });
+						check(!GrassName.IsNone());
 
 #if GRASS_ENABLE_EXTRAVALIDATION
-						TPair<FIntRect, int32>* RectAndIndex = GrassRectsAndChannelIndices.Find(GrassType);
+						TPair<FIntRect, int32>* RectAndIndex = GrassRectsAndChannelIndices.Find(GrassName);
 						check(RectAndIndex != nullptr);
 						check(ChannelIndex == RectAndIndex->Value);
 						FColor* GrassSampleData = &Samples[RectAndIndex->Key.Min.X];
 						check(GrassSampleData == SampleData);
 #endif // GRASS_ENABLE_EXTRAVALIDATION
 
-						FByteBuffer2DView* ByteBufferView = WeightData.Find(GrassType);
-						check((GrassType != nullptr) && (ByteBufferView != nullptr));
+						FByteBuffer2DView* ByteBufferView = WeightData.Find(GrassName);
+						check(ByteBufferView != nullptr);
 						static_assert(PLATFORM_LITTLE_ENDIAN, "Big-endian is not supported anymore");
 						ByteBufferView->BufferStart = reinterpret_cast<uint8*>((ChannelIndex == 0) ? &SampleData->R : (ChannelIndex == 1) ? &SampleData->G : (ChannelIndex == 2) ? &SampleData->B : &SampleData->A);
 					}
@@ -1132,7 +1137,7 @@ TMap<ULandscapeComponent*, TUniquePtr<FLandscapeComponentGrassData>, TInlineSetA
 			}
 		}
 		
-		TMap<ULandscapeGrassType*, IBuffer2DView<uint8>*> WeightDataPtrs;
+		TMap<FName, IBuffer2DView<uint8>*> WeightDataPtrs;
 		WeightDataPtrs.Reserve(WeightData.Num());
 		for (auto& ItPair : WeightData)
 		{
@@ -1177,7 +1182,7 @@ void FLandscapeGrassWeightExporter::ApplyResults(TMap<ULandscapeComponent*, TUni
 }
 
 FLandscapeGrassWeightExporter_RenderThread::FComponentInfo::FComponentInfo(ULandscapeComponent* InComponent, UE::Landscape::Grass::EGrassWeightExporterFlags InFlags,
-	TConstArrayView<TObjectPtr<ULandscapeGrassType>> InRequestedGrassTypes, TConstArrayView<int32> InRequestedHeightMips) 
+	TConstArrayView<FName> InRequestedGrassNames, TConstArrayView<int32> InRequestedHeightMips) 
 	: ComponentKey(InComponent->GetComponentKey())
 	, SceneProxy((FLandscapeComponentSceneProxy*)InComponent->SceneProxy)
 	, Flags(InFlags)
@@ -1222,19 +1227,20 @@ FLandscapeGrassWeightExporter_RenderThread::FComponentInfo::FComponentInfo(ULand
 
 	if (EnumHasAnyFlags(InFlags, EGrassWeightExporterFlags::NeedsGrassmap))
 	{
-		TArray<TObjectPtr<ULandscapeGrassType>> RenderableGrassTypes = InComponent->GetGrassTypes();
-		TArray<TObjectPtr<ULandscapeGrassType>> RequestedGrassTypes = InRequestedGrassTypes.IsEmpty() ? RenderableGrassTypes : TArray<TObjectPtr<ULandscapeGrassType>>(InRequestedGrassTypes);
-		for (TObjectPtr<ULandscapeGrassType> RequestedGrassType : RequestedGrassTypes)
+		TArray<FName> RenderableGrassNames;
+		InComponent->GetNamedGrassTypes().GetKeys(RenderableGrassNames);
+		TArray<FName> RequestedGrassNames = InRequestedGrassNames.IsEmpty() ? RenderableGrassNames : TArray<FName>(InRequestedGrassNames);
+		for (FName RequestedGrassName : RequestedGrassNames)
 		{
-			if (RequestedGrassType != nullptr)
+			if (!RequestedGrassName.IsNone())
 			{
-				if (int32 GrassOutputIndex = RenderableGrassTypes.IndexOfByKey(RequestedGrassType);
+				if (int32 GrassOutputIndex = RenderableGrassNames.IndexOfByKey(RequestedGrassName);
 					GrassOutputIndex != INDEX_NONE)
 				{
 					check(GrassOutputIndex < UMaterialExpressionLandscapeGrassOutput::MaxGrassTypes);
 					FPassChannelIndex PassChannelIndex = FindOrAddChannel();
 					SetValueAtPassIndex(PassChannelIndex, GrassOutputIndex);
-					PerGrassTypePassChannelIndices.Add(RequestedGrassType, PassChannelIndex);
+					PerGrassNamePassChannelIndices.Add(RequestedGrassName, PassChannelIndex);
 				}
 			}
 		}
diff --git a/Engine/Source/Runtime/Landscape/Private/LandscapeUtils.cpp b/Engine/Source/Runtime/Landscape/Private/LandscapeUtils.cpp
--- a/Engine/Source/Runtime/Landscape/Private/LandscapeUtils.cpp
+++ b/Engine/Source/Runtime/Landscape/Private/LandscapeUtils.cpp
@@ -7,6 +7,7 @@
 #include "LandscapeLayerInfoObject.h"
 #include "LandscapeProxy.h"
 #include "LandscapeEditTypes.h"
+#include "LandscapeGrassType.h"
 #include "Algo/Transform.h"
 #include "RenderGraphBuilder.h"
 #include "TextureResource.h"
@@ -612,4 +613,34 @@ bool LandscapeProxySortPredicate(const TWeakObjectPtr<ALandscapeProxy> APtr, con
 	return SectionBaseA.Y < SectionBaseB.Y;
 }
 
-} // end namespace UE::Landscape
\ No newline at end of file
+} // end namespace UE::Landscape
+
+
+// ----------------------------------------------------------------------------------
+
+namespace UE::Landscape::Grass
+{
+
+TArray<TPair<FName, TObjectPtr<ULandscapeGrassType>>> ConvertToNamedGrassTypesList(const TMap<FName, TObjectPtr<ULandscapeGrassType>>& InNamedGrassTypes)
+{
+	TArray<TPair<FName, TObjectPtr<ULandscapeGrassType>>> LocalNamedGrassTypes;
+	LocalNamedGrassTypes.Reserve(InNamedGrassTypes.Num());
+	Algo::Transform(InNamedGrassTypes, LocalNamedGrassTypes, [](const TPair<FName, TObjectPtr<ULandscapeGrassType>>& ItPair) { return ItPair; });
+	return LocalNamedGrassTypes;
+}
+
+TMap<FName, TObjectPtr<ULandscapeGrassType>> FilterValidNamedGrassTypes(const TMap<FName, TObjectPtr<ULandscapeGrassType>>& InNamedGrassTypes)
+{
+	TMap<FName, TObjectPtr<ULandscapeGrassType>> ValidNamedGrassTypes;
+	ValidNamedGrassTypes.Reserve(InNamedGrassTypes.Num());
+	for (const TPair<FName, TObjectPtr<ULandscapeGrassType>>& ItPair : InNamedGrassTypes)
+	{
+		if (!ItPair.Key.IsNone() && (ItPair.Value != nullptr))
+		{
+			ValidNamedGrassTypes.Add(ItPair);
+		}
+	}
+	return ValidNamedGrassTypes;
+}
+
+} // namespace UE::Landscape::Grass
\ No newline at end of file
diff --git a/Engine/Source/Runtime/Landscape/Private/LandscapeUtilsPrivate.cpp b/Engine/Source/Runtime/Landscape/Private/LandscapeUtilsPrivate.cpp
--- a/Engine/Source/Runtime/Landscape/Private/LandscapeUtilsPrivate.cpp
+++ b/Engine/Source/Runtime/Landscape/Private/LandscapeUtilsPrivate.cpp
@@ -4,8 +4,10 @@
 
 #include "DataDrivenShaderPlatformInfo.h"
 #include "LandscapeComponent.h"
+#include "LandscapeGrassType.h"
 #include "LandscapeProxy.h"
 #include "LandscapeInfo.h"
+#include "LandscapeUtils.h"
 #include "RenderingThread.h"
 
 #if WITH_EDITOR
@@ -218,4 +220,35 @@ void SubmitGPUCommands(bool bBlockUntilRTComplete, bool bBlockRTUntilGPUComplete
 	return;
 }
 
-} // end namespace UE::Landscape::Private
\ No newline at end of file
+} // end namespace UE::Landscape::Private
+
+
+// ----------------------------------------------------------------------------------
+
+namespace UE::Landscape::Grass::Private
+{
+
+uint32 ComputeNamedGrassTypesHash(const TMap<FName, TObjectPtr<ULandscapeGrassType>>& InNamedGrassTypes)
+{
+	uint32 Hash = 0;
+
+	// Only keep the valid entries
+	TMap<FName, TObjectPtr<ULandscapeGrassType>> ValidNamedGrassTypes = FilterValidNamedGrassTypes(InNamedGrassTypes);
+	// Turn into an array  
+	const TArray<TPair<FName, TObjectPtr<ULandscapeGrassType>>> SortedNamedGrassTypes = ConvertToNamedGrassTypesList(ValidNamedGrassTypes);
+	// Then sort by grass name for hash stability
+	Algo::SortBy(SortedNamedGrassTypes, &TPair<FName, TObjectPtr<ULandscapeGrassType>>::Key, FNameLexicalLess());
+
+	// If anything changes in the grass types, that affects the grass instances
+	for (const TPair<FName, TObjectPtr<ULandscapeGrassType>>& ItPair : SortedNamedGrassTypes)
+	{
+		check(!ItPair.Key.IsNone() && (ItPair.Value != nullptr));
+
+		Hash = FCrc::TypeCrc32(GetTypeHash(ItPair.Key.ToString().ToLower()), Hash);
+		Hash = FCrc::TypeCrc32(ItPair.Value->StateHash, Hash);
+	}
+
+	return Hash;
+}
+
+} // namespace UE::Landscape::Grass::Private
\ No newline at end of file
diff --git a/Engine/Source/Runtime/Landscape/Private/LandscapeUtilsPrivate.h b/Engine/Source/Runtime/Landscape/Private/LandscapeUtilsPrivate.h
--- a/Engine/Source/Runtime/Landscape/Private/LandscapeUtilsPrivate.h
+++ b/Engine/Source/Runtime/Landscape/Private/LandscapeUtilsPrivate.h
@@ -8,345 +8,357 @@
 enum EShaderPlatform : uint16;
 class ULandscapeComponent;
 class ULandscapeInfo;
+class ULandscapeGrassType;
 struct FLandscapeLayerComponentData;
 
 namespace UE::Landscape::Private
 {
-	int32 ComputeMaxDeltasOffsetForMip(int32 InMipIndex, int32 InNumRelevantMips);
-	int32 ComputeMaxDeltasCountForMip(int32 InMipIndex, int32 InNumRelevantMips);
-	int32 ComputeMipToMipMaxDeltasIndex(int32 InSourceMipIndex, int32 InDestinationMipIndex, int32 InNumRelevantMips);
-	int32 ComputeMipToMipMaxDeltasCount(int32 InNumRelevantMips);
+
+int32 ComputeMaxDeltasOffsetForMip(int32 InMipIndex, int32 InNumRelevantMips);
+int32 ComputeMaxDeltasCountForMip(int32 InMipIndex, int32 InNumRelevantMips);
+int32 ComputeMipToMipMaxDeltasIndex(int32 InSourceMipIndex, int32 InDestinationMipIndex, int32 InNumRelevantMips);
+int32 ComputeMipToMipMaxDeltasCount(int32 InNumRelevantMips);
 
 #if WITH_EDITOR
-	/** Returns true if InPlatform is a mobile platform and the Landscape.MobileWeightTextureArray CVar set */
-	bool UseWeightmapTextureArray(EShaderPlatform InPlatform);
+/** Returns true if InPlatform is a mobile platform and the Landscape.MobileWeightTextureArray CVar set */
+bool UseWeightmapTextureArray(EShaderPlatform InPlatform);
 
-	/** Check if Landscape.MobileWeightTextureArray CVar set and we should attempt to use texture arrays for weight maps on mobile platforms  */
-	bool IsMobileWeightmapTextureArrayEnabled();
+/** Check if Landscape.MobileWeightTextureArray CVar set and we should attempt to use texture arrays for weight maps on mobile platforms  */
+bool IsMobileWeightmapTextureArrayEnabled();
 
-	void CreateLandscapeComponentLayerDataDuplicate(const FLandscapeLayerComponentData& InSourceComponentData, FLandscapeLayerComponentData& OutDestComponentData);
+void CreateLandscapeComponentLayerDataDuplicate(const FLandscapeLayerComponentData& InSourceComponentData, FLandscapeLayerComponentData& OutDestComponentData);
 #endif //!WITH_EDITOR
 
-	/** 
-	 * Helper class to transform a list of values into a 2D array. This is a replacement for a TMap<FIntPoint, ValueType>, where FIntPoint is a 2D key and ValueType, the type of object to store. 
-	 * It uses bit arrays to note the presence of an object and the 2D key to index it. The advantage is that the key can trivially be turned into linear index (and vice-versa), derived from the bounds 
-	 * of all registered keys. It is meant to be used as a temporary helper, which is why there are currently no Add/Remove functions (all objects are registered in the constructor) but 
-	 * this could be done eventually, and even replace ULandscapeInfo's XYToComponentMap entirely */ 
-	template <typename ValueType, typename KeyFuncs>
-	class T2DIndexer
+/** 
+ * Helper class to transform a list of values into a 2D array. This is a replacement for a TMap<FIntPoint, ValueType>, where FIntPoint is a 2D key and ValueType, the type of object to store. 
+ * It uses bit arrays to note the presence of an object and the 2D key to index it. The advantage is that the key can trivially be turned into linear index (and vice-versa), derived from the bounds 
+ * of all registered keys. It is meant to be used as a temporary helper, which is why there are currently no Add/Remove functions (all objects are registered in the constructor) but 
+ * this could be done eventually, and even replace ULandscapeInfo's XYToComponentMap entirely */ 
+template <typename ValueType, typename KeyFuncs>
+class T2DIndexer
+{
+public:
+	T2DIndexer(TConstArrayView<ValueType> InValues)
 	{
-	public:
-		T2DIndexer(TConstArrayView<ValueType> InValues)
-		{
-			FIntRect MinMaxKeys(FIntPoint(MAX_int32, MAX_int32), FIntPoint(MIN_int32, MIN_int32));
-
-			TArray<FIntPoint> AllKeys;
-			const int32 NumValues = InValues.Num();
-			AllKeys.AddDefaulted(NumValues);
+		FIntRect MinMaxKeys(FIntPoint(MAX_int32, MAX_int32), FIntPoint(MIN_int32, MIN_int32));
 
-			// First, we need the bound of all values so that we can properly size the linear array : 
-			for (int32 Index = 0; Index < NumValues; ++Index)
-			{
-				const ValueType& Value = InValues[Index];
-				FIntPoint Key = KeyFuncs::GetKey(Value);
-				AllKeys[Index] = Key;
-				MinMaxKeys.Include(Key);
-			}
-			if (!AllKeys.IsEmpty())
-			{
-				KeyExclusiveBounds = FIntRect(MinMaxKeys.Min, MinMaxKeys.Max + 1);
-			}
+		TArray<FIntPoint> AllKeys;
+		const int32 NumValues = InValues.Num();
+		AllKeys.AddDefaulted(NumValues);
 
-			// Now we can properly size the array and fill it : 
-			int32 NumEntries = KeyExclusiveBounds.Area();
-			AllValues.AddDefaulted(NumEntries);
-			ValidValueBitIndices.Init(false, NumEntries);
-			for (int32 Index = 0; Index < NumValues; ++Index)
-			{
-				int32 ValueIndex = GetValueIndexForKey(AllKeys[Index]);
-				AllValues[ValueIndex] = InValues[Index];
-				ValidValueBitIndices[ValueIndex] = true;
-			}
-		}
-		virtual ~T2DIndexer() = default;
-		T2DIndexer(const T2DIndexer& Other) = default;
-		T2DIndexer(T2DIndexer&& Other) = default;
-		T2DIndexer& operator=(const T2DIndexer& Other) = default;
-		T2DIndexer& operator=(T2DIndexer&& Other) = default;
-
-		bool IsValidKey(const FIntPoint& InKey) const
+		// First, we need the bound of all values so that we can properly size the linear array : 
+		for (int32 Index = 0; Index < NumValues; ++Index)
 		{
-			return KeyExclusiveBounds.Contains(InKey);
+			const ValueType& Value = InValues[Index];
+			FIntPoint Key = KeyFuncs::GetKey(Value);
+			AllKeys[Index] = Key;
+			MinMaxKeys.Include(Key);
 		}
-
-		bool IsValidValueIndex(int32 InIndex) const
+		if (!AllKeys.IsEmpty())
 		{
-			return AllValues.IsValidIndex(InIndex);
+			KeyExclusiveBounds = FIntRect(MinMaxKeys.Min, MinMaxKeys.Max + 1);
 		}
 
-		ValueType& GetValueForKey(const FIntPoint& InKey) const
+		// Now we can properly size the array and fill it : 
+		int32 NumEntries = KeyExclusiveBounds.Area();
+		AllValues.AddDefaulted(NumEntries);
+		ValidValueBitIndices.Init(false, NumEntries);
+		for (int32 Index = 0; Index < NumValues; ++Index)
 		{
-			return AllValues[GetValueIndexForKey(InKey)];
+			int32 ValueIndex = GetValueIndexForKey(AllKeys[Index]);
+			AllValues[ValueIndex] = InValues[Index];
+			ValidValueBitIndices[ValueIndex] = true;
 		}
+	}
+	virtual ~T2DIndexer() = default;
+	T2DIndexer(const T2DIndexer& Other) = default;
+	T2DIndexer(T2DIndexer&& Other) = default;
+	T2DIndexer& operator=(const T2DIndexer& Other) = default;
+	T2DIndexer& operator=(T2DIndexer&& Other) = default;
+
+	bool IsValidKey(const FIntPoint& InKey) const
+	{
+		return KeyExclusiveBounds.Contains(InKey);
+	}
 
-		ValueType& GetValueForKeySafe(const FIntPoint& InKey) const
-		{
-			return IsValidKey(InKey) ? AllValues[GetValueIndexForKey(InKey)] : nullptr;
-		}
+	bool IsValidValueIndex(int32 InIndex) const
+	{
+		return AllValues.IsValidIndex(InIndex);
+	}
 
-		ValueType& GetValueForKeyChecked(const FIntPoint& InKey) const
-		{
-			check(IsValidKey(InKey));
-			return AllValues[GetValueIndexForKey(InKey)];
-		}
+	ValueType& GetValueForKey(const FIntPoint& InKey) const
+	{
+		return AllValues[GetValueIndexForKey(InKey)];
+	}
 
-		int32 GetValueIndexForKey(const FIntPoint& InKey) const
-		{
-			FIntPoint RelativeKey = InKey - KeyExclusiveBounds.Min;
-			return RelativeKey.Y * KeyExclusiveBounds.Width() + RelativeKey.X;
-		}
+	ValueType& GetValueForKeySafe(const FIntPoint& InKey) const
+	{
+		return IsValidKey(InKey) ? AllValues[GetValueIndexForKey(InKey)] : nullptr;
+	}
 
-		int32 GetValueIndexForKeySafe(const FIntPoint& InKey) const
-		{
-			return IsValidKey(InKey) ? GetValueIndexForKey(InKey) : INDEX_NONE;
-		}
+	ValueType& GetValueForKeyChecked(const FIntPoint& InKey) const
+	{
+		check(IsValidKey(InKey));
+		return AllValues[GetValueIndexForKey(InKey)];
+	}
 
-		int32 GetValueIndexForKeyChecked(const FIntPoint& InKey) const
-		{
-			check(IsValidKey(InKey));
-			return GetValueIndexForKey(InKey);
-		}
+	int32 GetValueIndexForKey(const FIntPoint& InKey) const
+	{
+		FIntPoint RelativeKey = InKey - KeyExclusiveBounds.Min;
+		return RelativeKey.Y * KeyExclusiveBounds.Width() + RelativeKey.X;
+	}
 
-		FIntPoint GetValueKeyForIndex(int32 InIndex) const
-		{
-			const int32 Stride = KeyExclusiveBounds.Width();
-			return KeyExclusiveBounds.Min + FIntPoint(InIndex % Stride, InIndex / Stride);
-		}
+	int32 GetValueIndexForKeySafe(const FIntPoint& InKey) const
+	{
+		return IsValidKey(InKey) ? GetValueIndexForKey(InKey) : INDEX_NONE;
+	}
 
-		FIntPoint GetValueKeyForIndexSafe(int32 InIndex) const
-		{
-			return IsValidValueIndex(InIndex) ? GetValueKeyForIndex(InIndex) : FIntPoint::NoneValue;
-		}
+	int32 GetValueIndexForKeyChecked(const FIntPoint& InKey) const
+	{
+		check(IsValidKey(InKey));
+		return GetValueIndexForKey(InKey);
+	}
 
-		FIntPoint GetValueKeyForIndexChecked(int32 InIndex) const
-		{
-			check(IsValidValueIndex(InIndex));
-			return GetValueKeyForIndex(InIndex);
-		}
+	FIntPoint GetValueKeyForIndex(int32 InIndex) const
+	{
+		const int32 Stride = KeyExclusiveBounds.Width();
+		return KeyExclusiveBounds.Min + FIntPoint(InIndex % Stride, InIndex / Stride);
+	}
 
-		int32 GetValueIndex(const ValueType& InValue) const
-		{
-			return GetValueIndexForKey(KeyFuncs::GetKey(InValue));
-		}
+	FIntPoint GetValueKeyForIndexSafe(int32 InIndex) const
+	{
+		return IsValidValueIndex(InIndex) ? GetValueKeyForIndex(InIndex) : FIntPoint::NoneValue;
+	}
 
-		int32 GetValueIndexSafe(const ValueType& InValue) const
-		{
-			return GetValueIndexForKeySafe(KeyFuncs::GetKey(InValue));
-		}
+	FIntPoint GetValueKeyForIndexChecked(int32 InIndex) const
+	{
+		check(IsValidValueIndex(InIndex));
+		return GetValueKeyForIndex(InIndex);
+	}
 
-		int32 GetValueIndexChecked(const ValueType& InValue) const
-		{
-			return GetValueIndexForKeyChecked(KeyFuncs::GetKey(InValue));
-		}
+	int32 GetValueIndex(const ValueType& InValue) const
+	{
+		return GetValueIndexForKey(KeyFuncs::GetKey(InValue));
+	}
 
-		bool IsValidValue(const FIntPoint& InKey) const
-		{
-			return ValidValueBitIndices[GetValueIndexForKey(InKey)];
-		}
+	int32 GetValueIndexSafe(const ValueType& InValue) const
+	{
+		return GetValueIndexForKeySafe(KeyFuncs::GetKey(InValue));
+	}
+
+	int32 GetValueIndexChecked(const ValueType& InValue) const
+	{
+		return GetValueIndexForKeyChecked(KeyFuncs::GetKey(InValue));
+	}
 
-		bool IsValidValueSafe(const FIntPoint& InKey) const
+	bool IsValidValue(const FIntPoint& InKey) const
+	{
+		return ValidValueBitIndices[GetValueIndexForKey(InKey)];
+	}
+
+	bool IsValidValueSafe(const FIntPoint& InKey) const
+	{
+		return IsValidKey(InKey) ? ValidValueBitIndices[GetValueIndexForKey(InKey)] : false;
+	}
+
+	bool IsValidValueChecked(const FIntPoint& InKey) const
+	{
+		check(IsValidKey(InKey));
+		return ValidValueBitIndices[GetValueIndexForKey(InKey)];
+	}
+
+	bool HasValidValueInBounds(const FIntRect& InBounds, bool bInInclusiveBounds)
+	{
+		FIntRect LocalBounds = InBounds;
+		// Convert to exclusive bounds if necessary :
+		if (bInInclusiveBounds)
 		{
-			return IsValidKey(InKey) ? ValidValueBitIndices[GetValueIndexForKey(InKey)] : false;
+			LocalBounds.Max += FIntPoint(1, 1);
 		}
-
-		bool IsValidValueChecked(const FIntPoint& InKey) const
+		// Intersect with the current key bounds in order to limit the test to only the tracked area :
+		LocalBounds.Clip(KeyExclusiveBounds);
+		if (LocalBounds.Area() <= 0)
 		{
-			check(IsValidKey(InKey));
-			return ValidValueBitIndices[GetValueIndexForKey(InKey)];
+			return false;
 		}
 
-		bool HasValidValueInBounds(const FIntRect& InBounds, bool bInInclusiveBounds)
+		for (int32 Y = LocalBounds.Min.Y; (Y < LocalBounds.Max.Y); ++Y)
 		{
-			FIntRect LocalBounds = InBounds;
-			// Convert to exclusive bounds if necessary :
-			if (bInInclusiveBounds)
-			{
-				LocalBounds.Max += FIntPoint(1, 1);
-			}
-			// Intersect with the current key bounds in order to limit the test to only the tracked area :
-			LocalBounds.Clip(KeyExclusiveBounds);
-			if (LocalBounds.Area() <= 0)
-			{
-				return false;
-			}
-
-			for (int32 Y = LocalBounds.Min.Y; (Y < LocalBounds.Max.Y); ++Y)
+			for (int32 X = LocalBounds.Min.X; (X < LocalBounds.Max.X); ++X)
 			{
-				for (int32 X = LocalBounds.Min.X; (X < LocalBounds.Max.X); ++X)
+				if (IsValidValue(FIntPoint(X, Y)))
 				{
-					if (IsValidValue(FIntPoint(X, Y)))
-					{
-						return true;
-					}
+					return true;
 				}
 			}
-
-			return false;
 		}
 
-		TBitArray<> GetValidValueBitIndicesInBounds(const FIntRect& InBounds, bool bInInclusiveBounds)
-		{
-			TBitArray<> Result(false, ValidValueBitIndices.Num());
+		return false;
+	}
 
-			FIntRect LocalBounds = InBounds;
-			// Convert to inclusive bounds if necessary :
-			if (bInInclusiveBounds)
-			{
-				LocalBounds.Max += FIntPoint(1, 1);
-			}
-			// Intersect with the current key bounds in order to limit the test to only the tracked area :
-			LocalBounds.Clip(KeyExclusiveBounds);
-			if (LocalBounds.Area() <= 0)
-			{
-				return Result;
-			}
-
-			for (int32 Y = LocalBounds.Min.Y; (Y < LocalBounds.Max.Y); ++Y)
-			{
-				for (int32 X = LocalBounds.Min.X; (X < LocalBounds.Max.X); ++X)
-				{
-					int32 ValueIndex = GetValueIndexForKey(FIntPoint(X, Y));
-					Result[ValueIndex] = ValidValueBitIndices[ValueIndex];
-				}
-			}
+	TBitArray<> GetValidValueBitIndicesInBounds(const FIntRect& InBounds, bool bInInclusiveBounds)
+	{
+		TBitArray<> Result(false, ValidValueBitIndices.Num());
 
-			return Result;
+		FIntRect LocalBounds = InBounds;
+		// Convert to inclusive bounds if necessary :
+		if (bInInclusiveBounds)
+		{
+			LocalBounds.Max += FIntPoint(1, 1);
 		}
-
-		TArray<ValueType> GetValidValuesForBitIndices(const TBitArray<>& InBitIndices)
+		// Intersect with the current key bounds in order to limit the test to only the tracked area :
+		LocalBounds.Clip(KeyExclusiveBounds);
+		if (LocalBounds.Area() <= 0)
 		{
-			// Intersect the list and the list of valid elements : 
-			check(InBitIndices.Num() == ValidValueBitIndices.Num());
-			TBitArray<> LocalBitIndices = TBitArray<>::BitwiseAND(ValidValueBitIndices, InBitIndices, EBitwiseOperatorFlags::MinSize);
-
-			TArray<ValueType> Result;
-			Result.Reserve(InBitIndices.CountSetBits());
-			for (TConstSetBitIterator It(LocalBitIndices); It; ++It)
-			{
-				Result.Add(AllValues[It.GetIndex()]);
-			}
 			return Result;
 		}
 
-		FIntRect GetValidValuesBoundsForBitIndices(const TBitArray<>& InBitIndices, bool bInInclusiveBounds)
+		for (int32 Y = LocalBounds.Min.Y; (Y < LocalBounds.Max.Y); ++Y)
 		{
-			// Intersect the list and the list of valid elements : 
-			check(InBitIndices.Num() == ValidValueBitIndices.Num());
-			TBitArray<> LocalBitIndices = TBitArray<>::BitwiseAND(ValidValueBitIndices, InBitIndices, EBitwiseOperatorFlags::MinSize);
-			if (LocalBitIndices.IsEmpty())
+			for (int32 X = LocalBounds.Min.X; (X < LocalBounds.Max.X); ++X)
 			{
-				return FIntRect();
+				int32 ValueIndex = GetValueIndexForKey(FIntPoint(X, Y));
+				Result[ValueIndex] = ValidValueBitIndices[ValueIndex];
 			}
+		}
 
-			const int32 FirstSetBitIndex = InBitIndices.Find(true);
-			const int32 LastSetBitIndex = InBitIndices.FindLast(true);
-			if (FirstSetBitIndex == INDEX_NONE)
-			{
-				check(LastSetBitIndex == INDEX_NONE);
-				return FIntRect();
-			}
+		return Result;
+	}
 
-			bool bIsValid = false;
-			const int32 Stride = KeyExclusiveBounds.Width();
-			FIntRect Bounds(FIntPoint(MAX_int32, MAX_int32), FIntPoint(MIN_int32, MIN_int32));
-			const int32 YMax = FMath::DivideAndRoundDown(LastSetBitIndex, Stride) + 1;
-			for (int32 Y = FirstSetBitIndex / Stride; Y < YMax; ++Y)
-			{
-				const int32 LineFirstSetBitIndex = LocalBitIndices.FindFrom(true, Y * Stride);
-				if (LineFirstSetBitIndex != INDEX_NONE)
-				{
-					const int32 LineLastSetBitIndex = LocalBitIndices.FindLastFrom(true, (Y + 1) * Stride - 1);
-					check(LineLastSetBitIndex != INDEX_NONE);
-					FIntPoint LineMinKey = GetValueKeyForIndex(LineFirstSetBitIndex);
-					FIntPoint LineMaxKey = GetValueKeyForIndex(LineLastSetBitIndex);
-					Bounds.Min = Bounds.Min.ComponentMin(LineMinKey);
-					Bounds.Max = Bounds.Max.ComponentMax(LineMaxKey);
-					bIsValid = true;
-				}
-			}
+	TArray<ValueType> GetValidValuesForBitIndices(const TBitArray<>& InBitIndices)
+	{
+		// Intersect the list and the list of valid elements : 
+		check(InBitIndices.Num() == ValidValueBitIndices.Num());
+		TBitArray<> LocalBitIndices = TBitArray<>::BitwiseAND(ValidValueBitIndices, InBitIndices, EBitwiseOperatorFlags::MinSize);
 
-			if (!bInInclusiveBounds)
-			{
-				Bounds.Max += FIntPoint(1, 1);
-			}
+		TArray<ValueType> Result;
+		Result.Reserve(InBitIndices.CountSetBits());
+		for (TConstSetBitIterator It(LocalBitIndices); It; ++It)
+		{
+			Result.Add(AllValues[It.GetIndex()]);
+		}
+		return Result;
+	}
 
-			check(bIsValid);
-			return Bounds;
+	FIntRect GetValidValuesBoundsForBitIndices(const TBitArray<>& InBitIndices, bool bInInclusiveBounds)
+	{
+		// Intersect the list and the list of valid elements : 
+		check(InBitIndices.Num() == ValidValueBitIndices.Num());
+		TBitArray<> LocalBitIndices = TBitArray<>::BitwiseAND(ValidValueBitIndices, InBitIndices, EBitwiseOperatorFlags::MinSize);
+		if (LocalBitIndices.IsEmpty())
+		{
+			return FIntRect();
+		}
+
+		const int32 FirstSetBitIndex = InBitIndices.Find(true);
+		const int32 LastSetBitIndex = InBitIndices.FindLast(true);
+		if (FirstSetBitIndex == INDEX_NONE)
+		{
+			check(LastSetBitIndex == INDEX_NONE);
+			return FIntRect();
 		}
 
-		TArray<ValueType> GetValidValues()
+		bool bIsValid = false;
+		const int32 Stride = KeyExclusiveBounds.Width();
+		FIntRect Bounds(FIntPoint(MAX_int32, MAX_int32), FIntPoint(MIN_int32, MIN_int32));
+		const int32 YMax = FMath::DivideAndRoundDown(LastSetBitIndex, Stride) + 1;
+		for (int32 Y = FirstSetBitIndex / Stride; Y < YMax; ++Y)
 		{
-			TArray<ValueType> Result;
-			Result.Reserve(ValidValueBitIndices.CountSetBits());
-			for (TConstSetBitIterator It(ValidValueBitIndices); It; ++It)
+			const int32 LineFirstSetBitIndex = LocalBitIndices.FindFrom(true, Y * Stride);
+			if (LineFirstSetBitIndex != INDEX_NONE)
 			{
-				Result.Add(AllValues[It.GetIndex()]);
+				const int32 LineLastSetBitIndex = LocalBitIndices.FindLastFrom(true, (Y + 1) * Stride - 1);
+				check(LineLastSetBitIndex != INDEX_NONE);
+				FIntPoint LineMinKey = GetValueKeyForIndex(LineFirstSetBitIndex);
+				FIntPoint LineMaxKey = GetValueKeyForIndex(LineLastSetBitIndex);
+				Bounds.Min = Bounds.Min.ComponentMin(LineMinKey);
+				Bounds.Max = Bounds.Max.ComponentMax(LineMaxKey);
+				bIsValid = true;
 			}
-			return Result;
 		}
 
-		const TArray<ValueType>& GetAllValues() const { return AllValues; }
-		const TBitArray<>& GetValidValueBitIndices() const { return ValidValueBitIndices; }
+		if (!bInInclusiveBounds)
+		{
+			Bounds.Max += FIntPoint(1, 1);
+		}
 
-	private:
-		FIntRect KeyExclusiveBounds;
-		TArray<ValueType> AllValues;
-		TBitArray<> ValidValueBitIndices;
-	};
+		check(bIsValid);
+		return Bounds;
+	}
 
-	struct FLandscapeComponent2DIndexerKeyFuncs
+	TArray<ValueType> GetValidValues()
 	{
-		static FIntPoint GetKey(ULandscapeComponent* InComponent);
-	};
+		TArray<ValueType> Result;
+		Result.Reserve(ValidValueBitIndices.CountSetBits());
+		for (TConstSetBitIterator It(ValidValueBitIndices); It; ++It)
+		{
+			Result.Add(AllValues[It.GetIndex()]);
+		}
+		return Result;
+	}
 
-	using FLandscapeComponent2DIndexer = T2DIndexer<ULandscapeComponent*, FLandscapeComponent2DIndexerKeyFuncs>;
-	FLandscapeComponent2DIndexer CreateLandscapeComponent2DIndexer(const ULandscapeInfo* InInfo);
+	const TArray<ValueType>& GetAllValues() const { return AllValues; }
+	const TBitArray<>& GetValidValueBitIndices() const { return ValidValueBitIndices; }
 
+private:
+	FIntRect KeyExclusiveBounds;
+	TArray<ValueType> AllValues;
+	TBitArray<> ValidValueBitIndices;
+};
 
-	struct FBlitBuffer2DDesc
-	{
-		int32 BufferSize = 0;  // Buffer size in elements.
-		int32 Stride = 0;      // Buffer stride in elements.  i.e. the distance between the start of each row.
-		FIntRect Rect;         // The 2D region represented by this buffer.  Rect.Area() should be <= BufferSize.
-	};
-
-	/* Copy height values from the R+G channels of the source buffer into a plain uint16 dest buffer.  The source and destination buffers
-	 * represent sub-regions of a shared 2D coordinate system expressed in the Rect parameters.  ClipRect is the desired area of that
-	 * shared coordinate space to limit the operation to.  Copies only values in the the intersection of Src.Rect, Dest.Rect, and ClipRect.
-	 * All FIntRect params are treated as half-open.  .Min is included.  .Max is excluded.
-	 * @param DestBuffer destination buffer
-	 * @param Dest description of the destination buffer
-	 * @param SrcBuffer source buffer
-	 * @param Src description of the source buffer
-	 * @param ClipRect The operation will be limited to this region, in the same shared coordinate space of the rects from the buffer descriptions.
-	 * @return number of elements copied.
-	 */
-	int32 BlitHeightChannelsToUint16(uint16* DestBuffer, const FBlitBuffer2DDesc& Dest, uint32* SrcBuffer, const FBlitBuffer2DDesc& Src, FIntRect ClipRect);
-
-
-	/* Copy 16-bit weight values from the R+G channels of the source buffer, scaled down into a plain uint8 dest buffer.  The source and destination buffers
-	 * represent sub-regions of a shared 2D coordinate system expressed in the Rect parameters.  ClipRect is the desired area of that
-	 * shared coordinate space to limit the operation to.  Copies only values in the the intersection of Src.Rect, Dest.Rect, and ClipRect.
-	 * All FIntRect params are treated as half-open.  .Min is included.  .Max is excluded.
-	 * @param DestBuffer destination buffer
-	 * @param Dest description of the destination buffer
-	 * @param SrcBuffer source buffer
-	 * @param Src description of the source buffer
-	 * @param ClipRect The operation will be limited to this region, in the same shared coordinate space of the rects from the buffer descriptions.
-	 * @return number of elements copied.
-	 */
-	int32 BlitWeightChannelsToUint8(uint8* DestBuffer, const FBlitBuffer2DDesc& Dest, uint32* SrcBuffer, const FBlitBuffer2DDesc& Src, FIntRect ClipRect);
-
-	void SubmitGPUCommands(bool bBlockUntilRTComplete, bool bBlockRTUntilGPUComplete);
+struct FLandscapeComponent2DIndexerKeyFuncs
+{
+	static FIntPoint GetKey(ULandscapeComponent* InComponent);
+};
+
+using FLandscapeComponent2DIndexer = T2DIndexer<ULandscapeComponent*, FLandscapeComponent2DIndexerKeyFuncs>;
+FLandscapeComponent2DIndexer CreateLandscapeComponent2DIndexer(const ULandscapeInfo* InInfo);
+
+
+struct FBlitBuffer2DDesc
+{
+	int32 BufferSize = 0;  // Buffer size in elements.
+	int32 Stride = 0;      // Buffer stride in elements.  i.e. the distance between the start of each row.
+	FIntRect Rect;         // The 2D region represented by this buffer.  Rect.Area() should be <= BufferSize.
+};
+
+/* Copy height values from the R+G channels of the source buffer into a plain uint16 dest buffer.  The source and destination buffers
+ * represent sub-regions of a shared 2D coordinate system expressed in the Rect parameters.  ClipRect is the desired area of that
+ * shared coordinate space to limit the operation to.  Copies only values in the the intersection of Src.Rect, Dest.Rect, and ClipRect.
+ * All FIntRect params are treated as half-open.  .Min is included.  .Max is excluded.
+ * @param DestBuffer destination buffer
+ * @param Dest description of the destination buffer
+ * @param SrcBuffer source buffer
+ * @param Src description of the source buffer
+ * @param ClipRect The operation will be limited to this region, in the same shared coordinate space of the rects from the buffer descriptions.
+ * @return number of elements copied.
+ */
+int32 BlitHeightChannelsToUint16(uint16* DestBuffer, const FBlitBuffer2DDesc& Dest, uint32* SrcBuffer, const FBlitBuffer2DDesc& Src, FIntRect ClipRect);
+
+
+/* Copy 16-bit weight values from the R+G channels of the source buffer, scaled down into a plain uint8 dest buffer.  The source and destination buffers
+ * represent sub-regions of a shared 2D coordinate system expressed in the Rect parameters.  ClipRect is the desired area of that
+ * shared coordinate space to limit the operation to.  Copies only values in the the intersection of Src.Rect, Dest.Rect, and ClipRect.
+ * All FIntRect params are treated as half-open.  .Min is included.  .Max is excluded.
+ * @param DestBuffer destination buffer
+ * @param Dest description of the destination buffer
+ * @param SrcBuffer source buffer
+ * @param Src description of the source buffer
+ * @param ClipRect The operation will be limited to this region, in the same shared coordinate space of the rects from the buffer descriptions.
+ * @return number of elements copied.
+ */
+int32 BlitWeightChannelsToUint8(uint8* DestBuffer, const FBlitBuffer2DDesc& Dest, uint32* SrcBuffer, const FBlitBuffer2DDesc& Src, FIntRect ClipRect);
+
+void SubmitGPUCommands(bool bBlockUntilRTComplete, bool bBlockRTUntilGPUComplete);
 
 } // end namespace UE::Landscape::Private
+
+
+// ----------------------------------------------------------------------------------
+
+namespace UE::Landscape::Grass::Private
+{
+
+uint32 ComputeNamedGrassTypesHash(const TMap<FName, TObjectPtr<ULandscapeGrassType>>& InNamedGrassTypes);
+
+} // namespace UE::Landscape::Grass::Private
\ No newline at end of file
diff --git a/Engine/Source/Runtime/Landscape/Public/LandscapeUtils.h b/Engine/Source/Runtime/Landscape/Public/LandscapeUtils.h
--- a/Engine/Source/Runtime/Landscape/Public/LandscapeUtils.h
+++ b/Engine/Source/Runtime/Landscape/Public/LandscapeUtils.h
@@ -18,6 +18,7 @@ class ALandscapeProxy;
 class FRDGBuilder;
 class FRDGEventName;
 class ULandscapeComponent;
+class ULandscapeGrassType;
 class ULandscapeLayerInfoObject;
 class ULandscapeMaterialInstanceConstant;
 class ULevel;
@@ -370,3 +371,15 @@ LANDSCAPE_API TArray<FName> RetrieveTargetLayerNamesFromMaterial(UMaterialInterf
 LANDSCAPE_API bool LandscapeProxySortPredicate(const TWeakObjectPtr<ALandscapeProxy> APtr, const TWeakObjectPtr<ALandscapeProxy> BPtr);
 
 } // end namespace UE::Landscape
+
+
+// ----------------------------------------------------------------------------------
+
+namespace UE::Landscape::Grass
+{
+
+LANDSCAPE_API TArray<TPair<FName, TObjectPtr<ULandscapeGrassType>>> ConvertToNamedGrassTypesList(const TMap<FName, TObjectPtr<ULandscapeGrassType>>& InNamedGrassTypes);
+/** From the list of named grass types, only retain those with a valid name and grass type. */
+LANDSCAPE_API TMap<FName, TObjectPtr<ULandscapeGrassType>> FilterValidNamedGrassTypes(const TMap<FName, TObjectPtr<ULandscapeGrassType>>& InNamedGrassTypes);
+
+} // namespace UE::Landscape::Grass
