Fixed assert when undoing a landscape paint change right after having deprecated the old data

#rnx

#Synced-CL 50597675

diff --git a/Engine/Source/Runtime/Landscape/Private/Landscape.cpp b/Engine/Source/Runtime/Landscape/Private/Landscape.cpp
--- a/Engine/Source/Runtime/Landscape/Private/Landscape.cpp
+++ b/Engine/Source/Runtime/Landscape/Private/Landscape.cpp
@@ -1075,6 +1075,7 @@ PRAGMA_DISABLE_DEPRECATION_WARNINGS
 					NamedGrassTypes.Emplace(GrassType->GetGrassNameFallback(), GrassType);
 				}
 			}
+			GrassTypes_DEPRECATED.Empty();
 		}
 	}
 PRAGMA_ENABLE_DEPRECATION_WARNINGS
