Landscape - Added support for Heightmap texture compression with landscape subsections and shared textures
* Compression did not account for duplicate rows/columns at the boundaries of subsections or mulit-component textures
* Fixed 5.6 regression where normals along the edges are noticeably wrong in these cases

#Synced-CL 50670310

diff --git a/Engine/Source/Runtime/Landscape/Classes/LandscapeTextureStorageProvider.h b/Engine/Source/Runtime/Landscape/Classes/LandscapeTextureStorageProvider.h
--- a/Engine/Source/Runtime/Landscape/Classes/LandscapeTextureStorageProvider.h
+++ b/Engine/Source/Runtime/Landscape/Classes/LandscapeTextureStorageProvider.h
@@ -180,6 +180,18 @@ public:
 	int32 NumNonStreamingMips = 0;
 	FVector LandscapeGridScale = FVector(ForceInit);
 
+	// Tracks the offset of the texture boundaries (defined by number of subsection/multi-components in texture). Each texture has (NumSubsections * NumComponents) - 1 boundaries for each axis 
+	UPROPERTY()
+	int32 BoundaryOffset = 0;
+
+	// Tracks the number of duplicated column boundaries (based on subsection/multi-component texture count)
+	UPROPERTY()
+	int32 BoundaryCountX = 0;
+
+	// Tracks the number of duplicated row boundaries (based on subsection/multi-component texture count) 
+	UPROPERTY()
+	int32 BoundaryCountY = 0;
+
 	TArray<FLandscapeTexture2DMipMap> Mips;
 	TObjectPtr<UTexture2D> Texture;
 
@@ -188,14 +200,14 @@ public:
 	TObjectPtr<ULandscapeHeightmapTextureEdgeFixup> EdgeFixup;
 
 	#if WITH_EDITORONLY_DATA
-	static ULandscapeTextureStorageProviderFactory* ApplyTo(UTexture2D* InTargetTexture, const FVector& InLandsapeGridScale, int32 InHeightmapCompressionMipThreshold);
-	void UpdateCompressedDataFromSource(UTexture2D* InTargetTexture, const FVector& InLandsapeGridScale, int32 InHeightmapCompressionMipThreshold);
+	static ULandscapeTextureStorageProviderFactory* ApplyTo(UTexture2D* InTargetTexture, const FVector& InLandsapeGridScale, int32 InHeightmapCompressionMipThreshold, int32 InNumSubsections, int32 InSubsectionSizeQuads);
+	void UpdateCompressedDataFromSource(UTexture2D* InTargetTexture, const FVector& InLandsapeGridScale, int32 InHeightmapCompressionMipThreshold, int32 InNumSubsections, int32 InSubsectionSizeQuads);
 	#endif // WITH_EDITORONLY_DATA
 
 	void SetupEdgeFixup(ULandscapeHeightmapTextureEdgeFixup* EdgeFixup);
 
 	// compress the mip data.
-	static void CompressMipToBulkData(int32 MipIndex, int32 MipSizeX, int32 MipSizeY, uint8* SourceData, int32 SourceDataBytes, FByteBulkData& DestBulkData);
+	static void CompressMipToBulkData(int32 MipIndex, int32 MipSizeX, int32 MipSizeY, uint8* SourceData, int32 SourceDataBytes, FByteBulkData& DestBulkData, int32 InMipBoundaryOffset, int32 InBoundaryCountX, int32 InBoundaryCountY);
 	static void CopyMipToBulkData(int32 MipIndex, int32 MipSizeX, int32 MipSizeY, uint8* SourceData, int32 SourceDataBytes, FByteBulkData& DestBulkData);
 
 	// decompress the mip data. Currently operates in-place (though this could be changed in the future)
diff --git a/Engine/Source/Runtime/Landscape/Private/Landscape.cpp b/Engine/Source/Runtime/Landscape/Private/Landscape.cpp
--- a/Engine/Source/Runtime/Landscape/Private/Landscape.cpp
+++ b/Engine/Source/Runtime/Landscape/Private/Landscape.cpp
@@ -2382,12 +2382,12 @@ ULandscapeHeightmapTextureEdgeFixup* ULandscapeComponent::InstallOrUpdateTexture
 		// Install or Update the Storage Factory [editor only]
 		if (StorageFactory == nullptr)
 		{
-			StorageFactory = ULandscapeTextureStorageProviderFactory::ApplyTo(HeightmapTexture, LandscapeGridScale, HeightmapCompressionMipThreshold);
+			StorageFactory = ULandscapeTextureStorageProviderFactory::ApplyTo(HeightmapTexture, LandscapeGridScale, HeightmapCompressionMipThreshold, NumSubsections, SubsectionSizeQuads);
 		}
 		else
 		{
 			// since different platforms may change compression settings / thresholds, we should update the compression each time
-			StorageFactory->UpdateCompressedDataFromSource(HeightmapTexture, LandscapeGridScale, HeightmapCompressionMipThreshold);
+			StorageFactory->UpdateCompressedDataFromSource(HeightmapTexture, LandscapeGridScale, HeightmapCompressionMipThreshold, NumSubsections, SubsectionSizeQuads);
 		}
 #endif // WITH_EDITOR
 		StorageFactory->SetupEdgeFixup(EdgeFixup);
diff --git a/Engine/Source/Runtime/Landscape/Private/LandscapeTextureStorageProvider.cpp b/Engine/Source/Runtime/Landscape/Private/LandscapeTextureStorageProvider.cpp
--- a/Engine/Source/Runtime/Landscape/Private/LandscapeTextureStorageProvider.cpp
+++ b/Engine/Source/Runtime/Landscape/Private/LandscapeTextureStorageProvider.cpp
@@ -24,6 +24,16 @@ LandscapeTextureStorageProvider.cpp: Alternative Texture Storage for Landscape T
 #define PROVIDER_DEBUG_LOG_DETAIL(...) do {} while(0)
 #endif // ENABLE_LANDSCAPE_PROVIDER_DEBUG_SPEW
 
+namespace UE::Landscape::Private
+{
+	// Returns true if pixel position exists at a boundary offset (subsection/multi-component texture boundary)
+	static inline bool IsDuplicateCoord(int32 XY, int32 MipBoundaryOffset)
+	{
+		checkSlow(FMath::IsPowerOfTwo(MipBoundaryOffset));
+		return (XY != 0) && ((XY & (MipBoundaryOffset - 1)) == 0);
+	}
+}
+
 FLandscapeTextureMipEdgeOverrideProvider::FLandscapeTextureMipEdgeOverrideProvider(ULandscapeHeightmapTextureEdgeFixup* InEdgeFixup, UTexture2D* InTexture)
 	: FTextureMipDataProvider(InTexture, ETickState::GetMips, ETickThread::Async)
 {
@@ -448,7 +458,7 @@ bool ULandscapeTextureStorageProviderFactory::GetInitialMipData(int32 FirstMipTo
 
 
 #if WITH_EDITORONLY_DATA
-ULandscapeTextureStorageProviderFactory* ULandscapeTextureStorageProviderFactory::ApplyTo(UTexture2D* InTargetTexture, const FVector& InLandsapeGridScale, int32 InHeightmapCompressionMipThreshold)
+ULandscapeTextureStorageProviderFactory* ULandscapeTextureStorageProviderFactory::ApplyTo(UTexture2D* InTargetTexture, const FVector& InLandsapeGridScale, int32 InHeightmapCompressionMipThreshold, int32 InNumSubsections, int32 InSubsectionSizeQuads)
 {
 	check(InTargetTexture);
 	check(InTargetTexture->Source.IsValid())
@@ -465,27 +475,34 @@ ULandscapeTextureStorageProviderFactory* ULandscapeTextureStorageProviderFactory
 		InTargetTexture->AddAssetUserData(Factory);
 	}
 
-	Factory->UpdateCompressedDataFromSource(InTargetTexture, InLandsapeGridScale, InHeightmapCompressionMipThreshold);
+	Factory->UpdateCompressedDataFromSource(InTargetTexture, InLandsapeGridScale, InHeightmapCompressionMipThreshold, InNumSubsections, InSubsectionSizeQuads);
 
 	return Factory;
 }
 
-void ULandscapeTextureStorageProviderFactory::UpdateCompressedDataFromSource(UTexture2D* InTargetTexture, const FVector& InLandsapeGridScale, int32 InHeightmapCompressionMipThreshold)
+void ULandscapeTextureStorageProviderFactory::UpdateCompressedDataFromSource(UTexture2D* InTargetTexture, const FVector& InLandsapeGridScale, int32 InHeightmapCompressionMipThreshold, int32 InNumSubsections, int32 InSubsectionSizeQuads)
 {
 	check(this->Texture == InTargetTexture);
 
-	EPixelFormat Format = PF_B8G8R8A8;
-
 	int32 Width = InTargetTexture->Source.GetSizeX();
 	int32 Height = InTargetTexture->Source.GetSizeY();
 	int32 MipCount = InTargetTexture->Source.GetNumMips();
 
-	uint32 SrcBpp = GPixelFormats[Format].BlockBytes;
-	uint32 SrcPitch = Width * SrcBpp;
-
 	// need this to properly calculate normals
 	this->LandscapeGridScale = InLandsapeGridScale;
 
+	// check if this is a shared texture (multi-component texture)
+	const int32 ComponentWidth = InNumSubsections * (InSubsectionSizeQuads + 1);
+	check(ComponentWidth > 0);
+	const int32 NumComponentsX = Width / ComponentWidth;
+	const int32 NumComponentsY = Height / ComponentWidth;
+
+	// if the texture is shared, or has subsections, there is a duplicate row/column at every BoundaryOffset 
+	const bool bHasDuplicateData = NumComponentsX > 1 || NumComponentsY > 1 || InNumSubsections > 1;
+	this->BoundaryOffset = InSubsectionSizeQuads + 1;
+	this->BoundaryCountX = bHasDuplicateData ? (NumComponentsX * InNumSubsections - 1) : 0;
+	this->BoundaryCountY = bHasDuplicateData ? (NumComponentsY * InNumSubsections - 1) : 0;
+
 	// calculate number of non-streaming mips
 	// TODO [chris.tchou] : we could make this calculation platform specific, like Texture2D does.
 	// We would have to calculate it during serialization, when we know the target platform.
@@ -530,7 +547,7 @@ void ULandscapeTextureStorageProviderFactory::UpdateCompressedDataFromSource(UTe
 		else
 		{
 			Mip->bCompressed = true;
-			CompressMipToBulkData(MipIndex, MipWidth, MipHeight, MipData.GetData(), MipData.Num(), Mip->BulkData);
+			CompressMipToBulkData(MipIndex, MipWidth, MipHeight, MipData.GetData(), MipData.Num(), Mip->BulkData, BoundaryOffset >> MipIndex, BoundaryCountX, BoundaryCountY);
 		}
 
 		MipWidth = FMath::Max(MipWidth / 2, 1);
@@ -822,39 +839,78 @@ void ULandscapeTextureStorageProviderFactory::CopyMipToBulkData(int32 MipIndex,
 	DestBulkData.Unlock();
 }
 
-void ULandscapeTextureStorageProviderFactory::CompressMipToBulkData(int32 MipIndex, int32 MipSizeX, int32 MipSizeY, uint8* SourceData, int32 SourceDataBytes, FByteBulkData& DestBulkData)
+void ULandscapeTextureStorageProviderFactory::CompressMipToBulkData(int32 MipIndex, int32 MipSizeX, int32 MipSizeY, uint8* SourceData, int32 SourceDataBytes, FByteBulkData& DestBulkData, int32 InMipBoundaryOffset, int32 InBoundaryCountX, int32 InBoundaryCountY)
 {
 	TRACE_CPUPROFILER_EVENT_SCOPE(ULandscapeTextureStorageProviderFactory::CompressMipToBulkData);
 
 	DestBulkData.Lock(LOCK_READ_WRITE);
 
-	int32 TotalPixels = MipSizeX * MipSizeY;
-	check(SourceDataBytes == TotalPixels * 4);
-	check(TotalPixels >= 16);	// shouldn't be used on very small mips
+	int32 DestPixels = MipSizeX * MipSizeY;
+	check(SourceDataBytes == DestPixels * 4);
+
+	const int32 NumSourcePixels = SourceDataBytes / 4;
 
-	// DestData consists of a 16 bit height per pixel, then an 8:8 normal per edge pixel
-	int32 DestBytes = (TotalPixels + (MipSizeX + MipSizeY) * 2 - 4) * 2;
+	// If the texture is shared or has subsections, there is a duplicate row/column at every BoundaryOffset
+	const bool bHasDuplicateData = (InBoundaryCountX > 0 || InBoundaryCountY > 0) && InMipBoundaryOffset > 0;
+	check(InBoundaryCountX >= 0 && InBoundaryCountY >= 0);
+
+	// Remove the duplicated row/columns along subsection boundaries in the destination texture 
+	const int32 DestWidth = MipSizeX - InBoundaryCountX;
+	const int32 DestHeight = MipSizeY - InBoundaryCountY;
+	const int32 NumDestPixels = DestWidth * DestHeight;
+	check(NumDestPixels >= 16);	// shouldn't be used on very small mips
+
+	/* DestData consists of a 16 bit height per pixel, then an 8:8 normal per edge pixel, edge normals are handled separately and do not remove duplicates
+	*  EdgePixels = (MipSizeX + MipSizeY) * 2 - 4 
+	*  TotalBytes = 2 * NumDestPixels + 2 * EdgePixels */
+	int32 DestBytes = (NumDestPixels + (MipSizeX + MipSizeY) * 2 - 4) * 2;
 	uint8* DestData = DestBulkData.Realloc(DestBytes);
 
 	// delta encode the heights -- this (usually) greatly reduces the variance in the data, which makes it compress much better on disk when package compression is applied.
 	uint16 LastHeight = 32768;
-	int32 DestOffset = 0;
-	for (int32 SourceOffset = 0; SourceOffset < SourceDataBytes; SourceOffset += 4)
+	int32 DuplicateRowsSkipped = 0;
+
+	for (int32 Y = 0; Y < MipSizeY; Y++)
 	{
-		// texture data is stored as BGRA, or [normal x, height low bits, height high bits, normal y]
-		uint16 Height = SourceData[SourceOffset + 2] * 256 + SourceData[SourceOffset + 1];
-		uint16 DeltaHeight = Height - LastHeight;
-		LastHeight = Height;
-
-		// store delta height
-		DestData[DestOffset + 0] = DeltaHeight >> 8;
-		DestData[DestOffset + 1] = DeltaHeight & 0xff;
-		DestOffset += 2;
+		int32 SrcLineOffsetInPixels = Y * MipSizeX;
+		int32 DestLineOffsetInPixels = (Y - DuplicateRowsSkipped) * DestWidth;
+		FColor* Src = (FColor*)&SourceData[SrcLineOffsetInPixels * 4] ;
+		uint8* Dst = &DestData[DestLineOffsetInPixels * 2];
+
+		// Ignore duplicate rows, they will be regenerated
+		if (bHasDuplicateData && UE::Landscape::Private::IsDuplicateCoord(Y, InMipBoundaryOffset))
+		{
+			DuplicateRowsSkipped++;
+			continue;
+		}
+
+		for (int32 X = 0; X < MipSizeX; X++)
+		{
+			// Ignore duplicate columns, they will be regenerated
+			if (bHasDuplicateData && UE::Landscape::Private::IsDuplicateCoord(X, InMipBoundaryOffset))
+			{
+				Src += 1;
+				continue;
+			}
+
+			// texture data is stored as BGRA, or [normal x, height low bits, height high bits, normal y]
+			uint16 Height = Src->R * 256 + Src->G;
+			uint16 DeltaHeight = Height - LastHeight;
+			LastHeight = Height;
+
+			// store delta height
+			Dst[0] = DeltaHeight >> 8;
+			Dst[1] = DeltaHeight & 0xff;
+			Dst += 2;
+			Src += 1;
+		}
 	}
 
-	int32 DeltaCount = DestOffset;
+	// edge normals are written after all the height data 
+	int32 DestOffset = NumDestPixels * 2;
 
 	// capture normals along the edge (delta encoded clockwise starting from top left)
+	// for efficiency, edge normals do not remove duplicate row/columns
 	uint8 LastNormalX = 128;
 	uint8 LastNormalY = 128;
 
@@ -1043,12 +1099,24 @@ void ULandscapeTextureStorageProviderFactory::DecompressMip(uint8* SourceData, i
 
 	check(Mip.bCompressed);		// uncompressed should be handled outside of this function
 
-	int32 Width = Mip.SizeX;
-	int32 Height = Mip.SizeY;
-	int32 TotalPixels = Width * Height;
-	int32 BorderPixels = (Width + Height) * 2 - 4;
-	check(SourceDataBytes == (TotalPixels + BorderPixels) * 2); // 2 bytes (height) for each pixel, plus 2 bytes (normal x/y) for each border pixel
-	check(DestDataBytes == TotalPixels * 4);
+	const int32 DestWidth = Mip.SizeX;
+	const int32 DestHeight = Mip.SizeY;
+
+	// If the texture is shared, or there are subsections, need to add back a duplicate row/column at every BoundaryOffset 
+	const int32 MipBoundaryOffset = BoundaryOffset >> MipIndex;
+	const bool bHasDuplicateData = (BoundaryCountX > 0 || BoundaryCountY > 0) && MipBoundaryOffset > 0;
+	check(!bHasDuplicateData || FMath::IsPowerOfTwo(MipBoundaryOffset));
+
+	const int32 SrcWidth = DestWidth - BoundaryCountX;
+	const int32 SrcHeight = DestHeight - BoundaryCountY;
+	check(SrcHeight >= 0 && SrcWidth >= 0);
+
+	int NumSrcPixels = SrcWidth * SrcHeight;
+
+	// Verify the Src heightmap data has removed the duplicate row/columns
+	// Edge normals do not strip the duplicated row/columns
+	check(SourceDataBytes == (NumSrcPixels + (DestWidth + DestHeight) * 2 - 4) * 2); // 2 bytes (height) for each pixel, plus 2 bytes (normal x/y) for each border pixel
+	check(DestDataBytes == DestWidth * DestHeight * 4);
 
 	// save some multiplying by premultiplying the grid scales, mip scale and ZScale
 	FVector2f PremultU16 = CalculatePremultU16(MipIndex, LandscapeGridScale);
@@ -1059,20 +1127,31 @@ void ULandscapeTextureStorageProviderFactory::DecompressMip(uint8* SourceData, i
 
 	// partial normal results recorded for the previous line
 	TArray<FVector3f, TInlineAllocator<512>> PrevLinePartialNormals;
-	PrevLinePartialNormals.SetNumZeroed(Width);
+	PrevLinePartialNormals.SetNumZeroed(DestWidth);
 
+	int32 DuplicateRowsSkipped = 0;
+	int32 PreviousRowOffset = 1;
 	// iterate each line
-	for (int32 Y = 0; Y < Height; Y++)
+	for (int32 Y = 0; Y < DestHeight; Y++)
 	{
-		int32 LineOffsetInPixels = Y * Width;
-		uint8* Src = &SourceData[LineOffsetInPixels * 2];
-		FColor* Dst = (FColor*)&DestData[LineOffsetInPixels * 4];
+		// Offset the src since the data has removed all of the duplicate rows
+		int32 SrcLineOffsetInPixels = (Y - DuplicateRowsSkipped) * SrcWidth;
+		int32 DestLineOffsetInPixels = Y * DestWidth;
+		uint8* Src = &SourceData[SrcLineOffsetInPixels * 2];
+		FColor* Dst = (FColor*)&DestData[DestLineOffsetInPixels * 4];
 
 		if (Y == 0)
 		{
 			// just decode heights for the first line (normals don't matter they will be stomped below)
-			for (int32 X = 0; X < Width; X++)
+			for (int32 X = 0; X < DestWidth; X++)
 			{
+				// Skip the duplicate column. Duplicate columns will be copied once at the end of function
+				if (bHasDuplicateData && UE::Landscape::Private::IsDuplicateCoord(X, MipBoundaryOffset))
+				{
+					Dst++;
+					continue;
+				}
+
 				uint16 DeltaHeight = Src[0] * 256 + Src[1];
 				CC += DeltaHeight;
 				*Dst = FColor(CC >> 8, CC & 0xff, 128, 128);
@@ -1082,6 +1161,18 @@ void ULandscapeTextureStorageProviderFactory::DecompressMip(uint8* SourceData, i
 		}
 		else
 		{
+			// Duplicate rows will be copied once at the end of function
+			if (bHasDuplicateData && UE::Landscape::Private::IsDuplicateCoord(Y, MipBoundaryOffset))
+			{
+				DuplicateRowsSkipped++;
+				PreviousRowOffset = 2;
+				continue;
+			}
+
+			// Duplicate rows are not considered in the normal calculation
+			// The top pixel is two rows above after skipping a duplicate row
+			const int32 TTOffset = DestWidth * PreviousRowOffset;
+
 			// compute initial values (first pixel)
 			FVector3f P1(ForceInitToZero), P01(ForceInitToZero);	// previous quad N1 and (N0+N1) normals
 			uint16 TT;												// previous quad TT height
@@ -1091,15 +1182,24 @@ void ULandscapeTextureStorageProviderFactory::DecompressMip(uint8* SourceData, i
 				*Dst = FColor(CC >> 8, CC & 0xff, 128, 128);
 
 				// load TT for first pixel (becomes TL for second pixel)
-				TT = DecodeHeightU16(Dst + 0 - Width);
+				TT = DecodeHeightU16(Dst - TTOffset);
 
 				Src += 2;
 				Dst++;
 			}
 
+			int32 PrevColumnOffset = -1;
 			// rest of the pixels in the line
-			for (int32 X = 1; X < Width; X++)
-			{
+			for (int32 X = 1; X < DestWidth; X++)
+			{ 
+				// Duplicate column data is copied once at the end of function
+				if (bHasDuplicateData && UE::Landscape::Private::IsDuplicateCoord(X, MipBoundaryOffset))
+				{
+					Dst++;
+					PrevColumnOffset = -2;
+					continue;
+				}
+
 				// re-use previous pixel TT and CC as this pixel TL and LL
 				uint16 TL = TT;
 				uint16 LL = CC;
@@ -1109,7 +1209,7 @@ void ULandscapeTextureStorageProviderFactory::DecompressMip(uint8* SourceData, i
 				CC += DeltaHeight;
 
 				// load TT
-				TT = DecodeHeightU16(Dst + 0 - Width);
+				TT = DecodeHeightU16(Dst - TTOffset);
 
 				// 2) Write Height at CC (normals get written during processing of the next line)
 				*Dst = FColor(CC >> 8, CC & 0xff, 128, 128);
@@ -1120,16 +1220,16 @@ void ULandscapeTextureStorageProviderFactory::DecompressMip(uint8* SourceData, i
 				FVector3f N01 = N0 + N1;
 
 				// 4) Complete Normal calculation for TL - this takes the partial result from the previous line and fills in the rest
-				FVector3f TL_Normal = PrevLinePartialNormals[X - 1] + P1 + N01;
+				FVector3f TL_Normal = PrevLinePartialNormals[X + PrevColumnOffset] + P1 + N01;
 				FastNormalize(TL_Normal);
-				
+
 				// 5) Write Normal for TL
-				Dst[-Width - 1].B = static_cast<uint8>(FMath::Clamp((TL_Normal.X * 127.5f + 127.5f), 0.0f, 255.0f));
-				Dst[-Width - 1].A = static_cast<uint8>(FMath::Clamp((TL_Normal.Y * 127.5f + 127.5f), 0.0f, 255.0f));
+				Dst[-TTOffset + PrevColumnOffset].B = static_cast<uint8>(FMath::Clamp((TL_Normal.X * 127.5f + 127.5f), 0.0f, 255.0f));
+				Dst[-TTOffset + PrevColumnOffset].A = static_cast<uint8>(FMath::Clamp((TL_Normal.Y * 127.5f + 127.5f), 0.0f, 255.0f));
 
 				// 6) Store Partial Normal for LL in PrevLinePartialNormals (P0 + P1 + N0) - the rest will be filled in when processing the next line
 				FVector3f LL_PartialNormal = P01 + N0;
-				PrevLinePartialNormals[X - 1] = LL_PartialNormal;
+				PrevLinePartialNormals[X + PrevColumnOffset] = LL_PartialNormal;
 
 				// pass normals to next pixel
 				P1 = N1;
@@ -1137,7 +1237,9 @@ void ULandscapeTextureStorageProviderFactory::DecompressMip(uint8* SourceData, i
 
 				Src += 2;
 				Dst++;
+				PrevColumnOffset = -1;
 			}
+			PreviousRowOffset = 1;
 		}
 	}
 
@@ -1145,13 +1247,13 @@ void ULandscapeTextureStorageProviderFactory::DecompressMip(uint8* SourceData, i
 	{
 		TRACE_CPUPROFILER_EVENT_SCOPE(EdgeNormalFixup);
 
-		uint8* Src = &SourceData[TotalPixels * 2];
+		uint8* Src = &SourceData[NumSrcPixels * 2];
 		uint8 LastNormalX = 128;
 		uint8 LastNormalY = 128;
 
-		auto DecodeNormal = [&LastNormalX, &LastNormalY, DestData, Width, &Src](int32 X, int32 Y)
+		auto DecodeNormal = [&LastNormalX, &LastNormalY, DestData, DestWidth, &Src](int32 X, int32 Y)
 			{
-				int32 DestOffset = (Y * Width + X) * 4;
+				int32 DestOffset = (Y * DestWidth + X) * 4;
 				LastNormalX += Src[0];
 				LastNormalY += Src[1];
 				DestData[DestOffset + 0] = LastNormalX;
@@ -1159,27 +1261,55 @@ void ULandscapeTextureStorageProviderFactory::DecompressMip(uint8* SourceData, i
 				Src += 2;
 			};
 
-		for (int32 X = 0; X < Width; X++)		// [0 ... Width-1], 0
+		for (int32 X = 0; X < DestWidth; X++)		// [0 ... Width-1], 0
 		{
 			DecodeNormal(X, 0);
 		}
 
-		for (int32 Y = 1; Y < Height; Y++)		// Width-1, [1 ... Height-1]
+		for (int32 Y = 1; Y < DestHeight; Y++)		// Width-1, [1 ... Height-1]
 		{
-			DecodeNormal(Width - 1, Y);
+			DecodeNormal(DestWidth - 1, Y);
 		}
 
-		for (int32 X = Width - 2; X >= 0; X--)	// [Width-2 ... 0], Height-1
+		for (int32 X = DestWidth - 2; X >= 0; X--)	// [Width-2 ... 0], Height-1
 		{
-			DecodeNormal(X, Height - 1);
+			DecodeNormal(X, DestHeight - 1);
 		}
 
-		for (int32 Y = Height - 2; Y >= 1; Y--)	// 0, [Height-2 ... 1]
+		for (int32 Y = DestHeight - 2; Y >= 1; Y--)	// 0, [Height-2 ... 1]
 		{
 			DecodeNormal(0, Y);
 		}
-
 		check(Src == &SourceData[SourceDataBytes]);
+
+		// Copy all duplicate row/column data once all height/normal data is set
+		if (bHasDuplicateData)
+		{
+			// For each duplicate row, copy the pixels above 
+			for (int32 Y = MipBoundaryOffset; Y < DestHeight; Y += MipBoundaryOffset)
+			{
+				check(UE::Landscape::Private::IsDuplicateCoord(Y, MipBoundaryOffset));
+				FColor* Dst = (FColor*)&DestData[Y * DestWidth * 4];
+
+				for (int32 X = 0; X < DestWidth; X++)
+				{
+					*Dst = Dst[-DestWidth];
+					Dst++;
+				}
+			}
+
+			// For each duplicate column, copy the pixels to the left
+			for (int32 X = MipBoundaryOffset; X < DestWidth; X += MipBoundaryOffset)
+			{
+				check(UE::Landscape::Private::IsDuplicateCoord(X, MipBoundaryOffset));
+				FColor* Dst = (FColor*)&DestData[X * 4];
+				for (int32 Y = 0; Y < DestHeight; Y++)
+				{
+					*Dst = Dst[-1];
+					Dst += DestWidth;
+				}
+			}
+		}
 	}
 }
 
