Landscape - Fixed heightmap compression assert with high index mips that remove row/column duplication
* Texture's with duplicate data (mutli-component/shared) drop the boundary duplication at a low texture resolution

#Synced-CL 50777462

diff --git a/Engine/Source/Runtime/Landscape/Private/LandscapeTextureStorageProvider.cpp b/Engine/Source/Runtime/Landscape/Private/LandscapeTextureStorageProvider.cpp
--- a/Engine/Source/Runtime/Landscape/Private/LandscapeTextureStorageProvider.cpp
+++ b/Engine/Source/Runtime/Landscape/Private/LandscapeTextureStorageProvider.cpp
@@ -847,18 +847,22 @@ void ULandscapeTextureStorageProviderFactory::CompressMipToBulkData(int32 MipInd
 
 	int32 DestPixels = MipSizeX * MipSizeY;
 	check(SourceDataBytes == DestPixels * 4);
+	check(DestPixels >= 16);	// shouldn't be used on very small mips
 
 	const int32 NumSourcePixels = SourceDataBytes / 4;
 
 	// If the texture is shared or has subsections, there is a duplicate row/column at every BoundaryOffset
-	const bool bHasDuplicateData = (InBoundaryCountX > 0 || InBoundaryCountY > 0) && InMipBoundaryOffset > 0;
 	check(InBoundaryCountX >= 0 && InBoundaryCountY >= 0);
+	// High index mips drop the row/column duplication at low resolutions. Offset is only relevant when duplication occurs atleast after every other pixel (Offset >= 2)
+	const bool bHasDuplicateData = (InBoundaryCountX > 0 || InBoundaryCountY > 0) && InMipBoundaryOffset > 1;
+	const int32 MipBoundaryCountX = bHasDuplicateData ? InBoundaryCountX : 0;
+	const int32 MipBoundaryCountY = bHasDuplicateData ? InBoundaryCountY : 0;
 
 	// Remove the duplicated row/columns along subsection boundaries in the destination texture 
-	const int32 DestWidth = MipSizeX - InBoundaryCountX;
-	const int32 DestHeight = MipSizeY - InBoundaryCountY;
+	const int32 DestWidth = MipSizeX - MipBoundaryCountX;
+	const int32 DestHeight = MipSizeY - MipBoundaryCountY;
+	check(DestWidth >= 0 && DestHeight >= 0);
 	const int32 NumDestPixels = DestWidth * DestHeight;
-	check(NumDestPixels >= 16);	// shouldn't be used on very small mips
 
 	/* DestData consists of a 16 bit height per pixel, then an 8:8 normal per edge pixel, edge normals are handled separately and do not remove duplicates
 	*  EdgePixels = (MipSizeX + MipSizeY) * 2 - 4 
@@ -1104,11 +1108,14 @@ void ULandscapeTextureStorageProviderFactory::DecompressMip(uint8* SourceData, i
 
 	// If the texture is shared, or there are subsections, need to add back a duplicate row/column at every BoundaryOffset 
 	const int32 MipBoundaryOffset = BoundaryOffset >> MipIndex;
-	const bool bHasDuplicateData = (BoundaryCountX > 0 || BoundaryCountY > 0) && MipBoundaryOffset > 0;
+	// High index mips drop the row/column duplication at low resolutions. Offset is only relevant when duplication occurs at least after every other pixel (Offset >= 2)
+	const bool bHasDuplicateData = (BoundaryCountX > 0 || BoundaryCountY > 0) && MipBoundaryOffset > 1;
+	const int32 MipBoundaryCountX = bHasDuplicateData ? BoundaryCountX : 0;
+	const int32 MipBoundaryCountY = bHasDuplicateData ? BoundaryCountY : 0;
 	check(!bHasDuplicateData || FMath::IsPowerOfTwo(MipBoundaryOffset));
 
-	const int32 SrcWidth = DestWidth - BoundaryCountX;
-	const int32 SrcHeight = DestHeight - BoundaryCountY;
+	const int32 SrcWidth = DestWidth - MipBoundaryCountX;
+	const int32 SrcHeight = DestHeight - MipBoundaryCountY;
 	check(SrcHeight >= 0 && SrcWidth >= 0);
 
 	int NumSrcPixels = SrcWidth * SrcHeight;
