Include landscape HLOD override material in the HLOD hash.  Allows change in HLOD override material to trigger HLOD update.

#Synced-CL 50537731

diff --git a/Engine/Source/Runtime/Landscape/Private/LandscapeHLODBuilder.cpp b/Engine/Source/Runtime/Landscape/Private/LandscapeHLODBuilder.cpp
--- a/Engine/Source/Runtime/Landscape/Private/LandscapeHLODBuilder.cpp
+++ b/Engine/Source/Runtime/Landscape/Private/LandscapeHLODBuilder.cpp
@@ -49,11 +49,11 @@ bool ULandscapeHLODBuilder::ComputeHLODHash(FHLODHashBuilder& HashBuilder, const
 		return false;
 	}
 	
-	// Base lanscape HLOD key, changing this will force a rebuild of all landscape HLODs
+	// Base landscape HLOD key, changing this will force a rebuild of all landscape HLODs
 	FString HLODBaseKey = "51EDA4343E26409DA6ABEACE14C138D1";
 	HashBuilder.HashField(HLODBaseKey, TEXT("LandscapeHLODBuilderBaseKey"));
 
-	ALandscapeProxy* LSProxy = LSComponent->GetLandscapeProxy();
+	const ALandscapeProxy* LSProxy = LSComponent->GetLandscapeProxy();
 
 	// LS LOD setup
 	HashBuilder.HashField(LSProxy->GetLODScreenSizeArray(), TEXT("LODScreenSize"));
@@ -69,18 +69,25 @@ bool ULandscapeHLODBuilder::ComputeHLODHash(FHLODHashBuilder& HashBuilder, const
 
 	// LS Materials
 	{
-		TArray<UMaterialInterface*> UsedMaterials;
-		if (LSProxy->bUseDynamicMaterialInstance)
+		if (LSProxy->HLODMaterialOverride)
 		{
-			HashBuilder << LSComponent->MaterialInstancesDynamic;
+			HashBuilder << LSProxy->HLODMaterialOverride;
 		}
 		else
 		{
-			HashBuilder << LSComponent->MaterialInstances;
+			if (LSProxy->bUseDynamicMaterialInstance)
+			{
+				HashBuilder << LSComponent->MaterialInstancesDynamic;
+			}
+			else
+			{
+				HashBuilder << LSComponent->MaterialInstances;
+			}
+
+			HashBuilder << LSComponent->OverrideMaterial;
 		}
 
-		HashBuilder << LSComponent->OverrideMaterial;
-		HashBuilder << LSComponent->OverrideHoleMaterial;
+		HashBuilder << LSComponent->OverrideHoleMaterial;  //Hole material matters regardless of which base material is used.
 	}
 
 	// Nanite enabled?
