Prevent update tracked component flow if grass generation is not enabled on any proxy.

#rnx

#Synced-CL 51084747

diff --git a/Engine/Source/Runtime/Landscape/Private/LandscapeGrassMapsBuilder.cpp b/Engine/Source/Runtime/Landscape/Private/LandscapeGrassMapsBuilder.cpp
--- a/Engine/Source/Runtime/Landscape/Private/LandscapeGrassMapsBuilder.cpp
+++ b/Engine/Source/Runtime/Landscape/Private/LandscapeGrassMapsBuilder.cpp
@@ -786,7 +786,8 @@ bool FLandscapeGrassMapsBuilder::CanCurrentlyRender() const
 void FLandscapeGrassMapsBuilder::AmortizedUpdateGrassMaps(
 	const TArray<FVector>& Cameras,
 	bool bPrioritizeCreation,
-	bool bAllowStartGrassMapGeneration)
+	bool bAllowStartGrassMapGeneration,
+	bool bAnyLandscapeActorHasGrassEnabled)
 {
 #if !WITH_EDITOR
 	if (!GGrassMapUseRuntimeGeneration) // in cooked builds, we don't run any updates at all unless runtime generation is enabled
@@ -795,6 +796,11 @@ void FLandscapeGrassMapsBuilder::AmortizedUpdateGrassMaps(
 	}
 #endif // WITH_EDITOR
 
+	if (!bAnyLandscapeActorHasGrassEnabled && !StreamingCount && !RenderingCount && !AsyncFetchCount && !PopulatedCount)
+	{
+		return;
+	}
+
 	if (!CanEverRender())
 	{
 		return; // if we can never ever render, don't bother to do anything here
diff --git a/Engine/Source/Runtime/Landscape/Private/LandscapeGrassMapsBuilder.h b/Engine/Source/Runtime/Landscape/Private/LandscapeGrassMapsBuilder.h
--- a/Engine/Source/Runtime/Landscape/Private/LandscapeGrassMapsBuilder.h
+++ b/Engine/Source/Runtime/Landscape/Private/LandscapeGrassMapsBuilder.h
@@ -69,7 +69,7 @@ public:
 	// if Cameras is empty, it considers all distances to be zero for update purposes:
 	//   - it won't evict anything for distance
 	//   - it will not start tracking (building grass maps) for any components
-	void AmortizedUpdateGrassMaps(const TArray<FVector>& Cameras, bool bPrioritizeCreation, bool bAllowStartGrassMapGeneration);
+	void AmortizedUpdateGrassMaps(const TArray<FVector>& Cameras, bool bPrioritizeCreation, bool bAllowStartGrassMapGeneration, bool bAnyLandscapeActorHasGrassEnabled);
 
 	// synchronous build of grassmaps for a specific set of components
 	// returns true if it successfully makes all LandscapeComponents have up to date Grass Maps.
diff --git a/Engine/Source/Runtime/Landscape/Private/LandscapeSubsystem.cpp b/Engine/Source/Runtime/Landscape/Private/LandscapeSubsystem.cpp
--- a/Engine/Source/Runtime/Landscape/Private/LandscapeSubsystem.cpp
+++ b/Engine/Source/Runtime/Landscape/Private/LandscapeSubsystem.cpp
@@ -96,6 +96,7 @@ static FAutoConsoleVariableRef CVarLandscapeMaxSimultaneousMultithreadNaniteBuil
 	LandscapeMaxSimultaneousMultithreadNaniteBuilds,
 	TEXT("Max number of simultaneous Nanite static mesh tasks (-1 = unlimited )"));
 
+extern int32 GGrassEnable;
 extern int32 GGrassMapUseRuntimeGeneration;
 
 DECLARE_CYCLE_STAT(TEXT("LandscapeSubsystem Tick"), STAT_LandscapeSubsystemTick, STATGROUP_Landscape);
@@ -755,12 +756,19 @@ void ULandscapeSubsystem::Tick(float DeltaTime)
 	// run early update on Proxies, and determine if all of the proxies are ready for grass generation to start
 	bool bAllProxiesReadyForGrassMapGeneration = true;
 	bool bAllProxiesRuntimeGrassMapsDisabled = true;
+	bool bAnyLandscapeActorHasGrassEnabled = false;
 
 #if WITH_EDITOR
 	TSet<ALandscape*> DisallowedGrassTickLandscapes;
+#endif // WITH_EDITOR
+
 	for (TObjectPtr<ALandscape> ActorPtr : LandscapeActors)
 	{
 		ALandscape* Landscape = ActorPtr.Get();
+
+		bAnyLandscapeActorHasGrassEnabled |= Landscape->GetGrassEnabled();
+
+#if WITH_EDITOR
 		if (ULandscapeInfo* LandscapeInfo = Landscape->GetLandscapeInfo())
 		{
 			const bool bLandscapeIsUpToDate = Landscape->IsUpToDate();
@@ -777,8 +785,10 @@ void ULandscapeSubsystem::Tick(float DeltaTime)
 				DisallowedGrassTickLandscapes.Add(Landscape);
 			}
 		}
-	}
 #endif // WITH_EDITOR
+	}
+
+	bAnyLandscapeActorHasGrassEnabled &= !!GGrassEnable;
 
 	static TArray<ALandscapeProxy*> ActiveProxies;
 	{
@@ -829,7 +839,7 @@ void ULandscapeSubsystem::Tick(float DeltaTime)
 
 	bool bAllowStartGrassMapGeneration = bAllProxiesReadyForGrassMapGeneration && !bGrassMapGenerationDisabled && (!bIsGameWorld || GGrassMapUseRuntimeGeneration);
 
-	GrassMapsBuilder->AmortizedUpdateGrassMaps(Cameras ? *Cameras : TArray<FVector>(), bIsGrassCreationPrioritized, bAllowStartGrassMapGeneration);
+	GrassMapsBuilder->AmortizedUpdateGrassMaps(Cameras ? *Cameras : TArray<FVector>(), bIsGrassCreationPrioritized, bAllowStartGrassMapGeneration, bAnyLandscapeActorHasGrassEnabled);
 
 	// run edge fixup for streaming proxies, which only exist in partitioned worlds
 	// TODO [chris.tchou] : remove this check once we support edge fixup on non-partitioned worlds (requires removing heightmap sharing between components)
