Landscape Editor - Move Component Tool - Fixed incorrect proxy transform when moving landscape components from one Level to another

#Synced-CL 51092926

diff --git a/Engine/Source/Runtime/Landscape/Private/LandscapeEdit.cpp b/Engine/Source/Runtime/Landscape/Private/LandscapeEdit.cpp
--- a/Engine/Source/Runtime/Landscape/Private/LandscapeEdit.cpp
+++ b/Engine/Source/Runtime/Landscape/Private/LandscapeEdit.cpp
@@ -5473,7 +5473,7 @@ ALandscapeProxy* ULandscapeInfo::MoveComponentsToProxy(const TArray<ULandscapeCo
 	TRACE_CPUPROFILER_EVENT_SCOPE(ULandscapeInfo::MoveComponentsToProxy);
 	
 	ALandscape* Landscape = LandscapeActor.Get();
-	check(Landscape != nullptr);
+	check(Landscape != nullptr && LandscapeProxy != nullptr);
 	
 	struct FCompareULandscapeComponentBySectionBase
 	{
@@ -5553,6 +5553,9 @@ ALandscapeProxy* ULandscapeInfo::MoveComponentsToProxy(const TArray<ULandscapeCo
 		ULandscapeComponent* FirstComponent = *TargetSelectedComponents.CreateConstIterator();
 		LandscapeProxy->GetRootComponent()->SetWorldLocationAndRotation(FirstComponent->GetComponentLocation(), FirstComponent->GetComponentRotation());
 		LandscapeProxy->SetSectionBase(FirstComponent->GetSectionBase());
+
+		// update the landscape info's proxy transform after changing the location and section base
+		FixupProxiesTransform(/*bDirty =*/true);
 	}
 
 	// Hide(unregister) the new landscape if owning level currently in hidden state
