Improve large landscape edit performance.  Skip processing inactive texture readback objects earlier.

#Synced-CL 51134463

diff --git a/Engine/Source/Runtime/Landscape/Private/LandscapeEditLayers.cpp b/Engine/Source/Runtime/Landscape/Private/LandscapeEditLayers.cpp
--- a/Engine/Source/Runtime/Landscape/Private/LandscapeEditLayers.cpp
+++ b/Engine/Source/Runtime/Landscape/Private/LandscapeEditLayers.cpp
@@ -5087,6 +5087,12 @@ void ALandscape::ResolveLayersHeightmapTexture(
 		check(LandscapeProxy);
 		if (FLandscapeEditLayerReadback** CPUReadback = LandscapeProxy->HeightmapsCPUReadback.Find(Heightmap))
 		{
+			// Only process the readback object if it has active tasks.
+			if (!(*CPUReadback)->HasActiveTasks())
+			{
+				continue;
+			}
+
 			const bool bChanged = ResolveLayersTexture(MapHelper, *CPUReadback, Heightmap, InOutComponentReadbackResults, /*bIsWeightmap = */false);
 			if (bChanged)
 			{
@@ -6721,6 +6727,12 @@ void ALandscape::ResolveLayersWeightmapTexture(
 		check(LandscapeProxy);
 		if (FLandscapeEditLayerReadback** CPUReadback = LandscapeProxy->WeightmapsCPUReadback.Find(Weightmap))
 		{
+			// Only process the readback object if it has active tasks.
+			if (!(*CPUReadback)->HasActiveTasks())
+			{
+				continue;
+			}
+
 			const bool bChanged = ResolveLayersTexture(MapHelper, *CPUReadback, Weightmap, InOutComponentReadbackResults, /*bIsWeightmap = */true);
 			if (bChanged)
 			{
@@ -7363,7 +7375,7 @@ void ALandscape::UpdateLayersContent(bool bInWaitForStreaming, bool bInSkipMonit
 		}
 	}
 
-	const bool bProcessReadbacks = FLandscapeEditLayerReadback::HasWork();
+	const bool bProcessReadbacks = FLandscapeEditLayerReadback::HasAnyGlobalWork();
 
 	// User triggered change has been completely processed, resetting user triggered flag on all components.
 	if(IsUpToDate())
@@ -8048,7 +8060,7 @@ bool ALandscape::IsUpToDate() const
 #if WITH_EDITORONLY_DATA
 	if (!IsTemplate() && GetWorld() != nullptr && !GetWorld()->IsGameWorld())
 	{
-		return LayerContentUpdateModes == 0 && !FLandscapeEditLayerReadback::HasWork();
+		return LayerContentUpdateModes == 0 && !FLandscapeEditLayerReadback::HasAnyGlobalWork();
 	}
 #endif
 
diff --git a/Engine/Source/Runtime/Landscape/Private/LandscapeEditReadback.cpp b/Engine/Source/Runtime/Landscape/Private/LandscapeEditReadback.cpp
--- a/Engine/Source/Runtime/Landscape/Private/LandscapeEditReadback.cpp
+++ b/Engine/Source/Runtime/Landscape/Private/LandscapeEditReadback.cpp
@@ -434,7 +434,7 @@ void FLandscapeEditLayerReadback::ReleaseCompletedResults(int32 InResultNum)
 	TaskHandles.RemoveAt(0, InResultNum, EAllowShrinking::No);
 }
 
-bool FLandscapeEditLayerReadback::HasWork()
+bool FLandscapeEditLayerReadback::HasAnyGlobalWork()
 {
 	return GReadbackTaskPool.AllocCount > 0;
 }
diff --git a/Engine/Source/Runtime/Landscape/Private/LandscapeEditReadback.h b/Engine/Source/Runtime/Landscape/Private/LandscapeEditReadback.h
--- a/Engine/Source/Runtime/Landscape/Private/LandscapeEditReadback.h
+++ b/Engine/Source/Runtime/Landscape/Private/LandscapeEditReadback.h
@@ -58,8 +58,13 @@ public:
 	/* Release results. Pass in a value which is no more than the value returned by GetCompletedResultNum(). */
 	void ReleaseCompletedResults(int32 InResultNum);
 
-	/** Returns true if there are any queued tasks that haven't been released. */
-	static bool HasWork();
+	bool HasActiveTasks() const
+	{
+		return !TaskHandles.IsEmpty();
+	}
+
+	/** Returns true if there are any queued tasks that haven't been released for any readback object. */
+	static bool HasAnyGlobalWork();
 
 	/** Optional call to allow garbage collection of any task resources and check for leaks. Call once per frame. */
 	static void GarbageCollectTasks();
diff --git a/Engine/Source/Runtime/Landscape/Private/LandscapeGroup.cpp b/Engine/Source/Runtime/Landscape/Private/LandscapeGroup.cpp
--- a/Engine/Source/Runtime/Landscape/Private/LandscapeGroup.cpp
+++ b/Engine/Source/Runtime/Landscape/Private/LandscapeGroup.cpp
@@ -823,7 +823,7 @@ void FLandscapeGroup::TickEdgeFixup(ULandscapeSubsystem* LandscapeSubsystem, boo
 
 
 #if WITH_EDITORONLY_DATA
-		if (FLandscapeEditLayerReadback::HasWork())
+		if (FLandscapeEditLayerReadback::HasAnyGlobalWork())
 		{
 			GROUP_DEBUG_LOG_DETAIL(TEXT("-- Patch (All Components) -------- PAUSED (deferred while any readbacks pending)"));
 		}
