* Prevent crash when trying to call one of the landscape's  RenderHeightmap/Weightmap(s) functions from an application without a valid rendering device
* Cleaned up detection of valid rendering devices in landscape code

#Synced-CL 51210907

diff --git a/Engine/Source/Runtime/Landscape/Private/Landscape.cpp b/Engine/Source/Runtime/Landscape/Private/Landscape.cpp
--- a/Engine/Source/Runtime/Landscape/Private/Landscape.cpp
+++ b/Engine/Source/Runtime/Landscape/Private/Landscape.cpp
@@ -3839,6 +3839,12 @@ bool ALandscape::RenderMergedTextureInternal(const FTransform& InRenderAreaWorld
 
 	TRACE_CPUPROFILER_EVENT_SCOPE(Landscape_RenderMergedTextureInternal);
 
+	if (!FApp::CanEverRender())
+	{
+		UE_LOG(LogLandscape, Log, TEXT("RenderMergedTexture : Cannot render anything on landscape (%s) : the app doesn't have a valid rendering device "), *GetFullName());
+		return false;
+	}
+
 	ULandscapeInfo* Info = GetLandscapeInfo();
 	if (Info == nullptr)
 	{
diff --git a/Engine/Source/Runtime/Landscape/Private/LandscapeGrass.cpp b/Engine/Source/Runtime/Landscape/Private/LandscapeGrass.cpp
--- a/Engine/Source/Runtime/Landscape/Private/LandscapeGrass.cpp
+++ b/Engine/Source/Runtime/Landscape/Private/LandscapeGrass.cpp
@@ -1105,7 +1105,7 @@ bool ULandscapeComponent::CanRenderGrassMap() const
 
 	// Check we can render
 	const UWorld* World = GetWorld();
-	if (GUsingNullRHI || !World || !SceneProxy)
+	if (!CanEverRender() || !World || !SceneProxy)
 	{
 		GRASS_DEBUG_LOG(TEXT("GrassMap No SceneProxy for %s"), *GetName());
 		return false;
diff --git a/Engine/Source/Runtime/Landscape/Private/LandscapeGrassMapsBuilder.cpp b/Engine/Source/Runtime/Landscape/Private/LandscapeGrassMapsBuilder.cpp
--- a/Engine/Source/Runtime/Landscape/Private/LandscapeGrassMapsBuilder.cpp
+++ b/Engine/Source/Runtime/Landscape/Private/LandscapeGrassMapsBuilder.cpp
@@ -762,16 +762,10 @@ void FLandscapeGrassMapsBuilder::UnregisterComponent(const ULandscapeComponent*
 	}
 }
 
-// false if this program instance will never be able to render grass
-bool FLandscapeGrassMapsBuilder::CanEverRender() const
-{
-	return FApp::CanEverRender() && !GUsingNullRHI;
-}
-
 // false if the world can not currently render the grass (but this may change later, for example if preview modes are modified)
 bool FLandscapeGrassMapsBuilder::CanCurrentlyRender() const
 {
-	if (CanEverRender())
+	if (FApp::CanEverRender())
 	{
 		check(World && World->Scene);
 		// GPU scene is required by landscape fixed grid vertex factory
@@ -801,7 +795,7 @@ void FLandscapeGrassMapsBuilder::AmortizedUpdateGrassMaps(
 		return;
 	}
 
-	if (!CanEverRender())
+	if (!FApp::CanEverRender())
 	{
 		return; // if we can never ever render, don't bother to do anything here
 	}
diff --git a/Engine/Source/Runtime/Landscape/Private/LandscapeGrassMapsBuilder.h b/Engine/Source/Runtime/Landscape/Private/LandscapeGrassMapsBuilder.h
--- a/Engine/Source/Runtime/Landscape/Private/LandscapeGrassMapsBuilder.h
+++ b/Engine/Source/Runtime/Landscape/Private/LandscapeGrassMapsBuilder.h
@@ -83,9 +83,6 @@ public:
 
 private:
 
-	// false if this program instance will never be able to render grass
-	bool CanEverRender() const;
-
 	// false if the world can not currently render the grass (but this may change later, for example if preview modes are modified)
 	bool CanCurrentlyRender() const;
 
