Landscape Editor - Avoid MapCheck warning when auto-adding Visibility Layer to landscape proxies
* Adds visibility in PostLoad, since the proxy doesn't need to cross reference data on the main ALandscape. This avoids FixupSharedData from handling the visibility layer upgrade

#Synced-CL 51293768

diff --git a/Engine/Source/Runtime/Landscape/Private/Landscape.cpp b/Engine/Source/Runtime/Landscape/Private/Landscape.cpp
--- a/Engine/Source/Runtime/Landscape/Private/Landscape.cpp
+++ b/Engine/Source/Runtime/Landscape/Private/Landscape.cpp
@@ -5028,6 +5028,14 @@ PRAGMA_ENABLE_DEPRECATION_WARNINGS
 		HLODTextureSizePolicy = ELandscapeHLODTextureSizePolicy::AutomaticSize;
 		HLODMeshSourceLODPolicy = ELandscapeHLODMeshSourceLODPolicy::AutomaticLOD;
 	}
+
+	// Silently add a visibility layer to each proxy. The cost is minimal, since components do not store weightmap info until data is written to the layer
+	// Target layers are no longer auto generated during material compilation so this guarantees the visibility layer is always found
+	const FName VisibilityLayerName = (ALandscapeProxy::VisibilityLayer != nullptr) ? ALandscapeProxy::VisibilityLayer->GetLayerName() : FName();
+	if (!IsTemplate() && !VisibilityLayerName.IsNone() && !HasTargetLayer(VisibilityLayerName))
+	{
+		AddTargetLayer(VisibilityLayerName, FLandscapeTargetLayerSettings(), /*bPostEditChange =*/false);
+	}
 #endif // WITH_EDITOR
 }
 
@@ -5641,14 +5649,6 @@ void ALandscapeProxy::FixupSharedData(ALandscape* Landscape, const bool bMapChec
 		SynchronizedProperties.Add(GET_MEMBER_NAME_CHECKED(ALandscapeProxy, TargetLayers));
 	}
 
-	// Silently add a visibility layer to each proxy. The cost is minimal, since components do not store weightmap info until data is written to the layer
-	// Target layers are no longer auto generated during material compilation so this guarantees the visibility layer is always found
-	const FName VisibilityLayerName = (ALandscapeProxy::VisibilityLayer != nullptr) ? ALandscapeProxy::VisibilityLayer->GetLayerName() : FName();
-	if (!VisibilityLayerName.IsNone() && !HasTargetLayer(VisibilityLayerName))
-	{
-		AddTargetLayer(VisibilityLayerName, FLandscapeTargetLayerSettings(), /*bPostEditChange =*/false);
-	}
-
 	if (bUpdated)
 	{
 		// In cases where LandscapeInfo is not fully ready yet, we forward the provided ALandscape. If ALandscape is available in LandscapeInfo, we let the object function naturally.
