Fix for long freezes when saving in editor due to grass map build step waiting for landscape components that have grass disabled.

#rnx
#Synced-CL 51462505

diff --git a/Engine/Source/Runtime/Landscape/Private/LandscapeGrassMapsBuilder.cpp b/Engine/Source/Runtime/Landscape/Private/LandscapeGrassMapsBuilder.cpp
--- a/Engine/Source/Runtime/Landscape/Private/LandscapeGrassMapsBuilder.cpp
+++ b/Engine/Source/Runtime/Landscape/Private/LandscapeGrassMapsBuilder.cpp
@@ -1156,16 +1156,33 @@ bool FLandscapeGrassMapsBuilder::BuildGrassMapsNowForComponents(
 
 			if (State->Stage == EComponentStage::NotReady)
 			{
-				// if it's not ready because of shader reasons
-				if (!FailedComponents.Contains(Component) && !Component->CanRenderGrassMap())
+				if (!Component->GetLandscapeActor()->ShouldGenerateGrass())
+				{
+					// nothing to do for this component
+					++UpToDateCount;
+				}
+				else if (!FailedComponents.Contains(Component))
 				{
 #if WITH_EDITOR
-					// in editor, try to force compilation to complete
-					UE::Landscape::Grass::Private::CompileGrassMapShader(State->Component);
+					// if it's not ready because the shaders are not ready
 					if (!Component->CanRenderGrassMap())
+					{
+						// try to force the shaders to compile
+						UE::Landscape::Grass::Private::CompileGrassMapShader(State->Component);
+						if (Component->CanRenderGrassMap())
+						{
+							// succeeded, let it try to complete the grass (the next UpdateTrackedComponents should pull it out of NotReady state if it can)
+						}
+						else
+						{
+							// failed to compile shaders, it's definitely failed for real
+							FailedComponents.Add(Component);
+						}
+					}
+					else
 #endif // WITH_EDITOR
 					{
-						// failed to compile shaders... we won't be able to build the grass map for this component
+						// failed for other unspecified reasons
 						FailedComponents.Add(Component);
 					}
 				}
