Avoid executing PostEditChangeProperty before the landscape has been registered with the world.

#Synced-CL 50939444

diff --git a/Engine/Source/Runtime/Landscape/Classes/LandscapeProxy.h b/Engine/Source/Runtime/Landscape/Classes/LandscapeProxy.h
--- a/Engine/Source/Runtime/Landscape/Classes/LandscapeProxy.h
+++ b/Engine/Source/Runtime/Landscape/Classes/LandscapeProxy.h
@@ -1273,15 +1273,15 @@ PRAGMA_ENABLE_DEPRECATION_WARNINGS
 
 #if WITH_EDITOR
 	// Copy properties from parent Landscape actor
-	LANDSCAPE_API void CopySharedProperties(ALandscapeProxy* InLandscape);
+	LANDSCAPE_API void CopySharedProperties(ALandscapeProxy* InLandscape, bool bInExecutePostEditChangeProperty = false);
 
 	/**
 	* Enforce property sharing and requirements related to the parent Landscape actor
 	* @param InLandscape the parent landscape to use as a template.
-	* @param bInRunPostEditChangeProperty When true, runs PostEditChangeProperty for each property synchronized
+	* @param bInExecutePostEditChangeProperty When true, executes PostEditChangeProperty for each property synchronized
 	* @return An array containing all synchronized properties names (properties that were different from the provided landscape).
 	*/
-	LANDSCAPE_API TArray<FName> SynchronizeSharedProperties(ALandscapeProxy* InLandscape, bool bInRunPostEditChangeProperty = false);
+	LANDSCAPE_API TArray<FName> SynchronizeSharedProperties(ALandscapeProxy* InLandscape, bool bInExecutePostEditChangeProperty = false);
 
 	// Returns true if the property is shared.
 	LANDSCAPE_API bool IsSharedProperty(const FName& InPropertyName) const;
@@ -1299,7 +1299,7 @@ PRAGMA_ENABLE_DEPRECATION_WARNINGS
 	virtual bool IsSharedPropertyOverridden(const FName& InPropertyName) const { return false; }
 
 	// Modifies the override state of the property given as argument.
-	virtual void SetSharedPropertyOverride(const FName& InPropertyName, const bool bIsOverriden) { }
+	virtual void SetSharedPropertyOverride(const FName& InPropertyName, const bool bIsOverriden, const bool bInExecutePostEditChangeProperty = true) { }
 
 	/** Delegate that will be called whenever FixupSharedData is called, to inject some custom logic */
 	static FOnLandscapeProxyFixupSharedDataDelegate::RegistrationType& OnLandscapeProxyFixupSharedData() { return OnLandscapeProxyFixupSharedDataDelegate; }
@@ -1379,14 +1379,14 @@ PRAGMA_ENABLE_DEPRECATION_WARNINGS
 	LANDSCAPE_API void ChangedPhysMaterial();
 
 	// Converts shared properties overridden before explicit override, asking the user how to proceed.
-	LANDSCAPE_API void UpgradeSharedProperties(ALandscape* InParentLandscape);
+	LANDSCAPE_API void UpgradeSharedProperties(ALandscape* InParentLandscape, const bool bInExecutePostEditChangeProperty = true);
 
 	// Displays the obsolete edit layers within the LandscapeComponents in a actionable MapCheck log
 	// Asks users to copy obsolete layer data or delete for good.
  	void DisplayObsoleteEditLayerMapCheck(ALandscape* InParentLandscape);
 
 	// Assign only mismatching data and mark proxy package dirty
-	LANDSCAPE_API void FixupSharedData(ALandscape* Landscape, bool bMapCheck = false);
+	LANDSCAPE_API void FixupSharedData(ALandscape* Landscape, bool bMapCheck = false, const bool bInExecutePostEditChangeProperty = true);
 
 	UE_DEPRECATED(5.7, "Use SetSectionBase instead.")
 	LANDSCAPE_API void SetAbsoluteSectionBase(FIntPoint SectionOffset);
diff --git a/Engine/Source/Runtime/Landscape/Classes/LandscapeStreamingProxy.h b/Engine/Source/Runtime/Landscape/Classes/LandscapeStreamingProxy.h
--- a/Engine/Source/Runtime/Landscape/Classes/LandscapeStreamingProxy.h
+++ b/Engine/Source/Runtime/Landscape/Classes/LandscapeStreamingProxy.h
@@ -73,7 +73,7 @@ public:
 #if WITH_EDITOR
 	//~ Begin ALandscapeProxy Interface
 	virtual bool IsSharedPropertyOverridden(const FName& InPropertyName) const override;
-	virtual void SetSharedPropertyOverride(const FName& InPropertyName, const bool bIsOverridden) override;
+	virtual void SetSharedPropertyOverride(const FName& InPropertyName, const bool bIsOverridden, const bool bInExecutePostEditChangeProperty = true) override;
 	//~ End ALandscapeProxy Interface
 
 	void RequestDeferredLayerContentUpdateModes(ELandscapeLayerUpdateMode InLayerContentUpdateModes);
diff --git a/Engine/Source/Runtime/Landscape/Private/Landscape.cpp b/Engine/Source/Runtime/Landscape/Private/Landscape.cpp
--- a/Engine/Source/Runtime/Landscape/Private/Landscape.cpp
+++ b/Engine/Source/Runtime/Landscape/Private/Landscape.cpp
@@ -4713,7 +4713,7 @@ void ALandscapeProxy::PostLoad()
 	if ((GetLandscapeInfo() != nullptr) && (LandscapeActor != nullptr) && (LandscapeActor != this))
 	{
 		const bool bMapCheck = true;
-		FixupSharedData(LandscapeActor, bMapCheck);
+		FixupSharedData(LandscapeActor, bMapCheck, false /*bInExecutePostEditChangeProperty*/);
 	}
 #endif // WITH_EDITOR
 
@@ -5107,8 +5107,27 @@ void ALandscapeProxy::Destroyed()
 
 namespace UE::Landscape::Private
 {
-	static bool CopyProperty(FProperty* InProperty, UObject* InSourceObject, UObject* InDestinationObject)
+	static bool CopyProperty(FProperty* InProperty, UObject* InSourceObject, UObject* InDestinationObject, bool bInExecutePostEditChangeProperty)
 	{
+		// Do not execute the post edit change if the containing level has not finished post loading.  
+		// This will trigger Registration of the components, before the owning level of the world has been determined.
+		// This avoids registring landscape proxies in the incorrect landscape info object.
+		if(bInExecutePostEditChangeProperty)
+		{
+			ULevel* Level = nullptr;
+			if (ALandscapeProxy* LandscapeProxy = Cast<ALandscapeProxy>(InDestinationObject))
+			{
+				Level = LandscapeProxy->GetLevel();
+			}
+			else if (USceneComponent* SceneComponent = Cast<USceneComponent>(InDestinationObject))
+			{
+				Level = SceneComponent->GetOwner()->GetLevel();
+			}
+
+			bInExecutePostEditChangeProperty = Level == nullptr || !Level->HasAnyFlags(EObjectFlags::RF_NeedPostLoad);
+			ensureAlwaysMsgf(bInExecutePostEditChangeProperty, TEXT("Executing CopyProperty with bInExecutePostEditChangeProperty set to true before the level has post loaded would trigger registration in the world before it's owning world has been determined."));
+		}
+
 		void* SrcValuePtr = InProperty->ContainerPtrToValuePtr<void>(InSourceObject);
 		void* DestValuePtr = InProperty->ContainerPtrToValuePtr<void>(InDestinationObject);
 
@@ -5118,11 +5137,18 @@ namespace UE::Landscape::Private
 		}
 
 		InProperty->CopyCompleteValue(DestValuePtr, SrcValuePtr);
+		if (bInExecutePostEditChangeProperty)
+		{
+			// Some properties may need additional processing (ex: LandscapeMaterial), notify the proxy of the change.
+			FPropertyChangedEvent PropertyChangedEvent(InProperty);
+			InDestinationObject->PostEditChangeProperty(PropertyChangedEvent);
+			InDestinationObject->Modify();
+		}
 
 		return true;
 	}
 
-	static bool CopyPostEditPropertyByName(const TWeakObjectPtr<ALandscapeProxy>& InLandscapeProxy, const TWeakObjectPtr<ALandscape>& InParentLandscape, const FName& InPropertyName)
+	static bool CopyPropertyByName(const TWeakObjectPtr<ALandscapeProxy>& InLandscapeProxy, const TWeakObjectPtr<ALandscape>& InParentLandscape, const FName& InPropertyName, bool bInExecutePostEditChangeProperty)
 	{
 		if (!InLandscapeProxy.IsValid() || !InParentLandscape.IsValid())
 		{
@@ -5143,14 +5169,7 @@ namespace UE::Landscape::Private
 			return false;
 		}
 
-		CopyProperty(PropertyToCopy, InParentLandscape.Get(), InLandscapeProxy.Get());
-
-		// Some properties may need additional processing (ex: LandscapeMaterial), notify the proxy of the change.
-		FPropertyChangedEvent PropertyChangedEvent(PropertyToCopy);
-		InLandscapeProxy->PostEditChangeProperty(PropertyChangedEvent);
-
-		InLandscapeProxy->Modify();
-
+		CopyProperty(PropertyToCopy, InParentLandscape.Get(), InLandscapeProxy.Get(), bInExecutePostEditChangeProperty);
 		return true;
 	}
 
@@ -5196,7 +5215,7 @@ namespace UE::Landscape::Private
 	}
 } // namespace UE::Landscape::Private
 
-void ALandscapeProxy::CopySharedProperties(ALandscapeProxy* InLandscape)
+void ALandscapeProxy::CopySharedProperties(ALandscapeProxy* InLandscape, bool bInExecutePostEditChangeProperty)
 {
 	SynchronizeUnmarkedSharedProperties(InLandscape);
 
@@ -5211,12 +5230,12 @@ void ALandscapeProxy::CopySharedProperties(ALandscapeProxy* InLandscape)
 
 		if (IsSharedProperty(Property))
 		{
-			UE::Landscape::Private::CopyProperty(Property, InLandscape, this);
+			UE::Landscape::Private::CopyProperty(Property, InLandscape, this, bInExecutePostEditChangeProperty);
 		}
 	}
 }
 
-TArray<FName> ALandscapeProxy::SynchronizeSharedProperties(ALandscapeProxy* InLandscape, bool bInRunPostEditChangeProperty)
+TArray<FName> ALandscapeProxy::SynchronizeSharedProperties(ALandscapeProxy* InLandscape, bool bInExecutePostEditChangeProperty)
 {
 	TArray<FName> SynchronizedProperties = SynchronizeUnmarkedSharedProperties(InLandscape);
 
@@ -5237,14 +5256,7 @@ TArray<FName> ALandscapeProxy::SynchronizeSharedProperties(ALandscapeProxy* InLa
 
 			// Caller may need a PostEditChangeProperty to trigger other landscape updates (Eg. Dirty Render state, nanite, materials)
 			// PostEditChangeProperty updates are unnecessary when syncing properties during new landscape creation
-			if (bInRunPostEditChangeProperty)
-			{
-				UE::Landscape::Private::CopyPostEditPropertyByName(this, GetLandscapeActor(), Property->GetFName());
-			}
-			else
-			{
-				UE::Landscape::Private::CopyProperty(Property, InLandscape, this);
-			}
+			UE::Landscape::Private::CopyProperty(Property, InLandscape, this, bInExecutePostEditChangeProperty);
 		}
 	}
 
@@ -5391,7 +5403,7 @@ bool ALandscapeStreamingProxy::IsSharedPropertyOverridden(const FName& InPropert
 	return OverriddenSharedProperties.Contains(InPropertyName);
 }
 
-void ALandscapeStreamingProxy::SetSharedPropertyOverride(const FName& InPropertyName, const bool bIsOverridden)
+void ALandscapeStreamingProxy::SetSharedPropertyOverride(const FName& InPropertyName, const bool bIsOverridden, const bool bInExecutePostEditChangeProperty)
 {
 	check(IsSharedProperty(InPropertyName));
 
@@ -5413,7 +5425,7 @@ void ALandscapeStreamingProxy::SetSharedPropertyOverride(const FName& InProperty
 		}
 		else
 		{
-			UE::Landscape::Private::CopyPostEditPropertyByName(LandscapeProxy, ParentLandscape, InPropertyName);
+			UE::Landscape::Private::CopyPropertyByName(LandscapeProxy, ParentLandscape, InPropertyName, bInExecutePostEditChangeProperty);
 		}
 
 		OverriddenSharedProperties.Remove(InPropertyName);
@@ -5448,7 +5460,7 @@ void ALandscapeStreamingProxy::PerformDeferredLayerContentUpdateModesRequest()
 	}
 }
 
-void ALandscapeProxy::UpgradeSharedProperties(ALandscape* InParentLandscape)
+void ALandscapeProxy::UpgradeSharedProperties(ALandscape* InParentLandscape, const bool bInExecutePostEditChangeProperty)
 {
 	TArray<FName> SynchronizedProperties;
 	bool bOpenMapCheckWindow = false;
@@ -5460,20 +5472,20 @@ void ALandscapeProxy::UpgradeSharedProperties(ALandscape* InParentLandscape)
 		FProperty* Property = *PropertyIterator;
 
 		if (Property == nullptr)
-{
+		{
 			continue;
 		}
 
 		if (IsPropertyInherited(Property) && !Property->Identical_InContainer(this, InParentLandscape))
 		{
 			SynchronizedProperties.Emplace(Property->GetFName());
-			UE::Landscape::Private::CopyProperty(Property, InParentLandscape, this);
+			UE::Landscape::Private::CopyProperty(Property, InParentLandscape, this, bInExecutePostEditChangeProperty);
 		}
 		else if (IsPropertyOverridable(Property) && !IsSharedPropertyOverridden(Property->GetFName()) && !Property->Identical_InContainer(this, InParentLandscape))
 		{
 			if (CVarSilenceSharedPropertyDeprecationFixup->GetBool())
 			{
-				SetSharedPropertyOverride(Property->GetFName(), true);
+				SetSharedPropertyOverride(Property->GetFName(), true, bInExecutePostEditChangeProperty);
 			}
 			else
 			{
@@ -5490,19 +5502,19 @@ void ALandscapeProxy::UpgradeSharedProperties(ALandscape* InParentLandscape)
 					->AddToken(FUObjectToken::Create(this, FText::FromString(GetActorNameOrLabel())))
 					->AddToken(FTextToken::Create(FText::Format(LOCTEXT("MapCheck_Message_LandscapeProxy_UpgradeSharedProperties", "Contains a property ({0}) different from parent's landscape actor. Please select between "), FText::FromString(PropertyName.ToString()))))
 					->AddToken(FActionToken::Create(LOCTEXT("MapCheck_OverrideProperty", "Override property"), LOCTEXT("MapCheck_OverrideProperty_Desc", "Keeping the current value and marking the property as overriding the parent landscape's value."),
-						FOnActionTokenExecuted::CreateLambda([LandscapeProxy, PropertyName]()
+						FOnActionTokenExecuted::CreateLambda([LandscapeProxy, PropertyName, bInExecutePostEditChangeProperty]()
 							{
 								if (LandscapeProxy.IsValid())
 								{
-									LandscapeProxy->SetSharedPropertyOverride(PropertyName, true);
+									LandscapeProxy->SetSharedPropertyOverride(PropertyName, true, bInExecutePostEditChangeProperty);
 								}
 							}),
 						/*bInSingleUse = */true))
 					->AddToken(FTextToken::Create(LOCTEXT("MapCheck_Message_LandscapeProxy_UpgradeSharedProperties_Or", " or ")))
 					->AddToken(FActionToken::Create(LOCTEXT("MapCheck_InheritProperty", "Inherit from parent landscape"), LOCTEXT("MapCheck_InheritProperty_Desc", "Copying the parent landscape's value for this property."),
-							FOnActionTokenExecuted::CreateLambda([LandscapeProxy, ParentLandscape, PropertyName]()
+							FOnActionTokenExecuted::CreateLambda([LandscapeProxy, ParentLandscape, PropertyName, bInExecutePostEditChangeProperty]()
 								{
-									UE::Landscape::Private::CopyPostEditPropertyByName(LandscapeProxy, ParentLandscape, PropertyName);
+									UE::Landscape::Private::CopyPropertyByName(LandscapeProxy, ParentLandscape, PropertyName, bInExecutePostEditChangeProperty);
 								}),
 							/*bInSingleUse = */true))
 					->AddToken(FMapErrorToken::Create(FMapErrors::LandscapeComponentPostLoad_Warning));
@@ -5591,7 +5603,7 @@ void ALandscapeProxy::DisplayObsoleteEditLayerMapCheck(ALandscape* InParentLands
 	}
 }
 
-void ALandscapeProxy::FixupSharedData(ALandscape* Landscape, const bool bMapCheck)
+void ALandscapeProxy::FixupSharedData(ALandscape* Landscape, const bool bMapCheck, const bool bInExecutePostEditChangeProperty)
 {
 	if (Landscape == nullptr)
 	{
@@ -5615,12 +5627,12 @@ void ALandscapeProxy::FixupSharedData(ALandscape* Landscape, const bool bMapChec
 		((GetLinkerCustomVersion(FFortniteReleaseBranchCustomObjectVersion::GUID) < FFortniteReleaseBranchCustomObjectVersion::LandscapeSharedPropertiesEnforcement)
 		|| (GetLinkerCustomVersion(FFortniteMainBranchObjectVersion::GUID) < FFortniteMainBranchObjectVersion::LandscapeBodyInstanceAsSharedProperty)))
 	{
-		UpgradeSharedProperties(Landscape);
+		UpgradeSharedProperties(Landscape, bInExecutePostEditChangeProperty);
 		bUpgradeSharedPropertiesPerformed = true;
 	}
 	else
 	{
-		SynchronizedProperties = SynchronizeSharedProperties(Landscape, /*bInRunPostEditChangeProperty =*/true);
+		SynchronizedProperties = SynchronizeSharedProperties(Landscape, bInExecutePostEditChangeProperty);
 		bUpdated |= !SynchronizedProperties.IsEmpty();
 	}
 
@@ -6414,7 +6426,7 @@ void ULandscapeInfo::RegisterLandscapeActorWithProxyInternal(ALandscapeProxy* Pr
 		// This can happen on Undo, when the component requests an update while ALandscape is unregistered 
 		StreamingProxy->PerformDeferredLayerContentUpdateModesRequest();
 
-		StreamingProxy->FixupSharedData(Landscape, bMapCheck);
+		StreamingProxy->FixupSharedData(Landscape, bMapCheck, true /*bInExecutePostEditChangeProperty*/);
 #endif // WITH_EDITOR
 
 		Proxy->bIsLandscapeActorRegisteredWithLandscapeInfo = true;
@@ -6429,7 +6441,7 @@ void ULandscapeInfo::RegisterLandscapeActorWithProxyInternal(ALandscapeProxy* Pr
 		//  This is a one-time thing because at some point during development, the target layer data was deprecated and the deprecation turned somewhat sour :S
 		ForEachLandscapeProxy([this, Landscape, bMapCheck](ALandscapeProxy* Proxy)
 		{
-			Proxy->FixupSharedData(Landscape, bMapCheck);
+			Proxy->FixupSharedData(Landscape, bMapCheck, true /*bInExecutePostEditChangeProperty*/);
 			return true;
 		});
 	}
diff --git a/Engine/Source/Runtime/Landscape/Private/LandscapeEdit.cpp b/Engine/Source/Runtime/Landscape/Private/LandscapeEdit.cpp
--- a/Engine/Source/Runtime/Landscape/Private/LandscapeEdit.cpp
+++ b/Engine/Source/Runtime/Landscape/Private/LandscapeEdit.cpp
@@ -5939,6 +5939,8 @@ void ALandscapeProxy::PostTransacted(const FTransactionObjectEvent& InTransactio
 
 void ALandscapeProxy::PostEditChangeProperty(FPropertyChangedEvent& PropertyChangedEvent)
 {
+	ensureAlwaysMsgf(!GetLevel() || (GetLevel() && !GetLevel()->HasAnyFlags(EObjectFlags::RF_NeedPostLoad)), TEXT("Executing ALandscapeProxy::PostEditChangeProperty before the landscape actor has finished loading with the world can result in registration events occuring in PostLoad, which could leave the landscape proxy in an temporary invalid state"));
+
 	const FName MemberPropertyName = PropertyChangedEvent.MemberProperty ? PropertyChangedEvent.MemberProperty->GetFName() : NAME_None;
 	const FName SubPropertyName = PropertyChangedEvent.Property ? PropertyChangedEvent.Property->GetFName() : NAME_None;
 
@@ -6748,7 +6750,7 @@ void ALandscape::PostEditChangeProperty(FPropertyChangedEvent& PropertyChangedEv
 				// The ALandscape is included in ForEachLandscapeProxy, avoid getting stuck in a PostEditChange loop
 				if (Proxy != this)
 				{
-					Proxy->SynchronizeSharedProperties(this, /*bInRunPostEditChangeProperty =*/true);
+					Proxy->SynchronizeSharedProperties(this, /*bInExecutePostEditChangeProperty =*/true);
 				}
 				return true;
 			});
