Landscape - Fixed edge fixup regression where the FHeightmapTextureEdgeSnapshot TextureSourceID never gets set
* Regression introduced by 43590839 where CaptureEdgeDataFromHeightmapSource_Internal previously set the TextureSourceID value

#Synced-CL 51540438

diff --git a/Engine/Source/Runtime/Landscape/Classes/LandscapeEdgeFixup.h b/Engine/Source/Runtime/Landscape/Classes/LandscapeEdgeFixup.h
--- a/Engine/Source/Runtime/Landscape/Classes/LandscapeEdgeFixup.h
+++ b/Engine/Source/Runtime/Landscape/Classes/LandscapeEdgeFixup.h
@@ -169,8 +169,6 @@ public:
 	// Create a snapshot from the heightmap source
 	static TSharedPtr<FHeightmapTextureEdgeSnapshot> CreateEdgeSnapshotFromHeightmapSource(UTexture2D* InHeightmap, const FVector& LandscapeGridScale);
 #endif // WITH_EDITOR
-	// Create a snapshot from an explicit byte array (in heightmap source standard layout)
-	static TSharedRef<FHeightmapTextureEdgeSnapshot> CreateEdgeSnapshotFromTextureData(const TArrayView64<UE::Landscape::FHeightmapTexel>& InHeightmapTextureData, int32 InEdgeLength, const FVector& LandscapeGridScale);
 
 	// Return the set of edges that are different (according to edge hashes) -- i.e. changes that could cause neighbors to patch
 	UE::Landscape::EEdgeFlags CompareEdges(const FHeightmapTextureEdgeSnapshot& OldSnapshot) const;
@@ -191,7 +189,7 @@ private:
 
 	// This function is not thread safe to call on existing snapshots
 	// it should only be used on a newly allocated FHeightmapTextureEdgeSnapshot
-	void CaptureEdgeDataFromTextureData_Internal(const TArrayView64<UE::Landscape::FHeightmapTexel>& InHeightmapTextureData, int32 InEdgeLength, const FVector& LandscapeGridScale);
+	void CaptureEdgeDataFromTextureData_Internal(const TArrayView64<UE::Landscape::FHeightmapTexel>& InHeightmapTextureData, int32 InEdgeLength, const FVector& LandscapeGridScale, const FGuid& InTextureSourceID);
 	
 	// must call ResizeForEdgeLength to set the texture size before calling CopyTextureEdgeDataAndRecomputeNormals
 	void CaptureSingleEdgeDataAndComputeNormalsAndHashes(const UE::Landscape::FHeightmapTexel* TextureData, const UE::Landscape::EEdgeIndex EdgeOrCorner, const FVector& LandscapeGridScale);
diff --git a/Engine/Source/Runtime/Landscape/Private/LandscapeEdgeFixup.cpp b/Engine/Source/Runtime/Landscape/Private/LandscapeEdgeFixup.cpp
--- a/Engine/Source/Runtime/Landscape/Private/LandscapeEdgeFixup.cpp
+++ b/Engine/Source/Runtime/Landscape/Private/LandscapeEdgeFixup.cpp
@@ -295,14 +295,6 @@ UE::Landscape::FHeightmapTexel FHeightmapTextureEdgeSnapshot::GetCornerData(UE::
 	return Texel;
 }
 
-TSharedRef<FHeightmapTextureEdgeSnapshot> FHeightmapTextureEdgeSnapshot::CreateEdgeSnapshotFromTextureData(const TArrayView64<UE::Landscape::FHeightmapTexel>& InHeightmapTextureData, int32 InEdgeLength, const FVector& LandscapeGridScale)
-{
-	TRACE_CPUPROFILER_EVENT_SCOPE(ULandscapeHeightmapTextureEdgeSnapshot::CreateEdgeDataFromTextureData);
-	TSharedRef<FHeightmapTextureEdgeSnapshot> NewEdgeSnapshot = MakeShared<FHeightmapTextureEdgeSnapshot>();
-	NewEdgeSnapshot->CaptureEdgeDataFromTextureData_Internal(InHeightmapTextureData, InEdgeLength, LandscapeGridScale);
-	return NewEdgeSnapshot;
-}
-
 #if WITH_EDITOR
 TSharedPtr<FHeightmapTextureEdgeSnapshot> FHeightmapTextureEdgeSnapshot::CreateEdgeSnapshotFromHeightmapSource(UTexture2D* InHeightmap, const FVector& LandscapeGridScale)
 {
@@ -318,13 +310,13 @@ TSharedPtr<FHeightmapTextureEdgeSnapshot> FHeightmapTextureEdgeSnapshot::CreateE
 	{
 		return nullptr;
 	}
-	FGuid TextureSourceID = ULandscapeTextureHash::GetHash(InHeightmap);
+	const FGuid TextureSourceID = ULandscapeTextureHash::GetHash(InHeightmap);
 
 	FIXUP_DEBUG_LOG_DETAIL(TEXT("    CaptureEdgeDataFromHeightmapSource_Internal %p (%s) - ID %s"), InHeightmap, *LandscapeGridScale.ToString(), *TextureSourceID.ToString());
 
 	TArrayView64<FHeightmapTexel> TexelData((FHeightmapTexel*)MipData.GetData(), MipData.Num() / sizeof(FHeightmapTexel));
 	TSharedPtr<FHeightmapTextureEdgeSnapshot> NewEdgeSnapshot = MakeShared<FHeightmapTextureEdgeSnapshot>();
-	NewEdgeSnapshot->CaptureEdgeDataFromTextureData_Internal(TexelData, Source.GetSizeY(), LandscapeGridScale);
+	NewEdgeSnapshot->CaptureEdgeDataFromTextureData_Internal(TexelData, Source.GetSizeY(), LandscapeGridScale, TextureSourceID);
 	return NewEdgeSnapshot;
 }
 #endif // WITH_EDITOR
@@ -616,11 +608,15 @@ void FHeightmapTextureEdgeSnapshot::CaptureSingleEdgeDataAndComputeNormalsAndHas
 	SnapshotEdgeHashes[static_cast<int32>(EdgeOrCorner)] = SnapshotHash;
 }
 
-void FHeightmapTextureEdgeSnapshot::CaptureEdgeDataFromTextureData_Internal(const TArrayView64<UE::Landscape::FHeightmapTexel>& InHeightmapTextureData, int32 InEdgeLength, const FVector& LandscapeGridScale)
+void FHeightmapTextureEdgeSnapshot::CaptureEdgeDataFromTextureData_Internal(const TArrayView64<UE::Landscape::FHeightmapTexel>& InHeightmapTextureData, int32 InEdgeLength, const FVector& LandscapeGridScale, const FGuid& InTextureSourceID)
 {
 	TRACE_CPUPROFILER_EVENT_SCOPE(FHeightmapTextureEdgeSnapshot::CaptureEdgeDataFromTextureData_Internal);
 	using namespace UE::Landscape;
 
+#if WITH_EDITOR
+	TextureSourceID = InTextureSourceID;
+#endif // WITH_EDITOR
+
 	check(InHeightmapTextureData.Num() == InEdgeLength * InEdgeLength);
 	ResizeForEdgeLength(InEdgeLength);
 
