Landscape Editor - LandscapeGroup TickEdgeFixup nullptr crash fix for proxies pending reregister
* Use a copy of the set incase proxies are added/removed during iteration

#Synced-CL 51540438

diff --git a/Engine/Source/Runtime/Landscape/Private/LandscapeSubsystem.cpp b/Engine/Source/Runtime/Landscape/Private/LandscapeSubsystem.cpp
--- a/Engine/Source/Runtime/Landscape/Private/LandscapeSubsystem.cpp
+++ b/Engine/Source/Runtime/Landscape/Private/LandscapeSubsystem.cpp
@@ -544,9 +544,9 @@ TStatId ULandscapeSubsystem::GetStatId() const
 	RETURN_QUICK_DECLARE_CYCLE_STAT(ULandscapeSubsystem, STATGROUP_Tickables);
 }
 
-FLandscapeGroup* ULandscapeSubsystem::GetLandscapeGroupForProxy(ALandscapeProxy* Proxy)
+FLandscapeGroup* ULandscapeSubsystem::GetLandscapeGroupForProxy(TObjectPtr<ALandscapeProxy> Proxy)
 {
-	check(Proxy);
+	check(IsValid(Proxy));
 
 	// use LODGroupKey instead of LandscapeGUID when LODGroupKey is non-zero
 	uint32 LandscapeGroupKey = (Proxy->LODGroupKey != 0) ? Proxy->LODGroupKey : GetTypeHash(Proxy->GetLandscapeGuid());
@@ -984,34 +984,33 @@ void ULandscapeSubsystem::TickEdgeFixup()
 	// apply any requested re-registrations first
 	if (StreamingProxiesNeedingReregister.Num() > 0)
 	{
-		int32 ComponentsUnregistered = 0;
-		int32 ComponentsRegistered = 0;
+		check(IsInGameThread());
 
-		// unregister all first, from whatever group they are registered to
-		for (ALandscapeProxy* StreamingProxy : StreamingProxiesNeedingReregister)
+		// Create a stable copy of the proxies in case the original set is modified while iterating
+		TSet<TObjectPtr<ALandscapeStreamingProxy>> StreamingProxiesNeedingReregisterCopy;
+		Swap(StreamingProxiesNeedingReregisterCopy, StreamingProxiesNeedingReregister);
+
+		for (TObjectPtr<ALandscapeProxy> StreamingProxy : StreamingProxiesNeedingReregisterCopy)
 		{
+			check(IsValid(StreamingProxy));
 			for (ULandscapeComponent* Component : StreamingProxy->LandscapeComponents)
 			{
 				if (FLandscapeGroup* Group = Component->RegisteredLandscapeGroup)
 				{
 					Group->UnregisterComponent(Component);
-					ComponentsUnregistered++;
 				}
 			}
 		}
 
 		// then re-register them to the group they should be in
-		for (ALandscapeProxy* StreamingProxy : StreamingProxiesNeedingReregister)
+		for (TObjectPtr<ALandscapeProxy> StreamingProxy : StreamingProxiesNeedingReregisterCopy)
 		{
 			FLandscapeGroup* NewGroup = GetLandscapeGroupForProxy(StreamingProxy);
 			for (ULandscapeComponent* Component : StreamingProxy->LandscapeComponents)
 			{
 				NewGroup->RegisterComponent(Component);
-				ComponentsRegistered++;
 			}
 		}
-
-		StreamingProxiesNeedingReregister.Empty();
 	}
 
 	// tick each group
diff --git a/Engine/Source/Runtime/Landscape/Public/LandscapeSubsystem.h b/Engine/Source/Runtime/Landscape/Public/LandscapeSubsystem.h
--- a/Engine/Source/Runtime/Landscape/Public/LandscapeSubsystem.h
+++ b/Engine/Source/Runtime/Landscape/Public/LandscapeSubsystem.h
@@ -226,7 +226,7 @@ public:
 	// runs per-tick edge fixup on ALL landscape groups in the subsystem
 	void TickEdgeFixup();
 
-	FLandscapeGroup* GetLandscapeGroupForProxy(ALandscapeProxy* Proxy);
+	FLandscapeGroup* GetLandscapeGroupForProxy(TObjectPtr<ALandscapeProxy> Proxy);
 	FLandscapeGroup* GetLandscapeGroupForComponent(ULandscapeComponent* Component);
 
 	LANDSCAPE_API void ForEachLandscapeInfo(TFunctionRef<bool(ULandscapeInfo*)> ForEachLandscapeInfoFunc) const;
