Landscape: Avoid long hangs at end of sculpt/paint tool strokes with large landscapes.

Split Update_Heightmap_All/Update_Weightmap_All into a combination of two flags each.  Update_Heightmap_All_Modes/Update_Weightmap_All_Modes and a separate Update_All_Components flag.

RequestLayersContentUpdate semantics don't change.  Update_Heightmap_All/Update_Weightmap_All will still trigger processing of all components as before.  Flag-testing code changes to only test for the _Modes flags in most places.  The Update_All_Components flag can be ommitted to avoid processing all components.

#Synced-CL 51561720

diff --git a/Engine/Source/Editor/LandscapeEditor/Private/LandscapeEdModeTools.h b/Engine/Source/Editor/LandscapeEditor/Private/LandscapeEdModeTools.h
--- a/Engine/Source/Editor/LandscapeEditor/Private/LandscapeEdModeTools.h
+++ b/Engine/Source/Editor/LandscapeEditor/Private/LandscapeEdModeTools.h
@@ -1259,7 +1259,7 @@ public:
 	virtual ELandscapeLayerUpdateMode GetEndToolContentUpdateFlag() const
 	{
 		bool bUpdateHeightmap = this->EdMode->CurrentToolTarget.TargetType == ELandscapeToolTargetType::Heightmap;
-		return bUpdateHeightmap ? ELandscapeLayerUpdateMode::Update_Heightmap_All : ELandscapeLayerUpdateMode::Update_Weightmap_All;
+		return bUpdateHeightmap ? ELandscapeLayerUpdateMode::Update_Heightmap_All_Modes : ELandscapeLayerUpdateMode::Update_Weightmap_All_Modes;
 	}
 
 	virtual bool BeginTool(FEditorViewportClient* ViewportClient, const FLandscapeToolTarget& InTarget, const FVector& InHitLocation) override
diff --git a/Engine/Source/Runtime/Landscape/Classes/LandscapeComponent.h b/Engine/Source/Runtime/Landscape/Classes/LandscapeComponent.h
--- a/Engine/Source/Runtime/Landscape/Classes/LandscapeComponent.h
+++ b/Engine/Source/Runtime/Landscape/Classes/LandscapeComponent.h
@@ -372,27 +372,41 @@ enum ELandscapeComponentUpdateFlag : uint32
 	Component_Update_Approximated_Bounds = 1 << 5
 };
 
+/**
+	ELandscapeLayerUpdateMode flags relate to types of updates to landscape layer data.  Stored on the ALandscape and on
+	ULandscapeComponent to indicate pending changes.  Stored on FLandscapeEditLayerComponentReadbackResult to indicate
+	what type of change triggered the creation of the new data.
+ 
+	Update_All_Components on ALandscape or passed to RequestLayersContentUpdate will cause all components to be enqueued
+	for layer merging, rather than only those components carrying relevant flags of their own.
+
+ 	See IsUpdateFlagEnabledForModes, where these flags are translated to ELandscapeComponentUpdateFlag, that indicate
+	specific tasks.
+**/
 enum ELandscapeLayerUpdateMode : uint32
 { 
 	// No Update
 	Update_None = 0,
 	// Update types
-	Update_Heightmap_All = 1 << 0,
+	Update_Heightmap_All_Modes = 1 << 0,
 	Update_Heightmap_Editing = 1 << 1,
 	Update_Heightmap_Editing_NoCollision = 1 << 2,
-	Update_Weightmap_All = 1 << 3,
+	Update_Weightmap_All_Modes = 1 << 3,
 	Update_Weightmap_Editing = 1 << 4,
 	Update_Weightmap_Editing_NoCollision = 1 << 5,
+	Update_All_Components = 1 << 6,
 	// Combinations
+	Update_Heightmap_All = Update_Heightmap_All_Modes | Update_All_Components,
+	Update_Weightmap_All = Update_Weightmap_All_Modes | Update_All_Components,
 	Update_All = Update_Weightmap_All | Update_Heightmap_All,
 	Update_All_Editing = Update_Weightmap_Editing | Update_Heightmap_Editing,
 	Update_All_Editing_NoCollision = Update_Weightmap_Editing_NoCollision | Update_Heightmap_Editing_NoCollision,
-	Update_Heightmap_Types = Update_Heightmap_All | Update_Heightmap_Editing | Update_Heightmap_Editing_NoCollision,
-	Update_Weightmap_Types = Update_Weightmap_All | Update_Weightmap_Editing | Update_Weightmap_Editing_NoCollision,
+	Update_Heightmap_Types = Update_Heightmap_All_Modes | Update_Heightmap_Editing | Update_Heightmap_Editing_NoCollision,
+	Update_Weightmap_Types = Update_Weightmap_All_Modes | Update_Weightmap_Editing | Update_Weightmap_Editing_NoCollision,
 	// In cases where we couldn't update the clients right away this flag will be set in RegenerateLayersContent
-	Update_Client_Deferred = 1 << 6,
+	Update_Client_Deferred = 1 << 7,
 	// Update landscape component clients while editing
-	Update_Client_Editing = 1 << 7
+	Update_Client_Editing = 1 << 8
 };
 
 static const uint32 DefaultSplineHash = 0xFFFFFFFF;
diff --git a/Engine/Source/Runtime/Landscape/Private/LandscapeEditLayers.cpp b/Engine/Source/Runtime/Landscape/Private/LandscapeEditLayers.cpp
--- a/Engine/Source/Runtime/Landscape/Private/LandscapeEditLayers.cpp
+++ b/Engine/Source/Runtime/Landscape/Private/LandscapeEditLayers.cpp
@@ -5041,7 +5041,7 @@ void ALandscape::UpdateForChangedHeightmaps(const TArrayView<FLandscapeEditLayer
 			ComponentReadbackResult.LandscapeComponent->SetPendingCollisionDataUpdate(true);
 		}
 
-		const uint32 HeightUpdateMode = ComponentReadbackResult.UpdateModes & (ELandscapeLayerUpdateMode::Update_Heightmap_All | ELandscapeLayerUpdateMode::Update_Heightmap_Editing | ELandscapeLayerUpdateMode::Update_Heightmap_Editing_NoCollision);
+		const uint32 HeightUpdateMode = ComponentReadbackResult.UpdateModes & (ELandscapeLayerUpdateMode::Update_Heightmap_All_Modes | ELandscapeLayerUpdateMode::Update_Heightmap_Editing | ELandscapeLayerUpdateMode::Update_Heightmap_Editing_NoCollision);
 
 		// Only update collision if there was an actual change performed on the source data : 
 		if (ComponentReadbackResult.LandscapeComponent->GetPendingCollisionDataUpdate())
@@ -6494,7 +6494,7 @@ void ALandscape::UpdateForChangedWeightmaps(const TArrayView<FLandscapeEditLayer
 			}
 		}
 
-		const int32 WeightUpdateMode = ComponentReadbackResult.UpdateModes & (ELandscapeLayerUpdateMode::Update_Weightmap_All | ELandscapeLayerUpdateMode::Update_Weightmap_Editing | ELandscapeLayerUpdateMode::Update_Weightmap_Editing_NoCollision);
+		const int32 WeightUpdateMode = ComponentReadbackResult.UpdateModes & (ELandscapeLayerUpdateMode::Update_Weightmap_All_Modes | ELandscapeLayerUpdateMode::Update_Weightmap_Editing | ELandscapeLayerUpdateMode::Update_Weightmap_Editing_NoCollision);
 		if (IsUpdateFlagEnabledForModes(ELandscapeComponentUpdateFlag::Component_Update_Weightmap_Collision, WeightUpdateMode) 
 			|| ComponentReadbackResult.bCleared)
 		{
@@ -6799,16 +6799,16 @@ void ALandscape::RequestLayersContentUpdateForceAll(ELandscapeLayerUpdateMode In
 		return;
 	}
 
-	const bool bUpdateWeightmap = (InModeMask & (ELandscapeLayerUpdateMode::Update_Weightmap_All | ELandscapeLayerUpdateMode::Update_Weightmap_Editing | ELandscapeLayerUpdateMode::Update_Weightmap_Editing_NoCollision)) != 0;
-	const bool bUpdateHeightmap = (InModeMask & (ELandscapeLayerUpdateMode::Update_Heightmap_All | ELandscapeLayerUpdateMode::Update_Heightmap_Editing | ELandscapeLayerUpdateMode::Update_Heightmap_Editing_NoCollision)) != 0;
-	const bool bUpdateWeightCollision = (InModeMask & (ELandscapeLayerUpdateMode::Update_Weightmap_All | ELandscapeLayerUpdateMode::Update_Weightmap_Editing)) != 0;
-	const bool bUpdateHeightCollision = (InModeMask & (ELandscapeLayerUpdateMode::Update_Heightmap_All | ELandscapeLayerUpdateMode::Update_Heightmap_Editing)) != 0;
-	const bool bUpdateAllHeightmap = (InModeMask & ELandscapeLayerUpdateMode::Update_Heightmap_All) != 0;
-	const bool bUpdateAllWeightmap = (InModeMask & ELandscapeLayerUpdateMode::Update_Weightmap_All) != 0;
-	const bool bUpdateClientUdpateEditing = (InModeMask & ELandscapeLayerUpdateMode::Update_Client_Editing) != 0;
+	const bool bUpdateWeightmap = (InModeMask & (ELandscapeLayerUpdateMode::Update_Weightmap_All_Modes | ELandscapeLayerUpdateMode::Update_Weightmap_Editing | ELandscapeLayerUpdateMode::Update_Weightmap_Editing_NoCollision)) != 0;
+	const bool bUpdateHeightmap = (InModeMask & (ELandscapeLayerUpdateMode::Update_Heightmap_All_Modes | ELandscapeLayerUpdateMode::Update_Heightmap_Editing | ELandscapeLayerUpdateMode::Update_Heightmap_Editing_NoCollision)) != 0;
+	const bool bUpdateWeightCollision = (InModeMask & (ELandscapeLayerUpdateMode::Update_Weightmap_All_Modes | ELandscapeLayerUpdateMode::Update_Weightmap_Editing)) != 0;
+	const bool bUpdateHeightCollision = (InModeMask & (ELandscapeLayerUpdateMode::Update_Heightmap_All_Modes | ELandscapeLayerUpdateMode::Update_Heightmap_Editing)) != 0;
+	const bool bUpdateAllHeightmapModes = (InModeMask & ELandscapeLayerUpdateMode::Update_Heightmap_All_Modes) != 0;
+	const bool bUpdateAllWeightmapModes = (InModeMask & ELandscapeLayerUpdateMode::Update_Weightmap_All_Modes) != 0;
+	const bool bUpdateClientUpdateEditing = (InModeMask & ELandscapeLayerUpdateMode::Update_Client_Editing) != 0;
 	if (ULandscapeInfo* LandscapeInfo = GetLandscapeInfo())
 	{
-		LandscapeInfo->ForEachLandscapeProxy([bUpdateHeightmap, bUpdateWeightmap, bUpdateAllHeightmap, bUpdateAllWeightmap, bUpdateHeightCollision, bUpdateWeightCollision, bUpdateClientUdpateEditing, bInUserTriggered](ALandscapeProxy* Proxy)
+		LandscapeInfo->ForEachLandscapeProxy([bUpdateHeightmap, bUpdateWeightmap, bUpdateAllHeightmapModes, bUpdateAllWeightmapModes, bUpdateHeightCollision, bUpdateWeightCollision, bUpdateClientUpdateEditing, bInUserTriggered](ALandscapeProxy* Proxy)
 		{
 			if (Proxy)
 			{
@@ -6816,15 +6816,15 @@ void ALandscape::RequestLayersContentUpdateForceAll(ELandscapeLayerUpdateMode In
 				{
 					if (bUpdateHeightmap)
 					{
-						Component->RequestHeightmapUpdate(bUpdateAllHeightmap, bUpdateHeightCollision, bInUserTriggered);
+						Component->RequestHeightmapUpdate(bUpdateAllHeightmapModes, bUpdateHeightCollision, bInUserTriggered);
 					}
 
 					if (bUpdateWeightmap)
 					{
-						Component->RequestWeightmapUpdate(bUpdateAllWeightmap, bUpdateWeightCollision, bInUserTriggered);
+						Component->RequestWeightmapUpdate(bUpdateAllWeightmapModes, bUpdateWeightCollision, bInUserTriggered);
 					}
 
-					if (bUpdateClientUdpateEditing)
+					if (bUpdateClientUpdateEditing)
 					{
 						Component->RequestEditingClientUpdate(bInUserTriggered);
 					}
@@ -6839,7 +6839,7 @@ void ALandscape::RequestLayersContentUpdateForceAll(ELandscapeLayerUpdateMode In
 
 bool ALandscape::IsUpdateFlagEnabledForModes(ELandscapeComponentUpdateFlag InFlag, uint32 InUpdateModes)
 {
-	if (InUpdateModes & ELandscapeLayerUpdateMode::Update_Heightmap_All)
+	if (InUpdateModes & ELandscapeLayerUpdateMode::Update_Heightmap_All_Modes)
 	{
 		const uint32 HeightmapAllFlags = ELandscapeComponentUpdateFlag::Component_Update_Heightmap_Collision | ELandscapeComponentUpdateFlag::Component_Update_Recreate_Collision | ELandscapeComponentUpdateFlag::Component_Update_Client;
 		if (HeightmapAllFlags & InFlag)
@@ -6857,7 +6857,7 @@ bool ALandscape::IsUpdateFlagEnabledForModes(ELandscapeComponentUpdateFlag InFla
 		}
 	}
 
-	if (InUpdateModes & ELandscapeLayerUpdateMode::Update_Weightmap_All)
+	if (InUpdateModes & ELandscapeLayerUpdateMode::Update_Weightmap_All_Modes)
 	{
 		const uint32 WeightmapAllFlags = ELandscapeComponentUpdateFlag::Component_Update_Weightmap_Collision | ELandscapeComponentUpdateFlag::Component_Update_Recreate_Collision | ELandscapeComponentUpdateFlag::Component_Update_Client;
 		if (WeightmapAllFlags & InFlag)
@@ -6933,7 +6933,7 @@ void ULandscapeComponent::RequestHeightmapUpdate(bool bUpdateAll, bool bUpdateCo
 	if (bUpdateAll || bUpdateCollision)
 	{
 		NewLayerUpdateFlags |= ELandscapeLayerUpdateMode::Update_Heightmap_Editing;
-		NewLayerUpdateFlags |= ELandscapeLayerUpdateMode::Update_Heightmap_All;
+		NewLayerUpdateFlags |= ELandscapeLayerUpdateMode::Update_Heightmap_All_Modes;
 	}
 	else
 	{
@@ -6942,14 +6942,12 @@ void ULandscapeComponent::RequestHeightmapUpdate(bool bUpdateAll, bool bUpdateCo
 
 	LayerUpdateFlagPerMode |= NewLayerUpdateFlags;
 
-	// The typical codepath for sculpt or paint will have bUpdateCollision=true and bUpdateAll=false.  It's unclear why this also causes
-	// Update_Heightmap_All on the component flags, but that's how it's been for a while.  We want either of the "_Editing" flags also
-	// passed through to the landscape actor.  Only pass through the _All flag in the case of bUpdateAll, not just from bUpdateCollision.
-	// Note that Update_Heightmap_All on the actor will cause the entire landscape to re-evaluate, not just this component.
+	// The typical codepath for sculpt or paint will have bUpdateCollision=true and bUpdateAll=false.  Only pass through
+	// the _All_Modes flag in the case of bUpdateAll, not just from bUpdateCollision.
 	uint32 ActorFlags = NewLayerUpdateFlags;
 	if (!bUpdateAll)
 	{
-		ActorFlags &= ~ELandscapeLayerUpdateMode::Update_Heightmap_All;  // Remove _All flag
+		ActorFlags &= ~ELandscapeLayerUpdateMode::Update_Heightmap_All_Modes;  // Remove _All_Modes flag
 	}
 
 	if (ALandscape* LandscapeActor = GetLandscapeActor())
@@ -6971,7 +6969,7 @@ void ULandscapeComponent::RequestWeightmapUpdate(bool bUpdateAll, bool bUpdateCo
 	if (bUpdateAll || bUpdateCollision)
 	{
 		NewLayerUpdateFlags |= ELandscapeLayerUpdateMode::Update_Weightmap_Editing;
-		NewLayerUpdateFlags |= ELandscapeLayerUpdateMode::Update_Weightmap_All;
+		NewLayerUpdateFlags |= ELandscapeLayerUpdateMode::Update_Weightmap_All_Modes;
 	}
 	else
 	{
@@ -6980,14 +6978,12 @@ void ULandscapeComponent::RequestWeightmapUpdate(bool bUpdateAll, bool bUpdateCo
 
 	LayerUpdateFlagPerMode |= NewLayerUpdateFlags;
 
-	// The typical codepath for sculpt or paint will have bUpdateCollision=true and bUpdateAll=false.  It's unclear why this also causes
-	// Update_Weightmap_All on the component flags, but that's how it's been for a while.  We want either of the "_Editing" flags also
-	// passed through to the landscape actor.  Only pass through the _All flag in the case of bUpdateAll, not just from bUpdateCollision.
-	// Note that Update_Weightmap_All on the actor will cause the entire landscape to re-evaluate, not just this component.
+	// The typical codepath for sculpt or paint will have bUpdateCollision=true and bUpdateAll=false.  Only pass through
+	// the _All_Modes flag in the case of bUpdateAll, not just from bUpdateCollision.
 	uint32 ActorFlags = NewLayerUpdateFlags;
 	if (!bUpdateAll)
 	{
-		ActorFlags &= ~ELandscapeLayerUpdateMode::Update_Weightmap_All;  // Remove _All flag
+		ActorFlags &= ~ELandscapeLayerUpdateMode::Update_Weightmap_All_Modes;  // Remove _All_Modes flag
 	}
 
 	if (ALandscape* LandscapeActor = GetLandscapeActor())
@@ -7463,8 +7459,8 @@ void ALandscape::UpdateLayersContent(bool bInWaitForStreaming, bool bInSkipMonit
 		return;
 	}
 
-	bool bUpdateAll = LayerContentUpdateModes & Update_All;
-	bool bPartialUpdate = !bUpdateAll && CVarLandscapeLayerOptim.GetValueOnAnyThread() == 1;
+	bool bUpdateAllComponents = LayerContentUpdateModes & Update_All_Components;
+	bool bPartialUpdate = !bUpdateAllComponents && CVarLandscapeLayerOptim.GetValueOnAnyThread() == 1;
 
 	FUpdateLayersContentContext UpdateLayersContentContext(MapHelper, bPartialUpdate);
 
