Clear landscape Update_All_Components flag when processed.

#Synced-CL 51561720

diff --git a/Engine/Source/Runtime/Landscape/Private/LandscapeEditLayers.cpp b/Engine/Source/Runtime/Landscape/Private/LandscapeEditLayers.cpp
--- a/Engine/Source/Runtime/Landscape/Private/LandscapeEditLayers.cpp
+++ b/Engine/Source/Runtime/Landscape/Private/LandscapeEditLayers.cpp
@@ -7470,6 +7470,7 @@ void ALandscape::UpdateLayersContent(bool bInWaitForStreaming, bool bInSkipMonit
 	ProcessedModes |= RegenerateLayersWeightmaps(UpdateLayersContentContext);
 	ProcessedModes |= (LayerContentUpdateModes & ELandscapeLayerUpdateMode::Update_Client_Deferred);
 	ProcessedModes |= (LayerContentUpdateModes & ELandscapeLayerUpdateMode::Update_Client_Editing);
+	ProcessedModes |= (LayerContentUpdateModes & ELandscapeLayerUpdateMode::Update_All_Components);
 
 	// If we are flushing then read back resolved textures immediately
 	if (bFlushRender || CVarLandscapeForceFlush.GetValueOnGameThread() != 0)
