Remove ALandscape::MonitorShaderCompilation.

#Synced-CL 51561720

diff --git a/Engine/Source/Runtime/Landscape/Classes/Landscape.h b/Engine/Source/Runtime/Landscape/Classes/Landscape.h
--- a/Engine/Source/Runtime/Landscape/Classes/Landscape.h
+++ b/Engine/Source/Runtime/Landscape/Classes/Landscape.h
@@ -587,7 +587,6 @@ private:
 
 	void UpdateLayersContent(bool bInWaitForStreaming, bool bInSkipMonitorLandscapeEdModeChanges, bool bFlushRender);
 	bool CanUpdateLayersContent() const;
-	void MonitorShaderCompilation();
 	void MonitorLandscapeEdModeChanges();
 	
 	//~ Begin UE::Landscape::EditLayers::IEditLayerRendererProvider implementation
@@ -752,8 +751,9 @@ private:
 	UPROPERTY(Transient)
 	bool bLandscapeLayersAreInitialized; 
 
+	UE_DEPRECATED(5.8, "WasCompilingShaders property is no longer used.")
 	UPROPERTY(Transient)
-	bool WasCompilingShaders;
+	bool WasCompilingShaders_DEPRECATED = false;
 
 	UPROPERTY(Transient)
 	uint32 LayerContentUpdateModes;
diff --git a/Engine/Source/Runtime/Landscape/Private/Landscape.cpp b/Engine/Source/Runtime/Landscape/Private/Landscape.cpp
--- a/Engine/Source/Runtime/Landscape/Private/Landscape.cpp
+++ b/Engine/Source/Runtime/Landscape/Private/Landscape.cpp
@@ -1832,7 +1832,6 @@ ALandscape::ALandscape(const FObjectInitializer& ObjectInitializer)
 {
 #if WITH_EDITORONLY_DATA
 	bLockLocation = false;
-	WasCompilingShaders = false;
 	LayerContentUpdateModes = 0;
 	bSplineLayerUpdateRequested = false;
 	bLandscapeLayersAreInitialized = false;
diff --git a/Engine/Source/Runtime/Landscape/Private/LandscapeEditLayers.cpp b/Engine/Source/Runtime/Landscape/Private/LandscapeEditLayers.cpp
--- a/Engine/Source/Runtime/Landscape/Private/LandscapeEditLayers.cpp
+++ b/Engine/Source/Runtime/Landscape/Private/LandscapeEditLayers.cpp
@@ -7041,27 +7041,6 @@ void ALandscape::MonitorLandscapeEdModeChanges()
 	}
 }
 
-void ALandscape::MonitorShaderCompilation()
-{
-	// Do not monitor changes when not editing Landscape
-	if (!LandscapeEdMode)
-	{
-		return;
-	}
-
-	// If doing editing while shader are compiling or at load of a map, it's possible we will need another update pass after shader are completed to see the correct result
-	const int32 RemainingShadersThisFrame = GShaderCompilingManager->GetNumRemainingJobs();
-	if (!WasCompilingShaders && RemainingShadersThisFrame > 0)
-	{
-		WasCompilingShaders = true;
-	}
-	else if (WasCompilingShaders)
-	{
-		WasCompilingShaders = false;
-		RequestLayersContentUpdateForceAll();
-	}
-}
-
 void ULandscapeComponent::GetLandscapeComponentNeighborsToRender(TSet<ULandscapeComponent*>& OutNeighborComponents) const
 {
 	ULandscapeInfo* LandscapeInfo = GetLandscapeInfo();
@@ -7339,7 +7318,6 @@ void ALandscape::UpdateLayersContent(bool bInWaitForStreaming, bool bInSkipMonit
 	{
 		MonitorLandscapeEdModeChanges();
 	}
-	MonitorShaderCompilation();
 
 	// Make sure Brush get a chance to request an update of the landscape
 	for (FLandscapeLayer& Layer : LandscapeEditLayers)
