Fixed landscape batched merge generating useless mips when generating a given weightmap in several batches

#Synced-CL 51768384

diff --git a/Engine/Source/Runtime/Landscape/Private/LandscapeEditLayers.cpp b/Engine/Source/Runtime/Landscape/Private/LandscapeEditLayers.cpp
--- a/Engine/Source/Runtime/Landscape/Private/LandscapeEditLayers.cpp
+++ b/Engine/Source/Runtime/Landscape/Private/LandscapeEditLayers.cpp
@@ -6227,7 +6227,7 @@ int32 ALandscape::PerformLayersWeightmapsBatchedMerge(FUpdateLayersContentContex
 						// We use the final texture as a temporary buffer when it's a partial resolve, so we only need to copy mip 0 then. All mips will be generated/copy when the resolve step is final :
 						const int32 NumMipsToCopy = WeightmapResolveInfo.bIsFinalResolve ? NumMips : 1;
 						RDG_EVENT_SCOPE(GraphBuilder, "Copy %d mips", NumMips);
-						for (int32 MipLevel = 0; MipLevel < NumMips; ++MipLevel)
+						for (int32 MipLevel = 0; MipLevel < NumMipsToCopy; ++MipLevel)
 						{
 							FRHICopyTextureInfo CopyInfo;
 							CopyInfo.SourceMipIndex = MipLevel;
