VulkanRHI: Change workaround for large number of texture samples in branches that also fixes UE-369443.

#rnx

#Synced-CL 51774848

diff --git a/Engine/Shaders/Private/Landscape/LandscapeEditLayersUtils.usf b/Engine/Shaders/Private/Landscape/LandscapeEditLayersUtils.usf
--- a/Engine/Shaders/Private/Landscape/LandscapeEditLayersUtils.usf
+++ b/Engine/Shaders/Private/Landscape/LandscapeEditLayersUtils.usf
@@ -32,37 +32,6 @@ Texture2D<float4> InSourceTexture;
 // CopyQuadsMultiSourcePS inputs/outputs :
 
 // Helper macro for handling an array of N textures (currently using 63 because 64 is the maximum supported number of SRV on all platforms needing to run this) :
-#if VULKAN_PROFILE_SM6
-#define FOREACH_SOURCETEXTURE(op) \
-	op(0) \
-	op(1) \
-	op(2) \
-	op(3) \
-	op(4) \
-	op(5) \
-	op(6) \
-	op(7) \
-	op(8) \
-	op(9) \
-	op(10) \
-	op(11) \
-	op(12) \
-	op(13) \
-	op(14) \
-	op(15) \
-	op(16) \
-	op(17) \
-	op(18) \
-	op(19) \
-	op(20) \
-	op(21) \
-	op(22) \
-	op(23) \
-	op(24) \
-	op(25) \
-	op(26) \
-	op(27) 
-#else // VULKAN_PROFILE_SM6 
 #define FOREACH_SOURCETEXTURE(op) \
 	op(0) \
 	op(1) \
@@ -126,7 +95,6 @@ Texture2D<float4> InSourceTexture;
 	op(59) \
 	op(60) \
 	op(61) 
-#endif // !VULKAN_PROFILE_SM6 
 
 Buffer<uint4> InQuadInfos; // .x = source texture index, .y = source texture channel (for weightmaps), .z = component id of the component this quad corresponds to, .w = unused
 Texture2D<float> InComponentIdTexture; // texture containing the linear index of the components in this batch. This allows to resolve conflicts in case multiple quads render to the same pixel
@@ -171,12 +139,27 @@ void CopyQuadsPS(float4 SvPosition : SV_POSITION, float2 QuadUV : TEXCOORD0, out
 
 #if COPY_QUADS_MULTISOURCE
 
+float4 GetSourceTextureColor(uint SourceTextureIndex, float2 SourceUV)
+{
+	// Having a switch with returns is necessaire in Vulkan (see  UE-369443):
+	// We have an issue in Vulkan when the switch is too large, the large if-else causes issues under investigation.
+
+	BRANCH
+	switch (SourceTextureIndex)
+	{
+	#define IMPLEMENT_SAMPLE_SWITCH_CASE(n) case n: return InSourceTexture_##n.SampleLevel(GlobalPointClampedSampler, SourceUV, 0);
+		FOREACH_SOURCETEXTURE(IMPLEMENT_SAMPLE_SWITCH_CASE);
+	#undef IMPLEMENT_SAMPLE_SWITCH_CASE
+
+	default: return float4(0.0f, 0.0f, 0.0f, 0.0f);
+	}
+}
+
 void CopyQuadsMultiSourcePS(float4 SvPosition : SV_POSITION, float2 SourceUV : TEXCOORD0, float2 UnusedQuadUV : TEXCOORD1, float InRectIndex : RECT_INDEX, out float4 OutColor : SV_Target0)
 {
 	uint QuadIndex = (uint)InRectIndex;
 	uint SourceTextureIndex = InQuadInfos[QuadIndex].x;
 	uint QuadComponentId = InQuadInfos[QuadIndex].z;
-	float4 ColorSample = 0.0f;
 
 	float ComponentIdSample = InComponentIdTexture.Load(int3(SvPosition.xy, 0)); // xy = coordinates, z = mip level
 	uint ComponentIdUint = asuint(ComponentIdSample);
@@ -186,16 +169,7 @@ void CopyQuadsMultiSourcePS(float4 SvPosition : SV_POSITION, float2 SourceUV : T
 		discard;
 	}
 
-	BRANCH
-	switch (SourceTextureIndex)
-	{
-#define IMPLEMENT_SAMPLE_SWITCH_CASE(n) case n: ColorSample = InSourceTexture_##n.SampleLevel(GlobalPointClampedSampler, SourceUV, 0); break; 
-	FOREACH_SOURCETEXTURE(IMPLEMENT_SAMPLE_SWITCH_CASE);
-#undef IMPLEMENT_SAMPLE_SWITCH_CASE
-
-	default: 
-		break;
-	}
+	const float4 ColorSample = GetSourceTextureColor(SourceTextureIndex, SourceUV);
 
 	OutColor = 0.0f;
 
diff --git a/Engine/Source/Runtime/Landscape/Private/LandscapeEditLayers.cpp b/Engine/Source/Runtime/Landscape/Private/LandscapeEditLayers.cpp
--- a/Engine/Source/Runtime/Landscape/Private/LandscapeEditLayers.cpp
+++ b/Engine/Source/Runtime/Landscape/Private/LandscapeEditLayers.cpp
@@ -948,36 +948,6 @@ public:
 	}
 };
 
-class FCopyQuadsMultiSourcePSVulkanSM6 : public FGlobalShader
-{
-public:
-	static constexpr int32 NumMultiSources = 28;
-
-	DECLARE_GLOBAL_SHADER(FCopyQuadsMultiSourcePSVulkanSM6);
-	SHADER_USE_PARAMETER_STRUCT(FCopyQuadsMultiSourcePSVulkanSM6, FGlobalShader);
-
-	BEGIN_SHADER_PARAMETER_STRUCT(FParameters, )
-		SHADER_PARAMETER_STRUCT_REF(FViewUniformShaderParameters, View)
-		SHADER_PARAMETER_RDG_BUFFER_SRV(Buffer<uint4>, InQuadInfos)
-		SHADER_PARAMETER_RDG_TEXTURE_SRV(Texture2D, InComponentIdTexture)
-		SHADER_PARAMETER_RDG_TEXTURE_ARRAY(Texture2D, InSourceTexture, [NumMultiSources])
-	END_SHADER_PARAMETER_STRUCT()
-
-	class FCopyWeightmap : SHADER_PERMUTATION_BOOL("COPY_WEIGHTMAP");
-	using FPermutationDomain = TShaderPermutationDomain<FCopyWeightmap>;
-
-	static bool ShouldCompilePermutation(const FGlobalShaderPermutationParameters& InParameters)
-	{
-		return UE::Landscape::DoesPlatformSupportEditLayers(InParameters.Platform) && (InParameters.Platform == SP_VULKAN_SM6);
-	}
-
-	static void ModifyCompilationEnvironment(const FGlobalShaderPermutationParameters& Parameters, FShaderCompilerEnvironment& OutEnvironment)
-	{
-		OutEnvironment.SetDefine(TEXT("COPY_QUADS_MULTISOURCE"), 1);
-		UE::Landscape::ModifyShaderCompilerEnvironmentForDebug(OutEnvironment);
-	}
-};
-
 // Helper class to trigger the appropriate PS regardless of the platform : 
 class FCopyQuadsMultiSourceBase
 {
@@ -1062,15 +1032,7 @@ BEGIN_SHADER_PARAMETER_STRUCT(FCopyQuadsMultiSourcePSParametersDefault, )
 END_SHADER_PARAMETER_STRUCT() 
 IMPLEMENT_GLOBAL_SHADER(FCopyQuadsMultiSourcePSDefault, "/Engine/Private/Landscape/LandscapeEditLayersUtils.usf", "CopyQuadsMultiSourcePS", SF_Pixel);
 
-BEGIN_SHADER_PARAMETER_STRUCT(FCopyQuadsMultiSourcePSParametersVulkanSM6, )
-SHADER_PARAMETER_STRUCT_INCLUDE(FPixelShaderUtils::FRasterizeToRectsVS::FParameters, VS)
-	SHADER_PARAMETER_STRUCT_INCLUDE(FCopyQuadsMultiSourcePSVulkanSM6::FParameters, PS)
-RENDER_TARGET_BINDING_SLOTS()
-END_SHADER_PARAMETER_STRUCT()
-IMPLEMENT_GLOBAL_SHADER(FCopyQuadsMultiSourcePSVulkanSM6, "/Engine/Private/Landscape/LandscapeEditLayersUtils.usf", "CopyQuadsMultiSourcePS", SF_Pixel);
-
 using FCopyQuadsMultiSourceDefault = TCopyQuadsMultiSource<FCopyQuadsMultiSourcePSDefault, FCopyQuadsMultiSourcePSParametersDefault>;
-using FCopyQuadsMultiSourceVulkanSM6 = TCopyQuadsMultiSource<FCopyQuadsMultiSourcePSVulkanSM6, FCopyQuadsMultiSourcePSParametersVulkanSM6>;
 
 // ----------------------------------------------------------------------------------
 // Copy texture render command
@@ -2799,14 +2761,7 @@ bool ULandscapeEditLayerPersistent::RenderLayer(UE::Landscape::EditLayers::FRend
 			FRDGTextureSRVRef BlackDummySRVRef = GraphBuilder.CreateSRV(FRDGTextureSRVDesc::Create(GSystemTextures.GetBlackDummy(GraphBuilder)));
 
 			TUniquePtr<FCopyQuadsMultiSourceBase> CopyQuadsMultiSource;
-			if (SceneInterface->GetShaderPlatform() == SP_VULKAN_SM6)
-			{
-				CopyQuadsMultiSource = MakeUnique<FCopyQuadsMultiSourceVulkanSM6>();
-			}
-			else
-			{
-				CopyQuadsMultiSource = MakeUnique<FCopyQuadsMultiSourceDefault>();
-			}
+			CopyQuadsMultiSource = MakeUnique<FCopyQuadsMultiSourceDefault>();
 
 			// TODO [jonathan.bard] this is just to avoid a RHI validation error for unoptimized shaders... once validation is made to not issue those errors, we can remove this
 			// Create a SceneView to please the shader bindings, but it's unused in practice 
