Hotfix for landscape grass crash : landscape components can get registered to the grass builder and have their grass map start to be rendered. But time-slicing may prevent the other landscape components from the proxy to be all registered, leading to the landscape proxy not being fully registered, i.e. not having its landscape info object created. By the time the grass maps get read back, they assume the landscape info was valid, which is not the case.
We now prevent grass maps from being generated until all components from the proxy have been registered

diff --git a/Engine/Source/Runtime/Landscape/Private/LandscapeGrass.cpp b/Engine/Source/Runtime/Landscape/Private/LandscapeGrass.cpp
--- a/Engine/Source/Runtime/Landscape/Private/LandscapeGrass.cpp
+++ b/Engine/Source/Runtime/Landscape/Private/LandscapeGrass.cpp
@@ -1294,6 +1294,13 @@ bool ALandscapeProxy::ShouldGenerateGrass() const
 		return false;
 	}
 
+	// HACK [jonathan.bard] : for hotfix : don't let grass be rendered until all components of the proxy have been registered. That prevents grass maps to be rendered on those components, which 
+	//  allows LandscapeInfo to be valid by the time the grass results are fetched
+	if (!HasActorRegisteredAllComponents())
+	{
+		return false;
+	}
+
 	if (const ALandscape* LandscapeActor = GetLandscapeActor())
 	{
 		return LandscapeActor->GetGrassEnabled();
diff --git a/Engine/Source/Runtime/Landscape/Private/LandscapeGrassWeightExporter.cpp b/Engine/Source/Runtime/Landscape/Private/LandscapeGrassWeightExporter.cpp
--- a/Engine/Source/Runtime/Landscape/Private/LandscapeGrassWeightExporter.cpp
+++ b/Engine/Source/Runtime/Landscape/Private/LandscapeGrassWeightExporter.cpp
@@ -966,7 +966,12 @@ TMap<ULandscapeComponent*, TUniquePtr<FLandscapeComponentGrassData>, TInlineSetA
 	TMap<FName, FByteBuffer2DView> WeightData;
 
 	ULandscapeInfo* LandscapeInfo = Landscape->GetLandscapeInfo();
-	check(LandscapeInfo != nullptr);
+	// HACK [jonathan.bard] : this should be changed to a check 
+	ensure(LandscapeInfo != nullptr);
+	if (LandscapeInfo == nullptr)
+	{
+		return {};
+	}
 
 	for (FComponentInfo& ComponentInfo : ComponentInfos)
 	{
