Landscape Editor - Added legacy non-edit layer landscape foliage snapping as ShouldUseLegacyFoliageSnapUpdates Project Setting
* Users can enable the legacy foliage snapping behavior caused by the deprecation in 45245215. Disables per-component snapping that runs during the edit layers merge

#Synced-CL 52170742

diff --git a/Engine/Source/Runtime/Landscape/Private/LandscapeEditLayers.cpp b/Engine/Source/Runtime/Landscape/Private/LandscapeEditLayers.cpp
--- a/Engine/Source/Runtime/Landscape/Private/LandscapeEditLayers.cpp
+++ b/Engine/Source/Runtime/Landscape/Private/LandscapeEditLayers.cpp
@@ -7503,7 +7503,10 @@ uint32 ALandscape::UpdateCollisionAndClients(const TArrayView<FLandscapeEditLaye
 				if (ULandscapeHeightfieldCollisionComponent* CollisionComp = LandscapeComponent->GetCollisionComponent())
 				{
 					FNavigationSystem::UpdateComponentData(*CollisionComp);
-					CollisionComp->SnapFoliageInstances();
+					if (!GetDefault<ULandscapeSettings>()->ShouldUseLegacyFoliageSnapUpdates())
+					{
+						CollisionComp->SnapFoliageInstances();
+					}
 				}
 			}
 			else
diff --git a/Engine/Source/Runtime/Landscape/Public/LandscapeEdit.h b/Engine/Source/Runtime/Landscape/Public/LandscapeEdit.h
--- a/Engine/Source/Runtime/Landscape/Public/LandscapeEdit.h
+++ b/Engine/Source/Runtime/Landscape/Public/LandscapeEdit.h
@@ -18,6 +18,7 @@ LandscapeEdit.h: Classes for the editor to access to Landscape data
 #include "LandscapeHeightfieldCollisionComponent.h"
 #include "InstancedFoliageActor.h"
 #include "LandscapeLayerInfoObject.h"
+#include "LandscapeSettings.h"
 
 #if WITH_EDITOR
 #include "Containers/ArrayView.h"
@@ -430,6 +431,55 @@ struct FHeightmapAccessor
 				Component->RequestHeightmapUpdate();
 			}
 
+			// Legacy foliage behavior: only snap foliage within the given bounds (non-edit-layer landscape behavior)
+			// instead of updating on per-landscape component basis
+			if (GetDefault<ULandscapeSettings>()->ShouldUseLegacyFoliageSnapUpdates())
+			{
+				bool bHasAttachedFoliage = false;
+				for (ULandscapeComponent* Component : Components)
+				{
+					if (ULandscapeHeightfieldCollisionComponent* CollisionComponent = Component->GetCollisionComponent();
+						CollisionComponent && AInstancedFoliageActor::HasFoliageAttached(CollisionComponent))
+					{
+						bHasAttachedFoliage = true;
+						break;
+					}
+				}
+
+				if (bHasAttachedFoliage)
+				{
+					TArray<ULandscapeHeightfieldCollisionComponent*> CollisionComponents;
+					TArray<FBox> PreUpdateLocalBoxes;
+					CollisionComponents.Reserve(Components.Num());
+					PreUpdateLocalBoxes.Reserve(Components.Num());
+
+					for (ULandscapeComponent* Component : Components)
+					{
+						if (ULandscapeHeightfieldCollisionComponent* CollisionComponent = Component->GetCollisionComponent())
+						{
+							CollisionComponents.Add(CollisionComponent);
+							PreUpdateLocalBoxes.Add(FBox(
+								FVector(static_cast<float>(X1), static_cast<float>(Y1), Component->CachedLocalBox.Min.Z),
+								FVector(static_cast<float>(X2), static_cast<float>(Y2), Component->CachedLocalBox.Max.Z)));
+						}
+					}
+
+					LandscapeEdit.SetHeightData(X1, Y1, X2, Y2, Data, 0, /*InCalcNormals=*/ false);
+
+					const FMatrix44d LandscapeToWorld = LandscapeInfo->GetLandscapeProxy()->LandscapeActorToWorld().ToMatrixWithScale();
+
+					for (int32 Index = 0; Index < CollisionComponents.Num(); ++Index)
+					{
+						CollisionComponents[Index]->SnapFoliageInstances(
+							PreUpdateLocalBoxes[Index].TransformBy(LandscapeToWorld).ExpandBy(1.0f));
+					}
+
+					// Early exit since data has been set already 
+					return;
+				}
+			}
+			
+			// Foliage will be snapped on a per-landscape component basis during the edit layers merge
 			LandscapeEdit.SetHeightData(X1, Y1, X2, Y2, Data, 0, /*InCalcNormals=*/ false);
 		}
 	}
diff --git a/Engine/Source/Runtime/Landscape/Public/LandscapeSettings.h b/Engine/Source/Runtime/Landscape/Public/LandscapeSettings.h
--- a/Engine/Source/Runtime/Landscape/Public/LandscapeSettings.h
+++ b/Engine/Source/Runtime/Landscape/Public/LandscapeSettings.h
@@ -65,6 +65,8 @@ public:
 	TSoftObjectPtr<ULandscapeLayerInfoObject> GetDefaultLayerInfoObject() const { return DefaultLayerInfoObject; }
 
 	bool GetShouldUpdateEditLayersDuringInteractiveChanges() const { return bShouldUpdateEditLayersDuringInteractiveChanges; }
+
+	bool ShouldUseLegacyFoliageSnapUpdates() const { return bUseLegacyFoliageSnapUpdates; }
  
 	bool ShouldDisplayTargetLayerThumbnails() const { return bDisplayTargetLayerThumbnails; }
 
@@ -125,7 +127,15 @@ protected:
 	/** For landscape layers-affecting changes, allows the landscape to be updated when performing an interactive change (e.g. when changing an edit layer's alpha). Set to false if the performance when editing gets too bad (the landscape will be properly updated when the dragging operation is done). */
 	UPROPERTY(EditAnywhere, config, Category = "Edit Layers")
 	bool bShouldUpdateEditLayersDuringInteractiveChanges = false;
- 
+
+	/**
+	 * When enabled, foliage snapping uses the legacy non-edit layer behavior: snapping is performed
+	 * synchronously during the sculpt stroke, limited to the brush affected collision component
+	 * bounds. The per-landscape component snap that normally runs during the edit layers merge is disabled.
+	 */
+	UPROPERTY(EditAnywhere, config, Category = "Edit Layers")
+	bool bUseLegacyFoliageSnapUpdates = false;
+
 	UPROPERTY(transient)
 	bool bRestrictiveMode = false;
 
