* Optimization for duplicating landscape in Linux : Since FCString:Strtoi and FCString::Atoi are pathologically slow on this platform, we parse the clipboard data differently
* Also fixed an infinite loop when FTextureSource(/ULandscapeHeightfieldCollisionComponent)::ImportCustomProperties receive a malformed text input

#Synced-CL 52169855

diff --git a/Engine/Source/Runtime/Engine/Private/Texture.cpp b/Engine/Source/Runtime/Engine/Private/Texture.cpp
--- a/Engine/Source/Runtime/Engine/Private/Texture.cpp
+++ b/Engine/Source/Runtime/Engine/Private/Texture.cpp
@@ -24,10 +24,11 @@
 #include "Interfaces/ITargetPlatformManagerModule.h"
 #include "AssetCompilingManager.h"
 #include "Engine/TextureLODSettings.h"
-#include "RenderUtils.h"
 #include "ObjectCacheContext.h"
+#include "ProfilingDebugging/CpuProfilerTrace.h"
 #include "Rendering/StreamableTextureResource.h"
 #include "RenderingThread.h"
+#include "RenderUtils.h"
 #include "Interfaces/ITextureFormat.h"
 #include "Interfaces/ITextureFormatModule.h"
 #include "UObject/UE5MainStreamObjectVersion.h"
@@ -3414,6 +3415,8 @@ void FTextureSource::ExportCustomProperties(FOutputDevice& Out, uint32 Indent)
 
 void FTextureSource::ImportCustomProperties(const TCHAR* SourceText, FFeedbackContext* Warn)
 {
+	TRACE_CPUPROFILER_EVENT_SCOPE(FTextureSource::ImportCustomProperties);
+
 #if WITH_EDITOR
 	FScopeLock BulkDataExclusiveScope(&const_cast<FCriticalSection &>(BulkDataLock.Get()));
 #endif
@@ -3435,28 +3438,39 @@ void FTextureSource::ImportCustomProperties(const TCHAR* SourceText, FFeedbackCo
 		bool bSuccess = true;
 		if (PayloadSize > (uint64)0)
 		{
-			TCHAR* StopStr;
 			FUniqueBuffer Buffer = FUniqueBuffer::Alloc(PayloadSize);
 			uint8* DestData = (uint8*)Buffer.GetData();
 			if (DestData)
 			{
+				uint64 NumTokensParsed = 0;
 				uint64 Index = 0;
 				while (FChar::IsHexDigit(*SourceText))
 				{
 					if (Index < PayloadSize)
 					{
-						DestData[Index++] = (uint8)FCString::Strtoi(SourceText, &StopStr, 16);
+						uint8 Value = 0;
+						// Don't use FCString::Strtoi as it can be pathologically slow on some platforms because of the StringCast conversion to convert from TCHAR to char. 
+						while (FChar::IsHexDigit(*SourceText))
+						{
+							Value = (Value << 4) | TCharToNibble(*SourceText++);
+						}
+						DestData[Index++] = Value;
+					}
+					else
+					{
+						// If there are more hex tokens despite the Index having reached PayloadSize, then the input is malformed. We still want to advance the parsing in such to avoid looping infintely
 						while (FChar::IsHexDigit(*SourceText))
 						{
-							SourceText++;
+							++SourceText;
 						}
 					}
+					++NumTokensParsed;
 					FParse::Next(&SourceText);
 				}
 
-				if (Index != PayloadSize)
+				if (NumTokensParsed != PayloadSize)
 				{
-					Warn->Log(*NSLOCTEXT("UnrealEd", "Importing_TextureSource_SyntaxError", "Syntax Error").ToString());
+					Warn->Log(*FText::Format(NSLOCTEXT("UnrealEd", "Importing_TextureSource_SyntaxError", "Syntax Error : num tokens: {0} vs. payload size: {1}"), { NumTokensParsed, PayloadSize}).ToString());
 					bSuccess = false;
 				}
 			}
diff --git a/Engine/Source/Runtime/Landscape/Private/LandscapeCollision.cpp b/Engine/Source/Runtime/Landscape/Private/LandscapeCollision.cpp
--- a/Engine/Source/Runtime/Landscape/Private/LandscapeCollision.cpp
+++ b/Engine/Source/Runtime/Landscape/Private/LandscapeCollision.cpp
@@ -44,6 +44,7 @@
 #include "InstancedFoliage.h"
 #include "AI/NavigationSystemHelpers.h"
 #include "Engine/CollisionProfile.h"
+#include "ProfilingDebugging/CpuProfilerTrace.h"
 #include "ProfilingDebugging/CookStats.h"
 #include "Interfaces/ITargetPlatform.h"
 #include "Interfaces/ITargetPlatformManagerModule.h"
@@ -2313,6 +2314,8 @@ void ULandscapeHeightfieldCollisionComponent::ExportCustomProperties(FOutputDevi
 
 void ULandscapeHeightfieldCollisionComponent::ImportCustomProperties(const TCHAR* SourceText, FFeedbackContext* Warn)
 {
+	TRACE_CPUPROFILER_EVENT_SCOPE(ULandscapeHeightfieldCollisionComponent::ImportCustomProperties);
+
 	if (FParse::Command(&SourceText, TEXT("CollisionHeightData")))
 	{
 		int32 CollisionSizeVerts = CollisionSizeQuads + 1;
@@ -2323,27 +2326,41 @@ void ULandscapeHeightfieldCollisionComponent::ImportCustomProperties(const TCHAR
 		uint16* Heights = (uint16*)CollisionHeightData.Realloc(NumHeights);
 		FMemory::Memzero(Heights, sizeof(uint16)*NumHeights);
 
+		// Helper to convert digits to uint (assume FChar::IsDigit(InTChar) == true)
+		auto TCharToDigit = [](TCHAR InTChar) { return static_cast<uint32>(InTChar - TEXT('0')); };
+
 		FParse::Next(&SourceText);
-		int32 i = 0;
+		uint64 NumTokensParsed = 0;
+		int32 Index = 0;
 		while (FChar::IsDigit(*SourceText))
 		{
-			if (i < NumHeights)
+			if (Index < NumHeights)
 			{
-				Heights[i++] = static_cast<uint16>(FCString::Atoi(SourceText));
+				uint32 Value = 0;
+				// Don't use FCString::Atoi as it can be pathologically slow on some platforms because of the StringCast conversion to convert from TCHAR to char string
 				while (FChar::IsDigit(*SourceText))
 				{
-					SourceText++;
+					Value = Value * 10 + TCharToDigit(*SourceText++);
 				}
+				Heights[Index++] = static_cast<uint16>(Value);
 			}
-
+			else
+			{
+				// If there are more tokens despite the Index having reached NumHeights, then the input is malformed. We still want to advance the parsing in such a case to avoid looping infintely
+				while (FChar::IsDigit(*SourceText))
+				{
+					++SourceText;
+				}
+			}
+			++NumTokensParsed;
 			FParse::Next(&SourceText);
 		}
 
 		CollisionHeightData.Unlock();
 
-		if (i != NumHeights)
+		if (NumTokensParsed != NumHeights)
 		{
-			Warn->Log(*NSLOCTEXT("Core", "SyntaxError", "Syntax Error").ToString());
+			Warn->Log(*FText::Format(NSLOCTEXT("Core", "Importing_HeightfieldCollision_SyntaxError", "Syntax Error : num tokens: {0} vs. num height: {1}"), { NumTokensParsed, NumHeights }).ToString());
 		}
 	}
 	else if (FParse::Command(&SourceText, TEXT("DominantLayerData")))
