Fixup incorrect landscape splines connection data on load. It's still unclear what led to this situation but this allows the user to reconnect the faulty spline points

#Synced-CL 52314902

diff --git a/Engine/Source/Runtime/Landscape/Private/LandscapeSplines.cpp b/Engine/Source/Runtime/Landscape/Private/LandscapeSplines.cpp
--- a/Engine/Source/Runtime/Landscape/Private/LandscapeSplines.cpp
+++ b/Engine/Source/Runtime/Landscape/Private/LandscapeSplines.cpp
@@ -549,11 +549,17 @@ void ULandscapeSplinesComponent::CheckSplinesValid()
 	for (ULandscapeSplineControlPoint* ControlPoint : ControlPoints)
 	{
 		ensureMsgf(ControlPoint->GetOuterULandscapeSplinesComponent() == this, TEXT("Error in ULandscapeSplinesComponent in %s, actor: %s"), *GetPathName(), GetOwner() ? *GetOwner()->GetActorNameOrLabel() : TEXT("null"));
-		for (const FLandscapeSplineConnection& Connection : ControlPoint->ConnectedSegments)
+
+		for (auto ConnectionIt = ControlPoint->ConnectedSegments.CreateIterator(); ConnectionIt; ++ConnectionIt)
 		{
-			if (Connection.Segment)
+			if (ensureMsgf((ConnectionIt->Segment != nullptr),
+				TEXT("Error in ULandscapeSplinesComponent in %s, actor: %s : null segment : removing connection"), *GetPathName(), GetOwner() ? *GetOwner()->GetActorNameOrLabel() : TEXT("null")))
+			{
+				ensureMsgf(ConnectionIt->Segment->GetOuterULandscapeSplinesComponent() == this, TEXT("Error in ULandscapeSplinesComponent in %s, actor: %s"), *GetPathName(), GetOwner() ? *GetOwner()->GetActorNameOrLabel() : TEXT("null"));
+			}
+			else
 			{
-				ensureMsgf(Connection.Segment->GetOuterULandscapeSplinesComponent() == this, TEXT("Error in ULandscapeSplinesComponent in %s, actor: %s"), *GetPathName(), GetOwner() ? *GetOwner()->GetActorNameOrLabel() : TEXT("null"));
+				ConnectionIt.RemoveCurrent();
 			}
 		}
 	}
