[Backout] - CL52361189
#rnx
Original CL Desc
-----------------------------------------------------------------
Fixup incorrect landscape splines connection data on load. It's still unclear what led to this situation but this allows the user to reconnect the faulty spline points

diff --git a/Engine/Source/Runtime/Landscape/Private/LandscapeSplines.cpp b/Engine/Source/Runtime/Landscape/Private/LandscapeSplines.cpp
--- a/Engine/Source/Runtime/Landscape/Private/LandscapeSplines.cpp
+++ b/Engine/Source/Runtime/Landscape/Private/LandscapeSplines.cpp
@@ -549,17 +549,11 @@ void ULandscapeSplinesComponent::CheckSplinesValid()
 	for (ULandscapeSplineControlPoint* ControlPoint : ControlPoints)
 	{
 		ensureMsgf(ControlPoint->GetOuterULandscapeSplinesComponent() == this, TEXT("Error in ULandscapeSplinesComponent in %s, actor: %s"), *GetPathName(), GetOwner() ? *GetOwner()->GetActorNameOrLabel() : TEXT("null"));
-
-		for (auto ConnectionIt = ControlPoint->ConnectedSegments.CreateIterator(); ConnectionIt; ++ConnectionIt)
+		for (const FLandscapeSplineConnection& Connection : ControlPoint->ConnectedSegments)
 		{
-			if (ensureMsgf((ConnectionIt->Segment != nullptr),
-				TEXT("Error in ULandscapeSplinesComponent in %s, actor: %s : null segment : removing connection"), *GetPathName(), GetOwner() ? *GetOwner()->GetActorNameOrLabel() : TEXT("null")))
-			{
-				ensureMsgf(ConnectionIt->Segment->GetOuterULandscapeSplinesComponent() == this, TEXT("Error in ULandscapeSplinesComponent in %s, actor: %s"), *GetPathName(), GetOwner() ? *GetOwner()->GetActorNameOrLabel() : TEXT("null"));
-			}
-			else
+			if (Connection.Segment)
 			{
-				ConnectionIt.RemoveCurrent();
+				ensureMsgf(Connection.Segment->GetOuterULandscapeSplinesComponent() == this, TEXT("Error in ULandscapeSplinesComponent in %s, actor: %s"), *GetPathName(), GetOwner() ? *GetOwner()->GetActorNameOrLabel() : TEXT("null"));
 			}
 		}
 	}
