Resubmit temp fix for landscape spline crash, now without an ensure

#rnx
#Synced-CL 52314902

diff --git a/Engine/Source/Runtime/Landscape/Private/LandscapeSplines.cpp b/Engine/Source/Runtime/Landscape/Private/LandscapeSplines.cpp
--- a/Engine/Source/Runtime/Landscape/Private/LandscapeSplines.cpp
+++ b/Engine/Source/Runtime/Landscape/Private/LandscapeSplines.cpp
@@ -542,18 +542,24 @@ void ULandscapeSplinesComponent::CheckSplinesValid()
 	ensureMsgf(!ControlPoints.Remove(nullptr), TEXT("Error found and resolved in ULandscapeSplinesComponent in %s, actor: %s"), *GetPathName(), GetOwner() ? *GetOwner()->GetActorNameOrLabel() : TEXT("null"));
 	ensureMsgf(!Segments.Remove(nullptr), TEXT("Error found and resolved in ULandscapeSplinesComponent in %s, actor: %s"), *GetPathName(), GetOwner() ? *GetOwner()->GetActorNameOrLabel() : TEXT("null"));
 
-#if DO_CHECK
 	// Check for cross-spline connections, as this is a potential source of nulls.  This can also happen if something went wrong
 	// with a segment/point connectivity edit (while merging splines, undo/redo of deletions, etc).  These errors show up some times,
 	// but the trigger conditions are not known.
 	for (ULandscapeSplineControlPoint* ControlPoint : ControlPoints)
 	{
 		ensureMsgf(ControlPoint->GetOuterULandscapeSplinesComponent() == this, TEXT("Error in ULandscapeSplinesComponent in %s, actor: %s"), *GetPathName(), GetOwner() ? *GetOwner()->GetActorNameOrLabel() : TEXT("null"));
-		for (const FLandscapeSplineConnection& Connection : ControlPoint->ConnectedSegments)
+
+		for (auto ConnectionIt = ControlPoint->ConnectedSegments.CreateIterator(); ConnectionIt; ++ConnectionIt)
 		{
-			if (Connection.Segment)
+			// HACK [jonathan.bard] : comment out this ensure temporarily while we fix the data : if (ensureMsgf((ConnectionIt->Segment != nullptr),
+			//	TEXT("Error in ULandscapeSplinesComponent in %s, actor: %s : null segment : removing connection"), *GetPathName(), GetOwner() ? *GetOwner()->GetActorNameOrLabel() : TEXT("null")))
+			if (ConnectionIt->Segment != nullptr)
 			{
-				ensureMsgf(Connection.Segment->GetOuterULandscapeSplinesComponent() == this, TEXT("Error in ULandscapeSplinesComponent in %s, actor: %s"), *GetPathName(), GetOwner() ? *GetOwner()->GetActorNameOrLabel() : TEXT("null"));
+				ensureMsgf(ConnectionIt->Segment->GetOuterULandscapeSplinesComponent() == this, TEXT("Error in ULandscapeSplinesComponent in %s, actor: %s"), *GetPathName(), GetOwner() ? *GetOwner()->GetActorNameOrLabel() : TEXT("null"));
+			}
+			else
+			{
+				ConnectionIt.RemoveCurrent();
 			}
 		}
 	}
@@ -569,7 +575,6 @@ void ULandscapeSplinesComponent::CheckSplinesValid()
 			}
 		}
 	}
-#endif
 }
 
 void ULandscapeSplinesComponent::OnRegister()
