Fixed ensures from some landscape objects that change their outer when deleted

These have to be flagged as root objects for transactions so that the outer change re-applies correctly on undo/redo, as the objects themselves are initially outered under an unstable component

#rnx
#Synced-CL 52465190

diff --git a/Engine/Source/Runtime/Landscape/Classes/LandscapeSplineControlPoint.h b/Engine/Source/Runtime/Landscape/Classes/LandscapeSplineControlPoint.h
--- a/Engine/Source/Runtime/Landscape/Classes/LandscapeSplineControlPoint.h
+++ b/Engine/Source/Runtime/Landscape/Classes/LandscapeSplineControlPoint.h
@@ -288,6 +288,7 @@ public:
 	virtual void PostDuplicate(bool bDuplicateForPIE) override;
 	virtual void PostEditImport() override;
 	virtual void PostEditChangeProperty(FPropertyChangedEvent& PropertyChangedEvent) override;
+	virtual bool IsCapturingAsRootObjectForTransaction() const override { return true; }
 	//~ End UObject Interface
 #endif // WITH_EDITOR
 
diff --git a/Engine/Source/Runtime/Landscape/Classes/LandscapeSplineSegment.h b/Engine/Source/Runtime/Landscape/Classes/LandscapeSplineSegment.h
--- a/Engine/Source/Runtime/Landscape/Classes/LandscapeSplineSegment.h
+++ b/Engine/Source/Runtime/Landscape/Classes/LandscapeSplineSegment.h
@@ -381,6 +381,7 @@ public:
 	virtual void PostDuplicate(bool bDuplicateForPIE) override;
 	virtual void PostEditImport() override;
 	virtual void PostEditChangeProperty(FPropertyChangedEvent& PropertyChangedEvent) override;
+	virtual bool IsCapturingAsRootObjectForTransaction() const override { return true; }
 #endif // WITH_EDITOR
 protected:
 	virtual void PostInitProperties() override;
